GeometryScripting: initial checkin of Experimental GeometryScripting plugin, enables manipulation of UDynamicMesh mesh objects via ufunctions/blueprints. These functions are generally wrappers around functions from the GeometryProcessing plugin.

MeshAssetFunctions: read/write from StaticMesh Assets
MeshBasicEditFunctions: low-level mesh operations, eg append vertex/triangle, append mesh, etc
MeshBooleanFunctions: Mesh Boolean, Self-Union, Plane Cut, Mirror
MeshDeformFunctions: Bend and Twist warps
MeshNormalsFunctions: Various methods to calculate normals, Flip, Autorepair
MeshPrimitiveFunctions: append standard geometric primitive meshes
MeshQueryFunctions: variety of queries for mesh information
MeshRepairFunctions: Compact, Weld, Fill Holes, remove hidden triangles
MeshSimplifyFunctions: various mesh simplification strategies
MeshTransformFunctions: apply various geometric transformations
MeshUVFunctions: set mesh UVs
MeshVertexColorFunctions: set mesh vertex colors
MeshVoxelFunctions: voxel/sdf-based solidify and morphology operations
SceneUtilityFunctions: helper functions to simplify reading mesh from Actors/Components

#rnx

diff --git a/Engine/Plugins/Experimental/GeometryScripting/GeometryScripting.uplugin b/Engine/Plugins/Experimental/GeometryScripting/GeometryScripting.uplugin
new file mode 100644
--- /dev/null
+++ b/Engine/Plugins/Experimental/GeometryScripting/GeometryScripting.uplugin
@@ -0,0 +1,33 @@
+{
+	"FileVersion": 3,
+	"Version": 1,
+	"VersionName": "0.1",
+	"FriendlyName": "Geometry Scripting",
+	"Description": "Geometry Scripting Modules enable procedural generation and manipulation of meshes/geometry via Blueprints",
+	"Category": "Other",
+	"CreatedBy": "Epic Games, Inc.",
+	"CreatedByURL": "http://epicgames.com",
+	"DocsURL": "",
+	"MarketplaceURL": "",
+	"SupportURL": "",
+	"CanContainContent": true,
+	"IsExperimentalVersion": true,
+	"Installed": false,
+	"Modules": [
+		{
+			"Name": "GeometryScriptingCore",
+			"Type": "Runtime",
+			"LoadingPhase": "Default"
+		}
+	],
+	"Plugins": [
+		{
+			"Name": "GeometryProcessing",
+			"Enabled": true
+		},
+		{
+			"Name": "MeshModelingToolset",
+			"Enabled": true
+		}
+	]
+}
diff --git a/Engine/Plugins/Experimental/GeometryScripting/Source/GeometryScriptingCore/GeometryScriptingCore.Build.cs b/Engine/Plugins/Experimental/GeometryScripting/Source/GeometryScriptingCore/GeometryScriptingCore.Build.cs
new file mode 100644
--- /dev/null
+++ b/Engine/Plugins/Experimental/GeometryScripting/Source/GeometryScriptingCore/GeometryScriptingCore.Build.cs
@@ -0,0 +1,67 @@
+// Copyright Epic Games, Inc. All Rights Reserved.
+
+using UnrealBuildTool;
+
+public class GeometryScriptingCore : ModuleRules
+{
+	public GeometryScriptingCore(ReadOnlyTargetRules Target) : base(Target)
+	{
+		PCHUsage = ModuleRules.PCHUsageMode.UseExplicitOrSharedPCHs;
+
+		PublicIncludePaths.AddRange(
+			new string[] {
+				// ... add public include paths required here ...
+			}
+			);
+
+
+		PrivateIncludePaths.AddRange(
+			new string[] {
+				// ... add other private include paths required here ...
+			}
+			);
+
+
+		PublicDependencyModuleNames.AddRange(
+			new string[]
+			{
+				"Core",
+				"PhysicsCore",
+				"RenderCore",
+				"GeometryCore",
+				"GeometryFramework",
+				"DynamicMesh"
+			}
+			);
+
+
+		PrivateDependencyModuleNames.AddRange(
+			new string[]
+			{
+				"CoreUObject",
+				"Engine",
+				"MeshDescription",
+				"StaticMeshDescription",
+				"MeshConversion",
+				"GeometryAlgorithms",
+				"ModelingOperators"
+			}
+			);
+
+
+		DynamicallyLoadedModuleNames.AddRange(
+			new string[]
+			{
+				// ... add any modules that your module loads dynamically here ...
+			}
+			);
+
+		// add UnrealEd dependency if building for editor
+		if (Target.bBuildEditor == true)
+		{
+			PrivateDependencyModuleNames.Add("EditorFramework");
+			PrivateDependencyModuleNames.Add("UnrealEd");
+		}
+
+	}
+}
diff --git a/Engine/Plugins/Experimental/GeometryScripting/Source/GeometryScriptingCore/Private/GeometryScriptTypes.cpp b/Engine/Plugins/Experimental/GeometryScripting/Source/GeometryScriptingCore/Private/GeometryScriptTypes.cpp
new file mode 100644
--- /dev/null
+++ b/Engine/Plugins/Experimental/GeometryScripting/Source/GeometryScriptingCore/Private/GeometryScriptTypes.cpp
@@ -0,0 +1,23 @@
+// Copyright Epic Games, Inc. All Rights Reserved.
+
+#include "GeometryScript/GeometryScriptTypes.h"
+#include "GeometryBase.h"
+
+
+
+FGeometryScriptDebugMessage UE::Geometry::MakeScriptError(EGeometryScriptErrorType ErrorTypeIn, const FText& MessageIn)
+{
+	UE_LOG(LogGeometry, Warning, TEXT("GeometryScriptError: %s"), *MessageIn.ToString() );
+
+	return FGeometryScriptDebugMessage{ EGeometryScriptDebugMessageType::ErrorMessage, ErrorTypeIn, MessageIn };
+}
+
+
+void UE::Geometry::AppendError(UGeometryScriptDebug* Debug, EGeometryScriptErrorType ErrorTypeIn, const FText& MessageIn)
+{
+	FGeometryScriptDebugMessage Result = MakeScriptError(ErrorTypeIn, MessageIn);
+	if (Debug != nullptr)
+	{
+		Debug->Append(Result);
+	}
+}
\ No newline at end of file
diff --git a/Engine/Plugins/Experimental/GeometryScripting/Source/GeometryScriptingCore/Private/GeometryScriptingCoreModule.cpp b/Engine/Plugins/Experimental/GeometryScripting/Source/GeometryScriptingCore/Private/GeometryScriptingCoreModule.cpp
new file mode 100644
--- /dev/null
+++ b/Engine/Plugins/Experimental/GeometryScripting/Source/GeometryScriptingCore/Private/GeometryScriptingCoreModule.cpp
@@ -0,0 +1,20 @@
+// Copyright Epic Games, Inc. All Rights Reserved.
+
+#include "GeometryScriptingCoreModule.h"
+
+#define LOCTEXT_NAMESPACE "FGeometryScriptingCoreModule"
+
+void FGeometryScriptingCoreModule::StartupModule()
+{
+	// This code will execute after your module is loaded into memory; the exact timing is specified in the .uplugin file per-module
+}
+
+void FGeometryScriptingCoreModule::ShutdownModule()
+{
+	// This function may be called during shutdown to clean up your module.  For modules that support dynamic reloading,
+	// we call this function before unloading the module.
+}
+
+#undef LOCTEXT_NAMESPACE
+	
+IMPLEMENT_MODULE(FGeometryScriptingCoreModule, GeometryScriptingCore)
\ No newline at end of file
diff --git a/Engine/Plugins/Experimental/GeometryScripting/Source/GeometryScriptingCore/Private/MeshAssetFunctions.cpp b/Engine/Plugins/Experimental/GeometryScripting/Source/GeometryScriptingCore/Private/MeshAssetFunctions.cpp
new file mode 100644
--- /dev/null
+++ b/Engine/Plugins/Experimental/GeometryScripting/Source/GeometryScriptingCore/Private/MeshAssetFunctions.cpp
@@ -0,0 +1,221 @@
+// Copyright Epic Games, Inc. All Rights Reserved.
+
+#include "GeometryScript/MeshAssetFunctions.h"
+
+#include "DynamicMesh/DynamicMesh3.h"
+#include "UDynamicMesh.h"
+
+#include "Engine/StaticMesh.h"
+#include "Engine/SkeletalMesh.h"
+#include "Engine/StaticMeshActor.h"
+#include "StaticMeshAttributes.h"
+#include "StaticMeshOperations.h"
+#include "RenderingThread.h"
+
+#include "MeshDescriptionToDynamicMesh.h"
+#include "DynamicMeshToMeshDescription.h"
+
+
+#if WITH_EDITOR
+#include "Editor.h"
+#endif
+
+
+#include "ExplicitUseGeometryMathTypes.h"		// using UE::Geometry::(math types)
+using namespace UE::Geometry;
+
+#define LOCTEXT_NAMESPACE "UGeometryScriptLibrary_MeshAssetFunctions"
+
+
+
+
+UDynamicMesh*  UGeometryScriptLibrary_StaticMeshFunctions::CopyMeshFromStaticMesh(
+	UStaticMesh* FromStaticMeshAsset, 
+	UDynamicMesh* ToDynamicMesh, 
+	FGeometryScriptCopyMeshFromAssetOptions AssetOptions,
+	FGeometryScriptMeshReadLOD RequestedLOD,
+	TEnumAsByte<EGeometryScriptOutcomePins>& Outcome,
+	UGeometryScriptDebug* Debug)
+{
+	Outcome = EGeometryScriptOutcomePins::Failure;
+
+	if (FromStaticMeshAsset == nullptr)
+	{
+		UE::Geometry::AppendError(Debug, EGeometryScriptErrorType::InvalidInputs, LOCTEXT("CopyMeshFromAsset_InvalidInput1", "CopyMeshFromStaticMesh: FromStaticMeshAsset is Null"));
+		return ToDynamicMesh;
+	}
+	if (ToDynamicMesh == nullptr)
+	{
+		UE::Geometry::AppendError(Debug, EGeometryScriptErrorType::InvalidInputs, LOCTEXT("CopyMeshFromAsset_InvalidInput2", "CopyMeshFromStaticMesh: ToDynamicMesh is Null"));
+		return ToDynamicMesh;
+	}
+	if (RequestedLOD.LODType != EGeometryScriptLODType::MaxAvailable && RequestedLOD.LODType != EGeometryScriptLODType::SourceModel)
+	{
+		UE::Geometry::AppendError(Debug, EGeometryScriptErrorType::InvalidInputs, LOCTEXT("CopyMeshFromAsset_LODNotAvailable", "CopyMeshFromStaticMesh: Requested LOD is not available"));
+		return ToDynamicMesh;
+	}
+
+#if WITH_EDITOR
+	int32 UseLODIndex = FMath::Clamp(RequestedLOD.LODIndex, 0, FromStaticMeshAsset->GetNumLODs() - 1);
+
+	const FMeshDescription* SourceMesh = FromStaticMeshAsset->GetMeshDescription(UseLODIndex);
+	const FStaticMeshSourceModel& SourceModel = FromStaticMeshAsset->GetSourceModel(UseLODIndex);
+	const FMeshBuildSettings& BuildSettings = SourceModel.BuildSettings;
+
+	bool bHasDirtyBuildSettings = BuildSettings.bRecomputeNormals
+		|| (BuildSettings.bRecomputeTangents && AssetOptions.bRequestTangents);
+
+	FMeshDescription LocalSourceMeshCopy;
+	if (AssetOptions.bApplyBuildSettings && bHasDirtyBuildSettings )
+	{
+		LocalSourceMeshCopy = *SourceMesh;
+
+		FStaticMeshAttributes Attributes(LocalSourceMeshCopy);
+		if (!Attributes.GetTriangleNormals().IsValid() || !Attributes.GetTriangleTangents().IsValid())
+		{
+			// If these attributes don't exist, create them and compute their values for each triangle
+			FStaticMeshOperations::ComputeTriangleTangentsAndNormals(LocalSourceMeshCopy);
+		}
+
+		EComputeNTBsFlags ComputeNTBsOptions = EComputeNTBsFlags::BlendOverlappingNormals;
+		ComputeNTBsOptions |= BuildSettings.bRecomputeNormals ? EComputeNTBsFlags::Normals : EComputeNTBsFlags::None;
+		if (AssetOptions.bRequestTangents)
+		{
+			ComputeNTBsOptions |= BuildSettings.bRecomputeTangents ? EComputeNTBsFlags::Tangents : EComputeNTBsFlags::None;
+			ComputeNTBsOptions |= BuildSettings.bUseMikkTSpace ? EComputeNTBsFlags::UseMikkTSpace : EComputeNTBsFlags::None;
+		}
+		ComputeNTBsOptions |= BuildSettings.bComputeWeightedNormals ? EComputeNTBsFlags::WeightedNTBs : EComputeNTBsFlags::None;
+		if (AssetOptions.bIgnoreRemoveDegenerates == false)
+		{
+			ComputeNTBsOptions |= BuildSettings.bRemoveDegenerates ? EComputeNTBsFlags::IgnoreDegenerateTriangles : EComputeNTBsFlags::None;
+		}
+
+		FStaticMeshOperations::ComputeTangentsAndNormals(LocalSourceMeshCopy, ComputeNTBsOptions);
+
+		SourceMesh = &LocalSourceMeshCopy;
+	}
+
+	FDynamicMesh3 NewMesh;
+	FMeshDescriptionToDynamicMesh Converter;
+	Converter.Convert(SourceMesh, NewMesh, AssetOptions.bRequestTangents);
+
+	ToDynamicMesh->SetMesh(MoveTemp(NewMesh));
+
+	Outcome = EGeometryScriptOutcomePins::Success;
+
+#else
+	UE::Geometry::AppendError(Debug, EGeometryScriptErrorType::InvalidInputs, LOCTEXT("CopyMeshFromAsset_EditorOnly", "CopyMeshFromStaticMesh: Not currently supported at Runtime"));
+#endif
+
+	return ToDynamicMesh;
+}
+
+
+
+
+UDynamicMesh*  UGeometryScriptLibrary_StaticMeshFunctions::CopyMeshToStaticMesh(
+	UDynamicMesh* FromDynamicMesh,
+	UStaticMesh* ToStaticMeshAsset,
+	FGeometryScriptCopyMeshToAssetOptions Options,
+	FGeometryScriptMeshWriteLOD TargetLOD,
+	TEnumAsByte<EGeometryScriptOutcomePins>& Outcome,
+	UGeometryScriptDebug* Debug)
+{
+	Outcome = EGeometryScriptOutcomePins::Failure;
+
+	if (FromDynamicMesh == nullptr)
+	{
+		UE::Geometry::AppendError(Debug, EGeometryScriptErrorType::InvalidInputs, LOCTEXT("CopyMeshToAsset_InvalidInput1", "CopyMeshToStaticMesh: FromDynamicMesh is Null"));
+		return FromDynamicMesh;
+	}
+	if (ToStaticMeshAsset == nullptr)
+	{
+		UE::Geometry::AppendError(Debug, EGeometryScriptErrorType::InvalidInputs, LOCTEXT("CopyMeshToAsset_InvalidInput2", "CopyMeshToStaticMesh: ToStaticMeshAsset is Null"));
+		return FromDynamicMesh;
+	}
+	if (TargetLOD.bWriteHiResSource)
+	{
+		UE::Geometry::AppendError(Debug, EGeometryScriptErrorType::InvalidInputs, LOCTEXT("CopyMeshToAsset_Unsupported", "CopyMeshToStaticMesh: Writing HiResSource LOD is not yet supported"));
+		return FromDynamicMesh;
+	}
+
+#if WITH_EDITOR
+
+	int32 UseLODIndex = FMath::Clamp(TargetLOD.LODIndex, 0, 32);
+
+	// don't allow build-in engine assets to be modified
+	if (ToStaticMeshAsset->GetPathName().StartsWith(TEXT("/Engine/")))
+	{
+		UE::Geometry::AppendError(Debug, EGeometryScriptErrorType::InvalidInputs, LOCTEXT("CopyMeshToAsset_EngineAsset", "CopyMeshToStaticMesh: Cannot modify built-in Engine asset"));
+		return FromDynamicMesh;
+	}
+
+	// flush any pending rendering commands, which might want to touch this StaticMesh while we are rebuilding it
+	FlushRenderingCommands();
+
+	if (Options.bEmitTransaction)
+	{
+		GEditor->BeginTransaction(LOCTEXT("UpdateStaticMesh", "Update Static Mesh"));
+	}
+
+	// make sure transactional flag is on for the Asset
+	ToStaticMeshAsset->SetFlags(RF_Transactional);
+	// mark as modified
+	ToStaticMeshAsset->Modify();
+	// mark mesh description for modify
+	ToStaticMeshAsset->ModifyMeshDescription(UseLODIndex);
+
+	FMeshDescription* MeshDescription = ToStaticMeshAsset->GetMeshDescription(UseLODIndex);
+	if (MeshDescription == nullptr)
+	{
+		MeshDescription = ToStaticMeshAsset->CreateMeshDescription(UseLODIndex);
+	}
+	if (!ensure(MeshDescription != nullptr))
+	{
+		UE::Geometry::AppendError(Debug, EGeometryScriptErrorType::InvalidInputs,
+			FText::Format(LOCTEXT("CopyMeshToAsset_EngineAsset", "CopyMeshToAsset: MeshDescription for LOD {0} is null?"), FText::AsNumber(UseLODIndex)));
+		return FromDynamicMesh;
+	}
+
+	FConversionToMeshDescriptionOptions ConversionOptions;
+	//if (FromDynamicMesh->GetMesh()->HasVertexColors())
+	//{
+	//	ConversionOptions.bUpdateVtxColors = true;		// is this still relevant now that we have vertex color attribute?
+	//}
+
+	FDynamicMeshToMeshDescription Converter(ConversionOptions);
+	FromDynamicMesh->ProcessMesh([&](const FDynamicMesh3& ReadMesh)
+	{
+		Converter.Convert(&ReadMesh, *MeshDescription);
+	});
+
+	ToStaticMeshAsset->CommitMeshDescription(UseLODIndex);
+
+	// Setting to prevent the standard static mesh reduction from running and replacing the render LOD.
+	FStaticMeshSourceModel& ThisSourceModel = ToStaticMeshAsset->GetSourceModel(UseLODIndex);
+	ThisSourceModel.ReductionSettings.PercentTriangles = 1.f;
+	ThisSourceModel.ReductionSettings.PercentVertices = 1.f;
+
+	if (Options.bDeferMeshPostEditChange == false)
+	{
+		ToStaticMeshAsset->PostEditChange();
+	}
+
+	if (Options.bEmitTransaction)
+	{
+		GEditor->EndTransaction();
+	}
+
+	Outcome = EGeometryScriptOutcomePins::Success;
+
+#else
+	UE::Geometry::AppendError(Debug, EGeometryScriptErrorType::InvalidInputs, LOCTEXT("CopyMeshToAsset_EditorOnly", "CopyMeshToStaticMesh: Not currently supported at Runtime"));
+#endif
+
+	return FromDynamicMesh;
+}
+
+
+
+
+#undef LOCTEXT_NAMESPACE
\ No newline at end of file
diff --git a/Engine/Plugins/Experimental/GeometryScripting/Source/GeometryScriptingCore/Private/MeshBasicEditFunctions.cpp b/Engine/Plugins/Experimental/GeometryScripting/Source/GeometryScriptingCore/Private/MeshBasicEditFunctions.cpp
new file mode 100644
--- /dev/null
+++ b/Engine/Plugins/Experimental/GeometryScripting/Source/GeometryScriptingCore/Private/MeshBasicEditFunctions.cpp
@@ -0,0 +1,321 @@
+// Copyright Epic Games, Inc. All Rights Reserved.
+
+#include "GeometryScript/MeshBasicEditFunctions.h"
+#include "UDynamicMesh.h"
+
+#include "DynamicMesh/DynamicMesh3.h"
+#include "DynamicMesh/DynamicMeshAttributeSet.h"
+#include "DynamicMesh/MeshTransforms.h"
+#include "DynamicMeshEditor.h"
+
+#include "ExplicitUseGeometryMathTypes.h"		// using UE::Geometry::(math types)
+using namespace UE::Geometry;
+
+#define LOCTEXT_NAMESPACE "UGeometryScriptLibrary_MeshBasicEditFunctions"
+
+
+UDynamicMesh* UGeometryScriptLibrary_MeshBasicEditFunctions::DiscardMeshAttributes(UDynamicMesh* TargetMesh, bool bDeferChangeNotifications)
+{
+	if (TargetMesh)
+	{
+		TargetMesh->EditMesh([&](FDynamicMesh3& EditMesh)
+		{
+			EditMesh.DiscardAttributes();
+			EditMesh.DiscardVertexNormals();
+
+		}, EDynamicMeshChangeType::GeneralEdit, EDynamicMeshAttributeChangeFlags::Unknown, bDeferChangeNotifications);
+	}
+	return TargetMesh;
+}
+
+
+
+
+
+UDynamicMesh* UGeometryScriptLibrary_MeshBasicEditFunctions::SetVertexPosition(UDynamicMesh* TargetMesh, int32 VertexID, FVector NewPosition, bool& bIsValidVertex, bool bDeferChangeNotifications)
+{
+	bIsValidVertex = false;
+	if (TargetMesh)
+	{
+		TargetMesh->EditMesh([&](FDynamicMesh3& EditMesh)
+		{
+			if (EditMesh.IsVertex(VertexID))
+			{
+				bIsValidVertex = true;
+				EditMesh.SetVertex(VertexID, (FVector3d)NewPosition);
+			}
+		}, EDynamicMeshChangeType::GeneralEdit, EDynamicMeshAttributeChangeFlags::Unknown, bDeferChangeNotifications);
+	}
+	return TargetMesh;
+}
+
+
+
+UDynamicMesh* UGeometryScriptLibrary_MeshBasicEditFunctions::AddVertexToMesh(
+	UDynamicMesh* TargetMesh,
+	FVector NewPosition,
+	int32& NewVertexIndex,
+	bool bDeferChangeNotifications)
+{
+	NewVertexIndex = INDEX_NONE;
+	if (TargetMesh)
+	{
+		TargetMesh->EditMesh([&](FDynamicMesh3& EditMesh)
+		{
+			NewVertexIndex = EditMesh.AppendVertex(NewPosition);
+
+		}, EDynamicMeshChangeType::GeneralEdit, EDynamicMeshAttributeChangeFlags::Unknown, bDeferChangeNotifications);
+	}
+	return TargetMesh;
+}
+
+UDynamicMesh* UGeometryScriptLibrary_MeshBasicEditFunctions::AddVerticesToMesh(
+	UDynamicMesh* TargetMesh,
+	const TArray<FVector>& NewPositions, 
+	TArray<int>& NewIndices,
+	bool bDeferChangeNotifications)
+{
+	NewIndices.Reset();
+	if (TargetMesh)
+	{
+		TargetMesh->EditMesh([&](FDynamicMesh3& EditMesh)
+		{
+			for (FVector Position : NewPositions)
+			{
+				int32 NewVertexIndex = EditMesh.AppendVertex(Position);
+				NewIndices.Add(NewVertexIndex);
+			}
+		
+		}, EDynamicMeshChangeType::GeneralEdit, EDynamicMeshAttributeChangeFlags::Unknown, bDeferChangeNotifications);
+	}
+	return TargetMesh;
+}
+
+
+UDynamicMesh* UGeometryScriptLibrary_MeshBasicEditFunctions::DeleteVertexFromMesh(
+	UDynamicMesh* TargetMesh,
+	int VertexID,
+	bool& bWasVertexDeleted,
+	bool bDeferChangeNotifications)
+{
+	bWasVertexDeleted = false;
+	if (TargetMesh)
+	{
+		TargetMesh->EditMesh([&](FDynamicMesh3& EditMesh)
+		{
+			EMeshResult Result = EditMesh.RemoveVertex(VertexID);
+			bWasVertexDeleted = (Result == EMeshResult::Ok);
+
+		}, EDynamicMeshChangeType::GeneralEdit, EDynamicMeshAttributeChangeFlags::Unknown, bDeferChangeNotifications);
+	}
+	return TargetMesh;
+}
+
+
+
+
+UDynamicMesh* UGeometryScriptLibrary_MeshBasicEditFunctions::AddTriangleToMesh(
+	UDynamicMesh* TargetMesh,
+	FIntVector NewTriangle,
+	int32& NewTriangleIndex,
+	int32 NewTriangleGroupID,
+	bool bDeferChangeNotifications,
+	UGeometryScriptDebug* Debug)
+{
+	NewTriangleIndex = INDEX_NONE;
+	if (TargetMesh)
+	{
+		TargetMesh->EditMesh([&](FDynamicMesh3& EditMesh)
+		{
+			NewTriangleIndex = EditMesh.AppendTriangle((FIndex3i)NewTriangle, NewTriangleGroupID);
+
+		}, EDynamicMeshChangeType::GeneralEdit, EDynamicMeshAttributeChangeFlags::Unknown, bDeferChangeNotifications);
+
+		if (NewTriangleIndex < 0)
+		{
+			if (NewTriangleIndex == FDynamicMesh3::NonManifoldID)
+			{
+				UE::Geometry::AppendError(Debug, EGeometryScriptErrorType::InvalidInputs, LOCTEXT("AddTriangleToMesh_NonManifold", "AddTriangleToMesh: Triangle cannot be added because it would create invalid Non-Manifold Mesh Topology"));
+			}
+			else if (NewTriangleIndex == FDynamicMesh3::DuplicateTriangleID)
+			{
+				UE::Geometry::AppendError(Debug, EGeometryScriptErrorType::InvalidInputs, LOCTEXT("AddTriangleToMesh_Duplicate", "AddTriangleToMesh: Triangle cannot be added because it is a duplicate of an existing Triangle"));
+			}
+			else
+			{
+				UE::Geometry::AppendError(Debug, EGeometryScriptErrorType::OperationFailed, LOCTEXT("AddTriangleToMesh_Unknown", "AddTriangleToMesh: adding Triangle Failed"));
+			}
+			NewTriangleIndex = INDEX_NONE;
+		}
+	}
+	else
+	{
+		UE::Geometry::AppendError(Debug, EGeometryScriptErrorType::InvalidInputs, LOCTEXT("AddTriangleToMesh_InvalidMesh", "AddTriangleToMesh: TargetMesh is Null"));
+	}
+	return TargetMesh;
+}
+
+
+
+
+UDynamicMesh* UGeometryScriptLibrary_MeshBasicEditFunctions::AddTrianglesToMesh(
+	UDynamicMesh* TargetMesh,
+	const TArray<FIntVector>& NewTriangles,
+	TArray<int>& NewIndices,
+	int32 NewTriangleGroupID,
+	bool bDeferChangeNotifications,
+	UGeometryScriptDebug* Debug)
+{
+	NewIndices.Reset();
+	if (TargetMesh)
+	{
+		TargetMesh->EditMesh([&](FDynamicMesh3& EditMesh)
+		{
+			for (FIntVector Triangle : NewTriangles)
+			{
+				int32 NewTriangleIndex = EditMesh.AppendTriangle((FIndex3i)Triangle, NewTriangleGroupID);
+				NewIndices.Add(NewTriangleIndex);
+			}
+
+		}, EDynamicMeshChangeType::GeneralEdit, EDynamicMeshAttributeChangeFlags::Unknown, bDeferChangeNotifications);
+
+		for (int32& NewTriangleIndex : NewIndices)
+		{
+			if (NewTriangleIndex < 0)
+			{
+				if (NewTriangleIndex == FDynamicMesh3::NonManifoldID)
+				{
+					UE::Geometry::AppendError(Debug, EGeometryScriptErrorType::InvalidInputs, LOCTEXT("AddTrianglesToMesh_NonManifold", "AddTrianglesToMesh: Triangle cannot be added because it would create invalid Non-Manifold Mesh Topology"));
+				}
+				else if (NewTriangleIndex == FDynamicMesh3::DuplicateTriangleID)
+				{
+					UE::Geometry::AppendError(Debug, EGeometryScriptErrorType::InvalidInputs, LOCTEXT("AddTrianglesToMesh_Duplicate", "AddTrianglesToMesh: Triangle cannot be added because it is a duplicate of an existing Triangle"));
+				}
+				else
+				{
+					UE::Geometry::AppendError(Debug, EGeometryScriptErrorType::OperationFailed, LOCTEXT("AddTrianglesToMesh_Unknown", "AddTrianglesToMesh: adding Triangle Failed"));
+				}
+				NewTriangleIndex = INDEX_NONE;
+			}
+		}
+	}
+	else
+	{
+		UE::Geometry::AppendError(Debug, EGeometryScriptErrorType::InvalidInputs, LOCTEXT("AddTrianglesToMesh_InvalidMesh", "AddTriangleToMesh: TargetMesh is Null"));
+	}
+	return TargetMesh;
+}
+
+
+
+UDynamicMesh* UGeometryScriptLibrary_MeshBasicEditFunctions::DeleteTriangleFromMesh(
+	UDynamicMesh* TargetMesh,
+	int TriangleID,
+	bool& bWasTriangleDeleted,
+	bool bDeferChangeNotifications)
+{
+	bWasTriangleDeleted = false;
+	if (TargetMesh)
+	{
+		TargetMesh->EditMesh([&](FDynamicMesh3& EditMesh)
+		{
+			EMeshResult Result = EditMesh.RemoveTriangle(TriangleID);
+			bWasTriangleDeleted = (Result == EMeshResult::Ok);
+
+		}, EDynamicMeshChangeType::GeneralEdit, EDynamicMeshAttributeChangeFlags::Unknown, bDeferChangeNotifications);
+	}
+	return TargetMesh;
+}
+
+
+
+
+UDynamicMesh* UGeometryScriptLibrary_MeshBasicEditFunctions::AppendMesh(
+	UDynamicMesh* TargetMesh,
+	UDynamicMesh* AppendMesh,
+	FTransform AppendTransform,
+	bool bDeferChangeNotifications,
+	UGeometryScriptDebug* Debug)
+{
+	if (TargetMesh == nullptr)
+	{
+		UE::Geometry::AppendError(Debug, EGeometryScriptErrorType::InvalidInputs, LOCTEXT("AppendMesh_InvalidInput1", "AppendMesh: TargetMesh is Null"));
+		return TargetMesh;
+	}
+	if (AppendMesh == nullptr)
+	{
+		UE::Geometry::AppendError(Debug, EGeometryScriptErrorType::InvalidInputs, LOCTEXT("AppendMesh_InvalidInput2", "AppendMesh: AppendMesh is Null"));
+		return TargetMesh;
+	}
+
+	TargetMesh->EditMesh([&](FDynamicMesh3& AppendToMesh)
+	{
+		AppendMesh->ProcessMesh([&](const FDynamicMesh3& OtherMesh)
+		{
+			UE::Geometry::FTransform3d XForm(AppendTransform);
+			FMeshIndexMappings TmpMappings;
+			FDynamicMeshEditor Editor(&AppendToMesh);
+			Editor.AppendMesh(&OtherMesh, TmpMappings,
+				[&](int, const FVector3d& Position) { return XForm.TransformPosition(Position); },
+				[&](int, const FVector3d& Normal) { return XForm.TransformNormal(Normal); });
+		});
+	}, EDynamicMeshChangeType::GeneralEdit, EDynamicMeshAttributeChangeFlags::Unknown, bDeferChangeNotifications);
+
+	return TargetMesh;
+}
+
+
+
+UDynamicMesh* UGeometryScriptLibrary_MeshBasicEditFunctions::AppendMeshRepeated(
+	UDynamicMesh* TargetMesh,
+	UDynamicMesh* AppendMesh,
+	FTransform AppendTransform,
+	int32 RepeatCount,
+	bool bApplyTransformToFirstInstance,
+	bool bDeferChangeNotifications,
+	UGeometryScriptDebug* Debug)
+{
+	if (TargetMesh == nullptr)
+	{
+		UE::Geometry::AppendError(Debug, EGeometryScriptErrorType::InvalidInputs, LOCTEXT("AppendMeshRepeated_InvalidInput1", "AppendMeshRepeated: TargetMesh is Null"));
+		return TargetMesh;
+	}
+	if (AppendMesh == nullptr)
+	{
+		UE::Geometry::AppendError(Debug, EGeometryScriptErrorType::InvalidInputs, LOCTEXT("AppendMeshRepeated_InvalidInput2", "AppendMeshRepeated: AppendMesh is Null"));
+		return TargetMesh;
+	}
+	if (RepeatCount > 0)
+	{
+		UE::Geometry::FTransform3d XForm(AppendTransform);
+		FDynamicMesh3 TmpMesh;
+		AppendMesh->ProcessMesh([&](const FDynamicMesh3& OtherMesh) { TmpMesh.Copy(OtherMesh); });
+		if (bApplyTransformToFirstInstance)
+		{
+			MeshTransforms::ApplyTransform(TmpMesh, XForm);
+		}
+		TargetMesh->EditMesh([&](FDynamicMesh3& AppendToMesh)
+		{
+			FMeshIndexMappings TmpMappings;
+			FDynamicMeshEditor Editor(&AppendToMesh);
+			for (int32 k = 0; k < RepeatCount; ++k)
+			{
+				Editor.AppendMesh(&TmpMesh, TmpMappings);
+				if (k < RepeatCount)
+				{
+					MeshTransforms::ApplyTransform(TmpMesh, XForm);
+					TmpMappings.Reset();
+				}
+			}
+		}, EDynamicMeshChangeType::GeneralEdit, EDynamicMeshAttributeChangeFlags::Unknown, bDeferChangeNotifications);
+	}
+
+	return TargetMesh;
+}
+
+
+
+
+
+
+#undef LOCTEXT_NAMESPACE
diff --git a/Engine/Plugins/Experimental/GeometryScripting/Source/GeometryScriptingCore/Private/MeshBooleanFunctions.cpp b/Engine/Plugins/Experimental/GeometryScripting/Source/GeometryScriptingCore/Private/MeshBooleanFunctions.cpp
new file mode 100644
--- /dev/null
+++ b/Engine/Plugins/Experimental/GeometryScripting/Source/GeometryScriptingCore/Private/MeshBooleanFunctions.cpp
@@ -0,0 +1,277 @@
+// Copyright Epic Games, Inc. All Rights Reserved.
+
+#include "GeometryScript/MeshBooleanFunctions.h"
+
+#include "DynamicMesh/DynamicMesh3.h"
+#include "DynamicMesh/MeshTransforms.h"
+#include "UDynamicMesh.h"
+
+#include "Operations/MeshBoolean.h"
+#include "Operations/MeshSelfUnion.h"
+#include "Operations/MeshPlaneCut.h"
+#include "Operations/MeshMirror.h"
+#include "MeshSimplification.h"
+#include "MeshBoundaryLoops.h"
+#include "Operations/MinimalHoleFiller.h"
+#include "ConstrainedDelaunay2.h"
+
+
+
+#include "ExplicitUseGeometryMathTypes.h"		// using UE::Geometry::(math types)
+using namespace UE::Geometry;
+
+#define LOCTEXT_NAMESPACE "UGeometryScriptLibrary_MeshBooleanFunctions"
+
+
+
+
+
+UDynamicMesh* UGeometryScriptLibrary_MeshBooleanFunctions::ApplyMeshBoolean(
+	UDynamicMesh* TargetMesh,
+	FTransform TargetTransform,
+	UDynamicMesh* ToolMesh,
+	FTransform ToolTransform,
+	EGeometryScriptBooleanOperation Operation,
+	FGeometryScriptMeshBooleanOptions Options,
+	UGeometryScriptDebug* Debug)
+{
+	if (TargetMesh == nullptr)
+	{
+		UE::Geometry::AppendError(Debug, EGeometryScriptErrorType::InvalidInputs, LOCTEXT("ApplyMeshBoolean_InvalidInput1", "ApplyMeshBoolean: TargetMesh is Null"));
+		return TargetMesh;
+	}
+	if (ToolMesh == nullptr)
+	{
+		UE::Geometry::AppendError(Debug, EGeometryScriptErrorType::InvalidInputs, LOCTEXT("ApplyMeshBoolean_InvalidInput2", "ApplyMeshBoolean: ToolMesh is Null"));
+		return TargetMesh;
+	}
+
+
+	FMeshBoolean::EBooleanOp ApplyOperation = FMeshBoolean::EBooleanOp::Union;
+	switch (Operation)
+	{
+	case EGeometryScriptBooleanOperation::Intersection: 
+		ApplyOperation = FMeshBoolean::EBooleanOp::Intersect;
+		break;
+	case EGeometryScriptBooleanOperation::Subtract: 
+		ApplyOperation = FMeshBoolean::EBooleanOp::Difference;
+		break;	
+	}
+
+	FDynamicMesh3 NewResultMesh;
+	bool bSuccess = false;
+	TArray<int> NewBoundaryEdges;
+
+	// TODO: can actually compute in place (ie without making a copy) if we pass the first mesh as the ResultMesh argument...
+	TargetMesh->ProcessMesh([&](const FDynamicMesh3& Mesh1)
+	{
+		ToolMesh->ProcessMesh([&](const FDynamicMesh3& Mesh2)
+		{
+			FMeshBoolean MeshBoolean(
+				&Mesh1, (UE::Geometry::FTransform3d)TargetTransform,
+				&Mesh2, (UE::Geometry::FTransform3d)ToolTransform,
+				&NewResultMesh, ApplyOperation);
+			MeshBoolean.bPutResultInInputSpace = true;
+			MeshBoolean.bSimplifyAlongNewEdges = Options.bSimplifyOutput;
+			bSuccess = MeshBoolean.Compute();
+			NewBoundaryEdges = MoveTemp(MeshBoolean.CreatedBoundaryEdges);
+		});
+	});
+
+	// bSuccess comes back false even if we only had small errors...
+	bSuccess = (NewResultMesh.TriangleCount() > 0);
+	 
+	if (bSuccess == false)
+	{
+		UE::Geometry::AppendError(Debug, EGeometryScriptErrorType::InvalidInputs, LOCTEXT("BooleanUnion_Failed", "BooleanUnion: Boolean operation failed"));
+		return TargetMesh;
+	}
+
+	// Boolean result is in the space of TargetTransform, so invert that
+	MeshTransforms::ApplyTransform(NewResultMesh, (UE::Geometry::FTransform3d)TargetTransform.Inverse());
+
+	if (NewBoundaryEdges.Num() > 0 && Options.bFillHoles)
+	{
+		FMeshBoundaryLoops OpenBoundary(&NewResultMesh, false);
+		TSet<int> ConsiderEdges(NewBoundaryEdges);
+		OpenBoundary.EdgeFilterFunc = [&ConsiderEdges](int EID)
+		{
+			return ConsiderEdges.Contains(EID);
+		};
+		OpenBoundary.Compute();
+
+		for (FEdgeLoop& Loop : OpenBoundary.Loops)
+		{
+			FMinimalHoleFiller Filler(&NewResultMesh, Loop);
+			Filler.Fill();
+		}
+	}
+
+	TargetMesh->SetMesh(MoveTemp(NewResultMesh));
+
+	return TargetMesh;
+}
+
+
+
+
+
+UDynamicMesh* UGeometryScriptLibrary_MeshBooleanFunctions::ApplyMeshSelfUnion(
+	UDynamicMesh* TargetMesh,
+	FGeometryScriptMeshSelfUnionOptions Options,
+	UGeometryScriptDebug* Debug)
+{
+	if (TargetMesh == nullptr)
+	{
+		UE::Geometry::AppendError(Debug, EGeometryScriptErrorType::InvalidInputs, LOCTEXT("ApplyMeshSelfUnion_InvalidInput", "ApplyMeshSelfUnion: TargetMesh is Null"));
+		return TargetMesh;
+	}
+
+	// we emit multiple change events below...not clear how to avoid this though as
+	// the hole-fill edit may not occur
+
+	bool bSuccess = false;
+	TArray<int> NewBoundaryEdges;
+	TargetMesh->EditMesh([&](FDynamicMesh3& EditMesh)
+	{
+		FMeshSelfUnion Union(&EditMesh);
+		Union.WindingThreshold = FMath::Clamp(Options.WindingThreshold, 0.0f, 1.0f);
+		Union.bTrimFlaps = Options.bTrimFlaps;
+		Union.bSimplifyAlongNewEdges = Options.bSimplifyOutput;
+		Union.SimplificationAngleTolerance = Options.SimplifyPlanarTolerance;
+		bSuccess = Union.Compute();
+		NewBoundaryEdges = MoveTemp(Union.CreatedBoundaryEdges);
+	}, EDynamicMeshChangeType::GeneralEdit, EDynamicMeshAttributeChangeFlags::Unknown, false);
+
+	// currently ignoring bSuccess as it returns false in many cases
+	//if (bSuccess == false)
+	//{
+	//	UE::Geometry::AppendError(Debug, EGeometryScriptErrorType::InvalidInputs, LOCTEXT("BooleanUnion_Failed", "BooleanUnion: Boolean operation failed"));
+	//	return TargetMesh;
+	//}
+
+	if (NewBoundaryEdges.Num() > 0 && Options.bFillHoles)
+	{
+		TargetMesh->EditMesh([&](FDynamicMesh3& EditMesh)
+		{
+			FMeshBoundaryLoops OpenBoundary(&EditMesh, false);
+			TSet<int> ConsiderEdges(NewBoundaryEdges);
+			OpenBoundary.EdgeFilterFunc = [&ConsiderEdges](int EID)
+			{
+				return ConsiderEdges.Contains(EID);
+			};
+			OpenBoundary.Compute();
+
+			for (FEdgeLoop& Loop : OpenBoundary.Loops)
+			{
+				FMinimalHoleFiller Filler(&EditMesh, Loop);
+				Filler.Fill();
+			}
+		}, EDynamicMeshChangeType::GeneralEdit, EDynamicMeshAttributeChangeFlags::Unknown, false);
+	}
+
+	return TargetMesh;
+}
+
+
+
+
+
+
+
+UDynamicMesh* UGeometryScriptLibrary_MeshBooleanFunctions::ApplyMeshPlaneCut(
+	UDynamicMesh* TargetMesh,
+	FVector CutPlaneOrigin,
+	FVector CutPlaneNormal,
+	FGeometryScriptMeshPlaneCutOptions Options,
+	UGeometryScriptDebug* Debug)
+{
+	if (TargetMesh == nullptr)
+	{
+		UE::Geometry::AppendError(Debug, EGeometryScriptErrorType::InvalidInputs, LOCTEXT("ApplyMeshPlaneCut_InvalidInput", "ApplyMeshPlaneCut: TargetMesh is Null"));
+		return TargetMesh;
+	}
+
+	if (Options.bFlipCutSide)
+	{
+		CutPlaneNormal = -CutPlaneNormal;
+	}
+	FVector3d UseNormal = CutPlaneNormal.GetSafeNormal();
+	if (FMath::Abs(1.0 - UseNormal.Length()) > 0.1)
+	{
+		UE::Geometry::AppendError(Debug, EGeometryScriptErrorType::InvalidInputs, LOCTEXT("ApplyMeshPlaneCut_InvalidNormal", "ApplyMeshPlaneCut: Normal vector is degenerate"));
+		UseNormal = FVector3d::UnitZ();
+	}
+
+	bool bSuccess = false;
+	TargetMesh->EditMesh([&](FDynamicMesh3& EditMesh)
+	{
+		FMeshPlaneCut Cut(&EditMesh, CutPlaneOrigin, UseNormal);
+		bSuccess = Cut.Cut();
+
+		if (Options.bFillHoles)
+		{
+			Cut.HoleFill(ConstrainedDelaunayTriangulate<double>, Options.bFillSpans);
+		}
+
+	}, EDynamicMeshChangeType::GeneralEdit, EDynamicMeshAttributeChangeFlags::Unknown, false);
+
+	return TargetMesh;
+}
+
+
+
+
+
+UDynamicMesh* UGeometryScriptLibrary_MeshBooleanFunctions::ApplyMeshMirror(
+	UDynamicMesh* TargetMesh,
+	FVector MirrorPlaneOrigin,
+	FVector MirrorPlaneNormal,
+	FGeometryScriptMeshMirrorOptions Options,
+	UGeometryScriptDebug* Debug)
+{
+	if (TargetMesh == nullptr)
+	{
+		UE::Geometry::AppendError(Debug, EGeometryScriptErrorType::InvalidInputs, LOCTEXT("ApplyMeshMirror_InvalidInput", "ApplyMeshMirror: TargetMesh is Null"));
+		return TargetMesh;
+	}
+
+	const double PlaneTolerance = FMathf::ZeroTolerance * 10.0;
+
+	if (Options.bApplyPlaneCut && Options.bFlipCutSide)
+	{
+		MirrorPlaneNormal = -MirrorPlaneNormal;
+	}
+	FVector3d UseNormal = MirrorPlaneNormal.GetSafeNormal();
+	if (FMath::Abs(1.0 - UseNormal.Length()) > 0.1)
+	{
+		UE::Geometry::AppendError(Debug, EGeometryScriptErrorType::InvalidInputs, LOCTEXT("ApplyMeshMirror_InvalidNormal", "ApplyMeshMirror: Normal vector is degenerate"));
+		UseNormal = FVector3d::UnitZ();
+	}
+
+	bool bSuccess = false;
+	TargetMesh->EditMesh([&](FDynamicMesh3& EditMesh)
+	{
+		if (Options.bApplyPlaneCut)
+		{
+			FMeshPlaneCut Cutter(&EditMesh, MirrorPlaneOrigin, UseNormal);
+			Cutter.PlaneTolerance = PlaneTolerance;
+			Cutter.Cut();
+		}
+
+		FMeshMirror Mirrorer(&EditMesh, MirrorPlaneOrigin, UseNormal);
+		Mirrorer.bWeldAlongPlane = Options.bWeldAlongPlane;
+		Mirrorer.bAllowBowtieVertexCreation = false;
+		Mirrorer.PlaneTolerance = PlaneTolerance;
+
+		Mirrorer.MirrorAndAppend(nullptr);
+
+	}, EDynamicMeshChangeType::GeneralEdit, EDynamicMeshAttributeChangeFlags::Unknown, false);
+
+	return TargetMesh;
+}
+
+
+
+
+#undef LOCTEXT_NAMESPACE
\ No newline at end of file
diff --git a/Engine/Plugins/Experimental/GeometryScripting/Source/GeometryScriptingCore/Private/MeshDeformFunctions.cpp b/Engine/Plugins/Experimental/GeometryScripting/Source/GeometryScriptingCore/Private/MeshDeformFunctions.cpp
new file mode 100644
--- /dev/null
+++ b/Engine/Plugins/Experimental/GeometryScripting/Source/GeometryScriptingCore/Private/MeshDeformFunctions.cpp
@@ -0,0 +1,145 @@
+// Copyright Epic Games, Inc. All Rights Reserved.
+
+#include "GeometryScript/MeshDeformFunctions.h"
+
+#include "DynamicMesh/DynamicMesh3.h"
+#include "DynamicMesh/DynamicMeshAttributeSet.h"
+#include "SpaceDeformerOps/BendMeshOp.h"
+#include "SpaceDeformerOps/TwistMeshOp.h"
+#include "UDynamicMesh.h"
+
+#include "ExplicitUseGeometryMathTypes.h"		// using UE::Geometry::(math types)
+using namespace UE::Geometry;
+
+#define LOCTEXT_NAMESPACE "UGeometryScriptLibrary_MeshDeformFunctions"
+
+
+UDynamicMesh* UGeometryScriptLibrary_MeshDeformFunctions::ApplyBendWarpToMesh(
+	UDynamicMesh* TargetMesh,
+	FGeometryScriptBendWarpOptions Options,
+	float BendAngle,
+	float BendExtent,
+	UGeometryScriptDebug* Debug)
+{
+	if (TargetMesh == nullptr)
+	{
+		UE::Geometry::AppendError(Debug, EGeometryScriptErrorType::InvalidInputs, LOCTEXT("ApplyBendWarpToMesh_InvalidInput", "ApplyBendWarpToMesh: TargetMesh is Null"));
+		return TargetMesh;
+	}
+
+	FVector3d UseAxis = Options.BendAxis.GetSafeNormal();
+	if (FMath::Abs(1.0 - UseAxis.Length()) > 0.1)
+	{
+		UE::Geometry::AppendError(Debug, EGeometryScriptErrorType::InvalidInputs, LOCTEXT("ApplyBendWarpToMesh_InvalidAxis", "ApplyBendWarpToMesh: BendAxis is degenerate"));
+		UseAxis = FVector3d::UnitZ();
+	}
+	FFrame3d WarpFrame = FFrame3d(Options.BendOrigin, UseAxis);
+
+	if (Options.TowardAxis.Length() > 0)
+	{
+		FVector UseTowardAxis = Options.TowardAxis.GetSafeNormal();
+		if ( WarpFrame.Z().Dot(UseTowardAxis) < 0.95 )
+		{
+			WarpFrame.ConstrainedAlignAxis(1, UseTowardAxis, WarpFrame.Z());
+		}
+	}
+
+	if (Options.BendAxisRotation != 0)
+	{
+		WarpFrame.Rotate(FQuaterniond( WarpFrame.Z(), Options.BendAxisRotation, true));
+	}
+
+	TargetMesh->EditMesh([&](FDynamicMesh3& EditMesh) 
+	{
+		// todo extract bend warp into standalone math code
+
+		TSharedPtr<FDynamicMesh3> TmpMeshPtr = MakeShared<FDynamicMesh3>();
+		*TmpMeshPtr = MoveTemp(EditMesh);
+
+		FBendMeshOp BendOp;
+		BendOp.OriginalMesh = TmpMeshPtr;
+		BendOp.GizmoFrame = WarpFrame;
+		BendOp.LowerBoundsInterval = (Options.bSymmetricExtents) ? -BendExtent : -Options.LowerExtent;
+		BendOp.UpperBoundsInterval = BendExtent;
+		BendOp.BendDegrees = BendAngle;
+		BendOp.bLockBottom = !Options.bBidirectional;
+
+		BendOp.CalculateResult(nullptr);
+
+		TUniquePtr<FDynamicMesh3> NewResultMesh = BendOp.ExtractResult();
+		FDynamicMesh3* NewResultMeshPtr = NewResultMesh.Release();
+		EditMesh = MoveTemp(*NewResultMeshPtr);
+
+	}, EDynamicMeshChangeType::GeneralEdit, EDynamicMeshAttributeChangeFlags::Unknown, false);
+
+	return TargetMesh;
+}
+
+
+
+
+
+UDynamicMesh* UGeometryScriptLibrary_MeshDeformFunctions::ApplyTwistWarpToMesh(
+	UDynamicMesh* TargetMesh,
+	FGeometryScriptTwistWarpOptions Options,
+	float TwistAngle,
+	float TwistExtent,
+	UGeometryScriptDebug* Debug)
+{
+	if (TargetMesh == nullptr)
+	{
+		UE::Geometry::AppendError(Debug, EGeometryScriptErrorType::InvalidInputs, LOCTEXT("ApplyTwistWarpToMesh_InvalidInput", "ApplyTwistWarpToMesh: TargetMesh is Null"));
+		return TargetMesh;
+	}
+
+	FVector3d UseAxis = Options.TwistAxis.GetSafeNormal();
+	if (FMath::Abs(1.0 - UseAxis.Length()) > 0.1)
+	{
+		UE::Geometry::AppendError(Debug, EGeometryScriptErrorType::InvalidInputs, LOCTEXT("ApplyTwistWarpToMesh_InvalidAxis", "ApplyTwistWarpToMesh: TwistAxis is degenerate"));
+		UseAxis = FVector3d::UnitZ();
+	}
+	FFrame3d WarpFrame = FFrame3d(Options.TwistOrigin, UseAxis);
+
+	if (Options.TowardAxis.Length() > 0)
+	{
+		FVector UseTowardAxis = Options.TowardAxis.GetSafeNormal();
+		if ( WarpFrame.Z().Dot(UseTowardAxis) < 0.95 )
+		{
+			WarpFrame.ConstrainedAlignAxis(1, UseTowardAxis, WarpFrame.Z());
+		}
+	}
+
+	if (Options.TwistAxisRotation != 0)
+	{
+		WarpFrame.Rotate(FQuaterniond( WarpFrame.Z(), Options.TwistAxisRotation, true));
+	}
+
+	TargetMesh->EditMesh([&](FDynamicMesh3& EditMesh) 
+	{
+		// todo extract Twist warp into standalone math code
+		TSharedPtr<FDynamicMesh3> TmpMeshPtr = MakeShared<FDynamicMesh3>();
+		*TmpMeshPtr = MoveTemp(EditMesh);
+
+		FTwistMeshOp TwistOp;
+		TwistOp.OriginalMesh = TmpMeshPtr;
+		TwistOp.GizmoFrame = WarpFrame;
+		TwistOp.LowerBoundsInterval = (Options.bSymmetricExtents) ? -TwistExtent : -Options.LowerExtent;
+		TwistOp.UpperBoundsInterval = TwistExtent;
+		TwistOp.TwistDegrees = TwistAngle;
+		TwistOp.bLockBottom = !Options.bBidirectional;
+
+		TwistOp.CalculateResult(nullptr);
+
+		TUniquePtr<FDynamicMesh3> NewResultMesh = TwistOp.ExtractResult();
+		FDynamicMesh3* NewResultMeshPtr = NewResultMesh.Release();
+		EditMesh = MoveTemp(*NewResultMeshPtr);
+
+	}, EDynamicMeshChangeType::GeneralEdit, EDynamicMeshAttributeChangeFlags::Unknown, false);
+
+	return TargetMesh;
+}
+
+
+
+
+#undef LOCTEXT_NAMESPACE
\ No newline at end of file
diff --git a/Engine/Plugins/Experimental/GeometryScripting/Source/GeometryScriptingCore/Private/MeshNormalsFunctions.cpp b/Engine/Plugins/Experimental/GeometryScripting/Source/GeometryScriptingCore/Private/MeshNormalsFunctions.cpp
new file mode 100644
--- /dev/null
+++ b/Engine/Plugins/Experimental/GeometryScripting/Source/GeometryScriptingCore/Private/MeshNormalsFunctions.cpp
@@ -0,0 +1,231 @@
+// Copyright Epic Games, Inc. All Rights Reserved.
+
+#include "GeometryScript/MeshNormalsFunctions.h"
+
+#include "DynamicMesh/DynamicMesh3.h"
+#include "DynamicMesh/MeshNormals.h"
+#include "Operations/RepairOrientation.h"
+#include "DynamicMesh/DynamicMeshAABBTree3.h"
+#include "Polygroups/PolygroupSet.h"
+#include "UDynamicMesh.h"
+
+#include "ExplicitUseGeometryMathTypes.h"		// using UE::Geometry::(math types)
+using namespace UE::Geometry;
+
+#define LOCTEXT_NAMESPACE "UGeometryScriptLibrary_MeshNormalsFunctions"
+
+
+UDynamicMesh* UGeometryScriptLibrary_MeshNormalsFunctions::FlipNormals( UDynamicMesh* TargetMesh, UGeometryScriptDebug* Debug )
+{
+	if (TargetMesh == nullptr)
+	{
+		UE::Geometry::AppendError(Debug, EGeometryScriptErrorType::InvalidInputs, LOCTEXT("FlipNormals_InvalidInput", "FlipNormals: TargetMesh is Null"));
+		return TargetMesh;
+	}
+
+	// todo: publish correct change types
+	TargetMesh->EditMesh([&](FDynamicMesh3& EditMesh) 
+	{
+		EditMesh.ReverseOrientation(true);
+
+		if (EditMesh.HasAttributes() && EditMesh.Attributes()->PrimaryNormals() != nullptr)
+		{
+			FDynamicMeshNormalOverlay* Normals = EditMesh.Attributes()->PrimaryNormals();
+			for (int elemid : Normals->ElementIndicesItr())
+			{
+				FVector3f Normal = Normals->GetElement(elemid);
+				Normals->SetElement(elemid, -Normal);
+			}
+		}
+
+	}, EDynamicMeshChangeType::GeneralEdit, EDynamicMeshAttributeChangeFlags::Unknown, false);
+
+	return TargetMesh;
+}
+
+
+
+UDynamicMesh* UGeometryScriptLibrary_MeshNormalsFunctions::AutoRepairNormals(UDynamicMesh* TargetMesh, UGeometryScriptDebug* Debug)
+{
+	if (TargetMesh == nullptr)
+	{
+		UE::Geometry::AppendError(Debug, EGeometryScriptErrorType::InvalidInputs, LOCTEXT("AutoRepairNormals_InvalidInput", "AutoRepairNormals: TargetMesh is Null"));
+		return TargetMesh;
+	}
+
+	// todo: publish correct change types
+	TargetMesh->EditMesh([&](FDynamicMesh3& EditMesh) 
+	{
+		FMeshRepairOrientation Repair(&EditMesh);
+		Repair.OrientComponents();
+		FDynamicMeshAABBTree3 Tree(&EditMesh, true);
+		Repair.SolveGlobalOrientation(&Tree);
+
+	}, EDynamicMeshChangeType::GeneralEdit, EDynamicMeshAttributeChangeFlags::Unknown, false);
+
+	return TargetMesh;
+}
+
+
+
+UDynamicMesh* UGeometryScriptLibrary_MeshNormalsFunctions::SetPerVertexNormals(UDynamicMesh* TargetMesh, UGeometryScriptDebug* Debug)
+{
+	if (TargetMesh == nullptr)
+	{
+		UE::Geometry::AppendError(Debug, EGeometryScriptErrorType::InvalidInputs, LOCTEXT("SetPerVertexNormals_InvalidInput", "SetPerVertexNormals: TargetMesh is Null"));
+		return TargetMesh;
+	}
+
+	// todo: publish correct change types
+	TargetMesh->EditMesh([&](FDynamicMesh3& EditMesh) 
+	{
+		if (EditMesh.HasAttributes() == false)
+		{
+			EditMesh.EnableAttributes();
+		}
+		FMeshNormals::InitializeOverlayToPerVertexNormals(EditMesh.Attributes()->PrimaryNormals(), false);
+
+	}, EDynamicMeshChangeType::GeneralEdit, EDynamicMeshAttributeChangeFlags::Unknown, false);
+
+	return TargetMesh;
+}
+
+
+
+UDynamicMesh* UGeometryScriptLibrary_MeshNormalsFunctions::SetPerFaceNormals(UDynamicMesh* TargetMesh, UGeometryScriptDebug* Debug)
+{
+	if (TargetMesh == nullptr)
+	{
+		UE::Geometry::AppendError(Debug, EGeometryScriptErrorType::InvalidInputs, LOCTEXT("SetPerFaceNormals_InvalidInput", "SetPerFaceNormals: TargetMesh is Null"));
+		return TargetMesh;
+	}
+
+	// todo: publish correct change types
+	TargetMesh->EditMesh([&](FDynamicMesh3& EditMesh) 
+	{
+		FMeshNormals::InitializeMeshToPerTriangleNormals(&EditMesh);
+
+	}, EDynamicMeshChangeType::GeneralEdit, EDynamicMeshAttributeChangeFlags::Unknown, false);
+
+	return TargetMesh;
+}
+
+
+
+UDynamicMesh* UGeometryScriptLibrary_MeshNormalsFunctions::RecomputeNormals(
+	UDynamicMesh* TargetMesh,
+	FGeometryScriptCalculateNormalsOptions CalculateOptions,
+	UGeometryScriptDebug* Debug)
+{
+	if (TargetMesh == nullptr)
+	{
+		UE::Geometry::AppendError(Debug, EGeometryScriptErrorType::InvalidInputs, LOCTEXT("RecomputeNormals_InvalidInput", "RecomputeNormals: TargetMesh is Null"));
+		return TargetMesh;
+	}
+
+	// todo: publish correct change types
+	TargetMesh->EditMesh([&](FDynamicMesh3& EditMesh) 
+	{
+		if (EditMesh.HasAttributes() == false)
+		{
+			EditMesh.EnableAttributes();
+		}
+		FMeshNormals MeshNormals(&EditMesh);
+		MeshNormals.RecomputeOverlayNormals(EditMesh.Attributes()->PrimaryNormals(), CalculateOptions.bAreaWeighted, CalculateOptions.bAngleWeighted);
+		MeshNormals.CopyToOverlay(EditMesh.Attributes()->PrimaryNormals(), false);
+
+	}, EDynamicMeshChangeType::GeneralEdit, EDynamicMeshAttributeChangeFlags::Unknown, false);
+
+	return TargetMesh;
+}
+
+
+
+UDynamicMesh* UGeometryScriptLibrary_MeshNormalsFunctions::ComputeSplitNormals( 
+	UDynamicMesh* TargetMesh, 
+	FGeometryScriptSplitNormalsOptions SplitOptions, 
+	FGeometryScriptCalculateNormalsOptions CalculateOptions,
+	UGeometryScriptDebug* Debug)
+{
+	if (TargetMesh == nullptr)
+	{
+		UE::Geometry::AppendError(Debug, EGeometryScriptErrorType::InvalidInputs, LOCTEXT("ComputeSplitNormals_InvalidInput", "ComputeSplitNormals: TargetMesh is Null"));
+		return TargetMesh;
+	}
+
+	// todo: publish correct change types
+	TargetMesh->EditMesh([&](FDynamicMesh3& EditMesh) 
+	{
+		if (EditMesh.HasAttributes() == false)
+		{
+			EditMesh.EnableAttributes();
+		}
+
+		float NormalDotProdThreshold = FMathf::Cos(SplitOptions.OpeningAngleDeg * FMathf::DegToRad);
+		FMeshNormals FaceNormals(&EditMesh);
+		FaceNormals.ComputeTriangleNormals();
+		const TArray<FVector3d>& Normals = FaceNormals.GetNormals();
+
+		FPolygroupSet FaceGroups(&EditMesh, SplitOptions.FaceGroupLayerName);
+
+		// create new normal overlay topology based on splitting criteria
+		EditMesh.Attributes()->PrimaryNormals()->CreateFromPredicate([&](int VID, int TA, int TB)
+		{
+			if (SplitOptions.bSplitByOpeningAngle && Normals[TA].Dot(Normals[TB]) < NormalDotProdThreshold)
+			{
+				return false;
+			}
+			if (SplitOptions.bSplitByFaceGroup && FaceGroups.GetTriangleGroup(TA) != FaceGroups.GetTriangleGroup(TB))
+			{
+				return false;
+			}
+			return true;
+
+		}, 0);
+
+		// recompute the normals in the overlay
+		FMeshNormals RecomputeMeshNormals(&EditMesh);
+		RecomputeMeshNormals.RecomputeOverlayNormals(EditMesh.Attributes()->PrimaryNormals(), CalculateOptions.bAreaWeighted, CalculateOptions.bAngleWeighted);
+		RecomputeMeshNormals.CopyToOverlay(EditMesh.Attributes()->PrimaryNormals(), false);
+
+		// todo: bAllowSharpVertices from EditNormalsOp?
+
+	}, EDynamicMeshChangeType::GeneralEdit, EDynamicMeshAttributeChangeFlags::Unknown, false);
+
+	return TargetMesh;
+}
+
+
+
+
+UDynamicMesh* UGeometryScriptLibrary_MeshNormalsFunctions::SetMeshTriangleNormals(
+	UDynamicMesh* TargetMesh,
+	int TriangleID,
+	FGeometryScriptTriangle Normals,
+	bool& bIsValidTriangle,
+	bool bDeferChangeNotifications)
+{
+	bIsValidTriangle = false;
+	if (TargetMesh)
+	{
+		TargetMesh->EditMesh([&](FDynamicMesh3& EditMesh)
+		{
+			if (EditMesh.IsTriangle(TriangleID) && EditMesh.HasAttributes() && EditMesh.Attributes()->PrimaryNormals() != nullptr)
+			{
+				bIsValidTriangle = true;
+				FDynamicMeshNormalOverlay* NormalOverlay = EditMesh.Attributes()->PrimaryNormals();
+				int32 Elem0 = NormalOverlay->AppendElement((FVector3f)Normals.Vector0);
+				int32 Elem1 = NormalOverlay->AppendElement((FVector3f)Normals.Vector1);
+				int32 Elem2 = NormalOverlay->AppendElement((FVector3f)Normals.Vector2);
+				NormalOverlay->SetTriangle(TriangleID, FIndex3i(Elem0, Elem1, Elem2), true);
+			}
+		}, EDynamicMeshChangeType::GeneralEdit, EDynamicMeshAttributeChangeFlags::Unknown, bDeferChangeNotifications);
+	}
+	return TargetMesh;	
+}
+
+
+
+
+
+#undef LOCTEXT_NAMESPACE
diff --git a/Engine/Plugins/Experimental/GeometryScripting/Source/GeometryScriptingCore/Private/MeshPrimitiveFunctions.cpp b/Engine/Plugins/Experimental/GeometryScripting/Source/GeometryScriptingCore/Private/MeshPrimitiveFunctions.cpp
new file mode 100644
--- /dev/null
+++ b/Engine/Plugins/Experimental/GeometryScripting/Source/GeometryScriptingCore/Private/MeshPrimitiveFunctions.cpp
@@ -0,0 +1,505 @@
+// Copyright Epic Games, Inc. All Rights Reserved.
+
+#include "GeometryScript/MeshPrimitiveFunctions.h"
+
+#include "DynamicMesh/DynamicMesh3.h"
+#include "DynamicMesh/MeshNormals.h"
+#include "DynamicMesh/MeshTransforms.h"
+#include "DynamicMeshEditor.h"
+#include "UDynamicMesh.h"
+
+#include "Generators/GridBoxMeshGenerator.h"
+#include "Generators/SphereGenerator.h"
+#include "Generators/BoxSphereGenerator.h"
+#include "Generators/CapsuleGenerator.h"
+#include "Generators/RectangleMeshGenerator.h"
+#include "Generators/DiscMeshGenerator.h"
+#include "Generators/SweepGenerator.h"
+
+#include "ExplicitUseGeometryMathTypes.h"		// using UE::Geometry::(math types)
+using namespace UE::Geometry;
+
+#define LOCTEXT_NAMESPACE "UGeometryScriptLibrary_MeshPrimitiveFunctions"
+
+
+static void AppendPrimitive(UDynamicMesh* TargetMesh, FMeshShapeGenerator* Generator, FTransform Transform, FGeometryScriptPrimitiveOptions PrimitiveOptions)
+{
+	auto ApplyOptionsToMesh = [&Transform, &PrimitiveOptions](FDynamicMesh3& Mesh)
+	{
+		MeshTransforms::ApplyTransform(Mesh, (UE::Geometry::FTransform3d)Transform);
+		if (PrimitiveOptions.PolygroupMode == EGeometryScriptPrimitivePolygroupMode::SingleGroup)
+		{
+			for (int32 tid : Mesh.TriangleIndicesItr())
+			{
+				Mesh.SetTriangleGroup(tid, 0);
+			}
+		}
+		if (PrimitiveOptions.bFlipOrientation)
+		{
+			Mesh.ReverseOrientation(true);
+			if (Mesh.HasAttributes())
+			{
+				FDynamicMeshNormalOverlay* Normals = Mesh.Attributes()->PrimaryNormals();
+				for (int elemid : Normals->ElementIndicesItr())
+				{
+					Normals->SetElement(elemid, -Normals->GetElement(elemid));
+				}
+			}
+		}
+	};
+
+	if (TargetMesh->IsEmpty())
+	{
+		TargetMesh->EditMesh([&](FDynamicMesh3& EditMesh)
+		{
+			EditMesh.Copy(Generator);
+			ApplyOptionsToMesh(EditMesh);
+		}, EDynamicMeshChangeType::GeneralEdit, EDynamicMeshAttributeChangeFlags::Unknown, false);
+	}
+	else
+	{
+		FDynamicMesh3 TempMesh(Generator);
+		ApplyOptionsToMesh(TempMesh);
+		TargetMesh->EditMesh([&](FDynamicMesh3& EditMesh)
+		{
+			FMeshIndexMappings TmpMappings;
+			FDynamicMeshEditor Editor(&EditMesh);
+			Editor.AppendMesh(&TempMesh, TmpMappings);
+
+		}, EDynamicMeshChangeType::GeneralEdit, EDynamicMeshAttributeChangeFlags::Unknown, false);
+	}
+}
+
+
+
+
+UDynamicMesh* UGeometryScriptLibrary_MeshPrimitiveFunctions::AppendBox(
+	UDynamicMesh* TargetMesh,
+	FGeometryScriptPrimitiveOptions PrimitiveOptions,
+	FTransform Transform,
+	FBox Box,
+	int32 StepsX,
+	int32 StepsY,
+	int32 StepsZ,
+	UGeometryScriptDebug* Debug)
+{
+	if (TargetMesh == nullptr)
+	{
+		UE::Geometry::AppendError(Debug, EGeometryScriptErrorType::InvalidInputs, LOCTEXT("PrimitiveFunctions_AppendBox", "AppendBox: TargetMesh is Null"));
+		return TargetMesh;
+	}
+
+	UE::Geometry::FAxisAlignedBox3d ConvertBox(Box);
+	if (ConvertBox.MaxDim() == 0)
+	{
+		ConvertBox = UE::Geometry::FAxisAlignedBox3d(FVector3d(-50, -50, 0), FVector3d(50, 50, 100));
+	}
+	
+	// todo: could apply transform to box here and skip later?
+	// todo: if Steps X/Y/Z are zero, can use trivial box generator
+
+	FGridBoxMeshGenerator GridBoxGenerator;
+	GridBoxGenerator.Box = UE::Geometry::FOrientedBox3d(ConvertBox);
+	GridBoxGenerator.EdgeVertices = FIndex3i(FMath::Max(0, StepsX), FMath::Max(0, StepsY), FMath::Max(0, StepsZ));
+	GridBoxGenerator.bPolygroupPerQuad = (PrimitiveOptions.PolygroupMode == EGeometryScriptPrimitivePolygroupMode::PerQuad);
+	GridBoxGenerator.Generate();
+
+	AppendPrimitive(TargetMesh, &GridBoxGenerator, Transform, PrimitiveOptions);
+
+	return TargetMesh;
+}
+
+
+
+UDynamicMesh* UGeometryScriptLibrary_MeshPrimitiveFunctions::AppendSphereLatLong(
+	UDynamicMesh* TargetMesh,
+	FGeometryScriptPrimitiveOptions PrimitiveOptions,
+	FTransform Transform,
+	float Radius,
+	int32 StepsPhi,
+	int32 StepsTheta,
+	UGeometryScriptDebug* Debug)
+{
+	if (TargetMesh == nullptr)
+	{
+		UE::Geometry::AppendError(Debug, EGeometryScriptErrorType::InvalidInputs, LOCTEXT("PrimitiveFunctions_AppendSphereLatLong", "AppendSphereLatLong: TargetMesh is Null"));
+		return TargetMesh;
+	}
+
+	FSphereGenerator SphereGenerator;
+	SphereGenerator.Radius = FMath::Max(FMathf::ZeroTolerance, Radius);
+	SphereGenerator.NumPhi = FMath::Max(3, StepsPhi);
+	SphereGenerator.NumTheta = FMath::Max(3, StepsTheta);
+	SphereGenerator.bPolygroupPerQuad = (PrimitiveOptions.PolygroupMode == EGeometryScriptPrimitivePolygroupMode::PerQuad);
+	SphereGenerator.Generate();
+
+	AppendPrimitive(TargetMesh, &SphereGenerator, Transform, PrimitiveOptions);
+
+	return TargetMesh;
+}
+
+
+
+UDynamicMesh* UGeometryScriptLibrary_MeshPrimitiveFunctions::AppendSphereBox(
+	UDynamicMesh* TargetMesh,
+	FGeometryScriptPrimitiveOptions PrimitiveOptions,
+	FTransform Transform,
+	float Radius,
+	int32 StepsX,
+	int32 StepsY,
+	int32 StepsZ,
+	UGeometryScriptDebug* Debug)
+{
+	if (TargetMesh == nullptr)
+	{
+		UE::Geometry::AppendError(Debug, EGeometryScriptErrorType::InvalidInputs, LOCTEXT("PrimitiveFunctions_AppendSphereBox", "AppendSphereBox: TargetMesh is Null"));
+		return TargetMesh;
+	}
+
+	FBoxSphereGenerator SphereGenerator;
+	SphereGenerator.Radius = FMath::Max(FMathf::ZeroTolerance, Radius);
+	SphereGenerator.EdgeVertices = FIndex3i(FMath::Max(0, StepsX), FMath::Max(0, StepsY), FMath::Max(0, StepsZ));
+	SphereGenerator.bPolygroupPerQuad = (PrimitiveOptions.PolygroupMode == EGeometryScriptPrimitivePolygroupMode::PerQuad);
+	SphereGenerator.Generate();
+
+	AppendPrimitive(TargetMesh, &SphereGenerator, Transform, PrimitiveOptions);
+
+	return TargetMesh;
+}
+
+
+
+
+UDynamicMesh* UGeometryScriptLibrary_MeshPrimitiveFunctions::AppendCapsule(
+	UDynamicMesh* TargetMesh,
+	FGeometryScriptPrimitiveOptions PrimitiveOptions,
+	FTransform Transform,
+	float Radius,
+	float LineLength,
+	int32 HemisphereSteps,
+	int32 CircleSteps,
+	UGeometryScriptDebug* Debug)
+{
+	if (TargetMesh == nullptr)
+	{
+		UE::Geometry::AppendError(Debug, EGeometryScriptErrorType::InvalidInputs, LOCTEXT("PrimitiveFunctions_AppendCapsule", "AppendCapsule: TargetMesh is Null"));
+		return TargetMesh;
+	}
+
+	FCapsuleGenerator CapsuleGenerator;
+	CapsuleGenerator.Radius = FMath::Max(FMathf::ZeroTolerance, Radius);
+	CapsuleGenerator.SegmentLength = FMath::Max(FMathf::ZeroTolerance, LineLength);
+	CapsuleGenerator.NumHemisphereArcSteps = FMath::Max(2, HemisphereSteps);
+	CapsuleGenerator.NumCircleSteps = FMath::Max(3, CircleSteps);
+	CapsuleGenerator.bPolygroupPerQuad = (PrimitiveOptions.PolygroupMode == EGeometryScriptPrimitivePolygroupMode::PerQuad);
+	CapsuleGenerator.Generate();
+
+	AppendPrimitive(TargetMesh, &CapsuleGenerator, Transform, PrimitiveOptions);
+
+	return TargetMesh;
+}
+
+
+
+
+UDynamicMesh* UGeometryScriptLibrary_MeshPrimitiveFunctions::AppendCylinder(
+	UDynamicMesh* TargetMesh,
+	FGeometryScriptPrimitiveOptions PrimitiveOptions,
+	FTransform Transform,
+	float Radius,
+	float Height,
+	int32 RadialSteps,
+	int32 HeightSteps,
+	bool bCapped,
+	UGeometryScriptDebug* Debug)
+{
+	if (TargetMesh == nullptr)
+	{
+		UE::Geometry::AppendError(Debug, EGeometryScriptErrorType::InvalidInputs, LOCTEXT("PrimitiveFunctions_AppendCylinder", "AppendCylinder: TargetMesh is Null"));
+		return TargetMesh;
+	}
+
+	FCylinderGenerator CylinderGenerator;
+	CylinderGenerator.Radius[0] = FMath::Max(FMathf::ZeroTolerance, Radius);
+	CylinderGenerator.Radius[1] = CylinderGenerator.Radius[0];
+	CylinderGenerator.Height = FMath::Max(FMathf::ZeroTolerance, Height);
+	CylinderGenerator.AngleSamples = FMath::Max(3, RadialSteps);
+	CylinderGenerator.LengthSamples = FMath::Max(0, HeightSteps);
+	CylinderGenerator.bCapped = bCapped;
+	CylinderGenerator.bPolygroupPerQuad = (PrimitiveOptions.PolygroupMode == EGeometryScriptPrimitivePolygroupMode::PerQuad);
+	CylinderGenerator.Generate();
+
+	AppendPrimitive(TargetMesh, &CylinderGenerator, Transform, PrimitiveOptions);
+
+	return TargetMesh;
+}
+
+
+
+UDynamicMesh* UGeometryScriptLibrary_MeshPrimitiveFunctions::AppendCone(
+	UDynamicMesh* TargetMesh,
+	FGeometryScriptPrimitiveOptions PrimitiveOptions,
+	FTransform Transform,
+	float BaseRadius,
+	float TopRadius,
+	float Height,
+	int32 RadialSteps,
+	int32 HeightSteps,
+	bool bCapped,
+	UGeometryScriptDebug* Debug)
+{
+	if (TargetMesh == nullptr)
+	{
+		UE::Geometry::AppendError(Debug, EGeometryScriptErrorType::InvalidInputs, LOCTEXT("PrimitiveFunctions_AppendCone", "AppendCone: TargetMesh is Null"));
+		return TargetMesh;
+	}
+
+	FCylinderGenerator CylinderGenerator;
+	CylinderGenerator.Radius[0] = FMath::Max(FMathf::ZeroTolerance, BaseRadius);
+	CylinderGenerator.Radius[1] = FMath::Max(0, TopRadius);
+	CylinderGenerator.Height = FMath::Max(FMathf::ZeroTolerance, Height);
+	CylinderGenerator.AngleSamples = FMath::Max(3, RadialSteps);
+	CylinderGenerator.LengthSamples = FMath::Max(0, HeightSteps);
+	CylinderGenerator.bCapped = bCapped;
+	CylinderGenerator.bPolygroupPerQuad = (PrimitiveOptions.PolygroupMode == EGeometryScriptPrimitivePolygroupMode::PerQuad);
+	CylinderGenerator.Generate();
+
+	AppendPrimitive(TargetMesh, &CylinderGenerator, Transform, PrimitiveOptions);
+
+	return TargetMesh;
+}
+
+
+UDynamicMesh* UGeometryScriptLibrary_MeshPrimitiveFunctions::AppendTorus(
+	UDynamicMesh* TargetMesh,
+	FGeometryScriptPrimitiveOptions PrimitiveOptions,
+	FTransform Transform,
+	float MajorRadius,
+	float MinorRadius,
+	int32 MajorSteps,
+	int32 MinorSteps,
+	UGeometryScriptDebug* Debug)
+{
+	FPolygon2d Circle = FPolygon2d::MakeCircle(FMath::Max(FMathf::ZeroTolerance, MinorRadius), FMath::Max(3, MinorSteps));
+	TArray<FVector2D> PolygonVertices;
+	for (FVector2d v : Circle.GetVertices())
+	{
+		PolygonVertices.Add((FVector2D)v);
+	}
+	return AppendSimpleRevolvePolygon(TargetMesh, PrimitiveOptions, Transform, PolygonVertices, MajorRadius, MajorSteps, Debug);
+}
+
+
+UDynamicMesh* UGeometryScriptLibrary_MeshPrimitiveFunctions::AppendSimpleRevolvePolygon(
+	UDynamicMesh* TargetMesh,
+	FGeometryScriptPrimitiveOptions PrimitiveOptions,
+	FTransform Transform,
+	const TArray<FVector2D>& PolygonVertices,
+	float Radius,
+	int32 Steps,
+	UGeometryScriptDebug* Debug)
+{
+	if (TargetMesh == nullptr)
+	{
+		UE::Geometry::AppendError(Debug, EGeometryScriptErrorType::InvalidInputs, LOCTEXT("AppendRevolvePolygon_NullMesh", "AppendRevolvePolygon: TargetMesh is Null"));
+		return TargetMesh;
+	}
+	if (PolygonVertices.Num() < 3)
+	{
+		UE::Geometry::AppendError(Debug, EGeometryScriptErrorType::InvalidInputs, LOCTEXT("AppendRevolvePolygon_InvalidPolygon", "AppendRevolvePolygon: PolygonVertices array requires at least 3 positions"));
+		return TargetMesh;
+	}
+
+	FGeneralizedCylinderGenerator RevolveGen;
+	for (FVector2D Point : PolygonVertices)
+	{
+		RevolveGen.CrossSection.AppendVertex(FVector2d(Point.X, Point.Y));
+	}
+	
+	FPolygon2d PathPoly = FPolygon2d::MakeCircle(FMath::Max(FMathf::ZeroTolerance, Radius), FMath::Max(3, Steps));
+	for (FVector2d v : PathPoly.GetVertices())
+	{
+		RevolveGen.Path.Add(FVector3d(v.X, v.Y, 0.0));
+	}
+	Algo::Reverse(RevolveGen.Path);
+
+	RevolveGen.bLoop = true;
+	RevolveGen.bCapped = false;
+	RevolveGen.bPolygroupPerQuad = (PrimitiveOptions.PolygroupMode == EGeometryScriptPrimitivePolygroupMode::PerQuad);
+	RevolveGen.InitialFrame = FFrame3d(RevolveGen.Path[0], FVector3d::UnitX(), FVector3d::UnitZ(), -FVector3d::UnitY());
+	RevolveGen.Generate();
+
+	AppendPrimitive(TargetMesh, &RevolveGen, Transform, PrimitiveOptions);
+	return TargetMesh;
+}
+
+
+
+UDynamicMesh* UGeometryScriptLibrary_MeshPrimitiveFunctions::AppendSimpleSweptPolygon(
+	UDynamicMesh* TargetMesh,
+	FGeometryScriptPrimitiveOptions PrimitiveOptions,
+	FTransform Transform,
+	const TArray<FVector2D>& PolygonVertices,
+	const TArray<FVector>& SweepPath,
+	bool bLoop,
+	bool bCapped,
+	float StartScale,
+	float EndScale,
+	UGeometryScriptDebug* Debug)
+{
+	if (TargetMesh == nullptr)
+	{
+		UE::Geometry::AppendError(Debug, EGeometryScriptErrorType::InvalidInputs, LOCTEXT("AppendSimpleSweptPolygon_NullMesh", "AppendSimpleSweptPolygon: TargetMesh is Null"));
+		return TargetMesh;
+	}
+	if (PolygonVertices.Num() < 3)
+	{
+		UE::Geometry::AppendError(Debug, EGeometryScriptErrorType::InvalidInputs, LOCTEXT("AppendSimpleSweptPolygon_InvalidPolygon", "AppendSimpleSweptPolygon: PolygonVertices array requires at least 3 positions"));
+		return TargetMesh;
+	}
+
+	FGeneralizedCylinderGenerator SweepGen;
+	for (FVector2D Point : PolygonVertices)
+	{
+		SweepGen.CrossSection.AppendVertex(FVector2d(Point.X, Point.Y));
+	}
+	for (FVector SweepPathPos : SweepPath)
+	{
+		SweepGen.Path.Add(SweepPathPos);
+	}
+
+	SweepGen.bLoop = bLoop;
+	SweepGen.bCapped = bCapped;
+	SweepGen.bPolygroupPerQuad = (PrimitiveOptions.PolygroupMode == EGeometryScriptPrimitivePolygroupMode::PerQuad);
+	SweepGen.InitialFrame = FFrame3d(SweepGen.Path[0]);
+	SweepGen.StartScale = StartScale;
+	SweepGen.EndScale = EndScale;
+
+	SweepGen.Generate();
+
+	AppendPrimitive(TargetMesh, &SweepGen, Transform, PrimitiveOptions);
+	return TargetMesh;
+}
+
+
+
+
+UDynamicMesh* UGeometryScriptLibrary_MeshPrimitiveFunctions::AppendRectangle(
+	UDynamicMesh* TargetMesh,
+	FGeometryScriptPrimitiveOptions PrimitiveOptions,
+	FTransform Transform,
+	FBox2D Box,
+	int32 StepsWidth,
+	int32 StepsHeight,
+	UGeometryScriptDebug* Debug)
+{
+	if (TargetMesh == nullptr)
+	{
+		UE::Geometry::AppendError(Debug, EGeometryScriptErrorType::InvalidInputs, LOCTEXT("PrimitiveFunctions_AppendRectangle", "AppendRectangle: TargetMesh is Null"));
+		return TargetMesh;
+	}
+
+	UE::Geometry::FAxisAlignedBox2d ConvertBox(Box);
+	if (ConvertBox.MaxDim() == 0)
+	{
+		ConvertBox = UE::Geometry::FAxisAlignedBox2d(FVector2d(-50, -50), FVector2d(50, 50));
+	}
+	FVector2d Center = ConvertBox.Center();
+
+	FRectangleMeshGenerator RectGenerator;
+	RectGenerator.Origin = FVector3d(Center.X, Center.Y, 0);
+	RectGenerator.Normal = FVector3f::UnitZ();
+	RectGenerator.Width = ConvertBox.Width();
+	RectGenerator.Height = ConvertBox.Height();
+	RectGenerator.WidthVertexCount = FMath::Max(0, StepsWidth);
+	RectGenerator.HeightVertexCount = FMath::Max(0, StepsHeight);
+	RectGenerator.bSinglePolygroup = (PrimitiveOptions.PolygroupMode != EGeometryScriptPrimitivePolygroupMode::PerQuad);
+	RectGenerator.Generate();
+
+	AppendPrimitive(TargetMesh, &RectGenerator, Transform, PrimitiveOptions);
+
+	return TargetMesh;
+}
+
+
+
+
+UDynamicMesh* UGeometryScriptLibrary_MeshPrimitiveFunctions::AppendRoundRectangle(
+	UDynamicMesh* TargetMesh,
+	FGeometryScriptPrimitiveOptions PrimitiveOptions,
+	FTransform Transform,
+	FBox2D Box,
+	float CornerRadius,
+	int32 StepsWidth,
+	int32 StepsHeight,
+	int32 StepsRound,
+	UGeometryScriptDebug* Debug)
+{
+	if (TargetMesh == nullptr)
+	{
+		UE::Geometry::AppendError(Debug, EGeometryScriptErrorType::InvalidInputs, LOCTEXT("PrimitiveFunctions_AppendRoundRectangle", "AppendRoundRectangle: TargetMesh is Null"));
+		return TargetMesh;
+	}
+
+	UE::Geometry::FAxisAlignedBox2d ConvertBox(Box);
+	if (ConvertBox.MaxDim() == 0)
+	{
+		ConvertBox = UE::Geometry::FAxisAlignedBox2d(FVector2d(-50, -50), FVector2d(50, 50));
+	}
+	FVector2d Center = ConvertBox.Center();
+
+	FRoundedRectangleMeshGenerator RectGenerator;
+	RectGenerator.Origin = FVector3d(Center.X, Center.Y, 0);
+	RectGenerator.Normal = FVector3f::UnitZ();
+	RectGenerator.Width = ConvertBox.Width();
+	RectGenerator.Height = ConvertBox.Height();
+	RectGenerator.WidthVertexCount = FMath::Max(0, StepsWidth);
+	RectGenerator.HeightVertexCount = FMath::Max(0, StepsHeight);
+	RectGenerator.Radius = FMath::Max(FMathf::ZeroTolerance, CornerRadius);
+	RectGenerator.AngleSamples = FMath::Max(StepsRound, 3);
+	RectGenerator.bSinglePolygroup = (PrimitiveOptions.PolygroupMode != EGeometryScriptPrimitivePolygroupMode::PerQuad);
+	RectGenerator.Generate();
+
+	AppendPrimitive(TargetMesh, &RectGenerator, Transform, PrimitiveOptions);
+
+	return TargetMesh;
+}
+
+
+
+UDynamicMesh* UGeometryScriptLibrary_MeshPrimitiveFunctions::AppendDisc(
+	UDynamicMesh* TargetMesh,
+	FGeometryScriptPrimitiveOptions PrimitiveOptions,
+	FTransform Transform,
+	float Radius,
+	int32 AngleSteps,
+	int32 SpokeSteps,
+	float StartAngle,
+	float EndAngle,
+	UGeometryScriptDebug* Debug)
+{
+	if (TargetMesh == nullptr)
+	{
+		UE::Geometry::AppendError(Debug, EGeometryScriptErrorType::InvalidInputs, LOCTEXT("PrimitiveFunctions_AppendDisc", "AppendDisc: TargetMesh is Null"));
+		return TargetMesh;
+	}
+
+	FDiscMeshGenerator DiscGenerator;
+	DiscGenerator.Radius = FMath::Max(FMathf::ZeroTolerance, Radius);
+	DiscGenerator.Normal = FVector3f::UnitZ();
+	DiscGenerator.AngleSamples = FMath::Max(3, AngleSteps);
+	DiscGenerator.RadialSamples = FMath::Max(3, SpokeSteps);
+	DiscGenerator.StartAngle = StartAngle;
+	DiscGenerator.EndAngle = EndAngle;
+	DiscGenerator.bSinglePolygroup = (PrimitiveOptions.PolygroupMode != EGeometryScriptPrimitivePolygroupMode::PerQuad);
+	DiscGenerator.Generate();
+
+	AppendPrimitive(TargetMesh, &DiscGenerator, Transform, PrimitiveOptions);
+
+	return TargetMesh;
+}
+
+
+
+
+#undef LOCTEXT_NAMESPACE
\ No newline at end of file
diff --git a/Engine/Plugins/Experimental/GeometryScripting/Source/GeometryScriptingCore/Private/MeshQueryFunctions.cpp b/Engine/Plugins/Experimental/GeometryScripting/Source/GeometryScriptingCore/Private/MeshQueryFunctions.cpp
new file mode 100644
--- /dev/null
+++ b/Engine/Plugins/Experimental/GeometryScripting/Source/GeometryScriptingCore/Private/MeshQueryFunctions.cpp
@@ -0,0 +1,278 @@
+// Copyright Epic Games, Inc. All Rights Reserved.
+
+#include "GeometryScript/MeshQueryFunctions.h"
+#include "UDynamicMesh.h"
+
+#include "DynamicMesh/DynamicMesh3.h"
+#include "DynamicMesh/DynamicMeshAttributeSet.h"
+#include "MeshQueries.h"
+#include "MeshBoundaryLoops.h"
+
+
+#include "ExplicitUseGeometryMathTypes.h"		// using UE::Geometry::(math types)
+using namespace UE::Geometry;
+
+#define LOCTEXT_NAMESPACE "UGeometryScriptLibrary_MeshQueryFunctions"
+
+template<typename ReturnType> 
+ReturnType SimpleMeshQuery(UDynamicMesh* Mesh, ReturnType DefaultValue, 
+	TFunctionRef<ReturnType(const FDynamicMesh3& Mesh)> QueryFunc)
+{
+	if (Mesh)
+	{
+		ReturnType RetVal;
+		Mesh->ProcessMesh([&](const FDynamicMesh3& ReadMesh)
+		{
+			RetVal = QueryFunc(ReadMesh);
+		});
+		return RetVal;
+	}
+	return DefaultValue;
+}
+
+
+template<typename ReturnType> 
+ReturnType SimpleMeshUVSetQuery(UDynamicMesh* Mesh, int32 UVSetIndex, bool& bIsValidUVSet, ReturnType DefaultValue, 
+	TFunctionRef<ReturnType(const FDynamicMesh3& Mesh, const FDynamicMeshUVOverlay& UVOverlay)> QueryFunc)
+{
+	bIsValidUVSet = false;
+	ReturnType RetVal = DefaultValue;
+	if (Mesh)
+	{
+		Mesh->ProcessMesh([&](const FDynamicMesh3& ReadMesh)
+		{
+			if (ReadMesh.HasAttributes() && UVSetIndex < ReadMesh.Attributes()->NumUVLayers() )
+			{
+				const FDynamicMeshUVOverlay* Overlay = ReadMesh.Attributes()->GetUVLayer(UVSetIndex);
+				if (Overlay != nullptr)
+				{
+					bIsValidUVSet = true;
+					RetVal = QueryFunc(ReadMesh, *ReadMesh.Attributes()->GetUVLayer(UVSetIndex));
+				}
+			}
+		});
+	}
+	return RetVal;
+}
+
+
+
+
+
+bool UGeometryScriptLibrary_MeshQueryFunctions::GetIsDenseMesh( UDynamicMesh* TargetMesh )
+{
+	return SimpleMeshQuery<bool>(TargetMesh, true, [&](const FDynamicMesh3& Mesh) { return Mesh.IsCompact(); });
+}
+
+bool UGeometryScriptLibrary_MeshQueryFunctions::GetMeshHasAttributeSet( UDynamicMesh* TargetMesh )
+{
+	return SimpleMeshQuery<bool>(TargetMesh, true, [&](const FDynamicMesh3& Mesh) { return Mesh.HasAttributes(); });
+}
+
+FBox UGeometryScriptLibrary_MeshQueryFunctions::GetMeshBoundingBox( UDynamicMesh* TargetMesh )
+{
+	return SimpleMeshQuery<FBox>(TargetMesh, FBox(EForceInit::ForceInit), [&](const FDynamicMesh3& Mesh) { return (FBox)Mesh.GetBounds(true); });
+}
+
+void UGeometryScriptLibrary_MeshQueryFunctions::GetMeshVolumeArea( UDynamicMesh* TargetMesh, float& SurfaceArea, float& Volume )
+{
+	SurfaceArea = Volume = 0;
+	SimpleMeshQuery<bool>(TargetMesh, false, [&](const FDynamicMesh3& Mesh) {
+		FVector2d VolArea = TMeshQueries<FDynamicMesh3>::GetVolumeArea(Mesh);
+		Volume = VolArea.X;
+		SurfaceArea = VolArea.Y;
+		return true;
+	});
+}
+
+bool UGeometryScriptLibrary_MeshQueryFunctions::GetIsClosedMesh( UDynamicMesh* TargetMesh )
+{
+	return SimpleMeshQuery<bool>(TargetMesh, true, [&](const FDynamicMesh3& Mesh) { return Mesh.IsClosed(); });
+}
+
+int32 UGeometryScriptLibrary_MeshQueryFunctions::GetNumOpenBorderEdges( UDynamicMesh* TargetMesh )
+{
+	return SimpleMeshQuery<int32>(TargetMesh, 0, [&](const FDynamicMesh3& Mesh) { 
+		int32 Count = 0;
+		for (int32 eid : Mesh.EdgeIndicesItr())
+		{
+			if (Mesh.IsBoundaryEdge(eid))
+			{
+				Count++;
+			}
+		}
+		return Count;
+	});
+}
+
+int32 UGeometryScriptLibrary_MeshQueryFunctions::GetNumOpenBorderLoops( UDynamicMesh* TargetMesh, bool& bAmbiguousTopologyFound )
+{
+	return SimpleMeshQuery<int32>(TargetMesh, 0, [&](const FDynamicMesh3& Mesh) { 
+		FMeshBoundaryLoops BoundaryLoops(&Mesh, false);
+		bAmbiguousTopologyFound = (BoundaryLoops.Compute() == false);
+		return BoundaryLoops.GetLoopCount();
+	});
+}
+
+
+
+int32 UGeometryScriptLibrary_MeshQueryFunctions::GetTriangleCount( UDynamicMesh* TargetMesh )
+{
+	return SimpleMeshQuery<int32>(TargetMesh, 0, [&](const FDynamicMesh3& Mesh) { return Mesh.TriangleCount(); });
+}
+
+int32 UGeometryScriptLibrary_MeshQueryFunctions::GetNumTriangleIDs( UDynamicMesh* TargetMesh )
+{
+	return SimpleMeshQuery<int32>(TargetMesh, 0, [&](const FDynamicMesh3& Mesh) { return Mesh.MaxTriangleID(); });
+}
+
+bool UGeometryScriptLibrary_MeshQueryFunctions::GetHasTriangleIDGaps( UDynamicMesh* TargetMesh )
+{
+	return SimpleMeshQuery<bool>(TargetMesh, false, [&](const FDynamicMesh3& Mesh) { return !Mesh.IsCompactT(); });
+}
+
+bool UGeometryScriptLibrary_MeshQueryFunctions::IsValidTriangleID( UDynamicMesh* TargetMesh, int32 TriangleID )
+{
+	return SimpleMeshQuery<bool>(TargetMesh, false, [&](const FDynamicMesh3& Mesh) { return Mesh.IsTriangle(TriangleID); });
+}
+
+FIntVector UGeometryScriptLibrary_MeshQueryFunctions::GetTriangleIndices(UDynamicMesh* TargetMesh, int32 TriangleID, bool& bIsValidTriangle)
+{
+	return SimpleMeshQuery<FIntVector>(TargetMesh, FIntVector::NoneValue, [&](const FDynamicMesh3& Mesh) {
+		bIsValidTriangle = Mesh.IsTriangle(TriangleID);
+		return (bIsValidTriangle) ? (FIntVector)Mesh.GetTriangle(TriangleID) : FIntVector::NoneValue;
+	});
+}
+
+void UGeometryScriptLibrary_MeshQueryFunctions::GetTrianglePositions(UDynamicMesh* TargetMesh, int32 TriangleID, bool& bIsValidTriangle, FVector& Vertex1, FVector& Vertex2, FVector& Vertex3)
+{
+	bIsValidTriangle = SimpleMeshQuery<bool>(TargetMesh, false, [&](const FDynamicMesh3& Mesh) 
+	{
+		bIsValidTriangle = Mesh.IsTriangle(TriangleID);
+		if (bIsValidTriangle)
+		{
+			Mesh.GetTriVertices<FVector>(TriangleID, Vertex1, Vertex2, Vertex3);
+		}
+		else
+		{
+			Vertex1 = Vertex2 = Vertex3 = FVector::ZeroVector;
+		}
+		return bIsValidTriangle;
+	});
+}
+
+FVector UGeometryScriptLibrary_MeshQueryFunctions::GetTriangleFaceNormal(UDynamicMesh* TargetMesh, int32 TriangleID, bool& bIsValidTriangle)
+{
+	return SimpleMeshQuery<FVector>(TargetMesh, FVector::ZeroVector, [&](const FDynamicMesh3& Mesh) {
+		bIsValidTriangle = Mesh.IsTriangle(TriangleID);
+		return (bIsValidTriangle) ? (FVector)Mesh.GetTriNormal(TriangleID) : FVector::ZeroVector;
+	});
+}
+
+
+
+int32 UGeometryScriptLibrary_MeshQueryFunctions::GetVertexCount( UDynamicMesh* TargetMesh )
+{
+	return SimpleMeshQuery<int32>(TargetMesh, 0, [&](const FDynamicMesh3& Mesh) { return Mesh.VertexCount(); });
+}
+
+int32 UGeometryScriptLibrary_MeshQueryFunctions::GetNumVertexIDs( UDynamicMesh* TargetMesh )
+{
+	return SimpleMeshQuery<int32>(TargetMesh, 0, [&](const FDynamicMesh3& Mesh) { return Mesh.MaxVertexID(); });
+}
+
+bool UGeometryScriptLibrary_MeshQueryFunctions::GetHasVertexIDGaps( UDynamicMesh* TargetMesh )
+{
+	return SimpleMeshQuery<bool>(TargetMesh, false, [&](const FDynamicMesh3& Mesh) { return !Mesh.IsCompactV(); });
+}
+
+bool UGeometryScriptLibrary_MeshQueryFunctions::IsValidVertexID( UDynamicMesh* TargetMesh, int32 VertexID )
+{
+	return SimpleMeshQuery<bool>(TargetMesh, false, [&](const FDynamicMesh3& Mesh) { return Mesh.IsVertex(VertexID); });
+}
+
+FVector UGeometryScriptLibrary_MeshQueryFunctions::GetVertexPosition(UDynamicMesh* TargetMesh, int32 VertexID, bool& bIsValidVertex)
+{
+	return SimpleMeshQuery<FVector>(TargetMesh, FVector::ZeroVector, [&](const FDynamicMesh3& Mesh) {
+		bIsValidVertex = Mesh.IsTriangle(VertexID);
+		return (bIsValidVertex) ? (FVector)Mesh.GetVertex(VertexID) : FVector::ZeroVector;
+	});
+}
+
+TArray<FVector> UGeometryScriptLibrary_MeshQueryFunctions::GetAllVertexPositions(UDynamicMesh* TargetMesh, bool& bHasVertexIDGaps)
+{
+	TArray<FVector> Positions;
+	if (TargetMesh)
+	{
+		TargetMesh->ProcessMesh([&](const FDynamicMesh3& ReadMesh)
+		{
+			Positions.Init(FVector::ZeroVector, ReadMesh.MaxVertexID());
+			for (int32 vid : ReadMesh.VertexIndicesItr())
+			{
+				Positions[vid] = (FVector)ReadMesh.GetVertex(vid);
+			}
+			bHasVertexIDGaps = ! ReadMesh.IsCompactV();
+		});
+	}
+	return Positions;
+}
+
+
+
+
+
+int32 UGeometryScriptLibrary_MeshQueryFunctions::GetNumUVSets( UDynamicMesh* TargetMesh )
+{
+	return SimpleMeshQuery<int32>(TargetMesh, 0, [&](const FDynamicMesh3& Mesh) { 
+		return Mesh.HasAttributes() ? Mesh.Attributes()->NumUVLayers() : 0; });
+}
+
+FBox2D UGeometryScriptLibrary_MeshQueryFunctions::GetUVSetBoundingBox(UDynamicMesh* TargetMesh, int UVSetIndex, bool& bIsValidUVSet, bool& bUVSetIsEmpty)
+{
+	return SimpleMeshUVSetQuery<FBox2D>(TargetMesh, UVSetIndex, bIsValidUVSet, FBox2D(EForceInit::ForceInit), 
+		[&](const FDynamicMesh3& Mesh, const FDynamicMeshUVOverlay& UVSet)
+	{
+		FAxisAlignedBox2d Bounds = FAxisAlignedBox2d::Empty();
+		int32 ValidTriCount = 0;
+		for (int32 TriangleID : Mesh.TriangleIndicesItr())
+		{
+			FVector2f A, B, C;
+			if (UVSet.IsSetTriangle(TriangleID))
+			{
+				UVSet.GetTriElements(TriangleID, A, B, C);
+				Bounds.Contain((FVector2d)A); Bounds.Contain((FVector2d)B); Bounds.Contain((FVector2d)C);
+				ValidTriCount++;
+			}
+		}
+		bUVSetIsEmpty = (ValidTriCount == 0);
+		return (ValidTriCount > 0) ? (FBox2D)Bounds : FBox2D(EForceInit::ForceInit);
+	});
+}
+
+void UGeometryScriptLibrary_MeshQueryFunctions::GetTriangleUVs(UDynamicMesh* TargetMesh, int32 UVSetIndex, int32 TriangleID, 
+	FVector2D& UV1, FVector2D& UV2, FVector2D& UV3, bool& bHaveValidUVs)
+{
+	bHaveValidUVs = SimpleMeshUVSetQuery<bool>(TargetMesh, UVSetIndex, bHaveValidUVs, false, 
+		[&](const FDynamicMesh3& Mesh, const FDynamicMeshUVOverlay& UVSet) 
+	{
+		if (Mesh.IsTriangle(TriangleID) && UVSet.IsSetTriangle(TriangleID))
+		{
+			FVector2f A, B, C;
+			UVSet.GetTriElements(TriangleID, A, B, C);
+			UV1 = (FVector2D)A; UV2 = (FVector2D)B; UV3 = (FVector2D)C;
+			return true;
+		}
+		else
+		{
+			UV1 = UV2 = UV3 = FVector2D::ZeroVector;
+			return false;
+		}
+	});
+}
+
+
+
+
+
+
+#undef LOCTEXT_NAMESPACE
diff --git a/Engine/Plugins/Experimental/GeometryScripting/Source/GeometryScriptingCore/Private/MeshRepairFunctions.cpp b/Engine/Plugins/Experimental/GeometryScripting/Source/GeometryScriptingCore/Private/MeshRepairFunctions.cpp
new file mode 100644
--- /dev/null
+++ b/Engine/Plugins/Experimental/GeometryScripting/Source/GeometryScriptingCore/Private/MeshRepairFunctions.cpp
@@ -0,0 +1,226 @@
+// Copyright Epic Games, Inc. All Rights Reserved.
+
+#include "GeometryScript/MeshRepairFunctions.h"
+
+#include "DynamicMesh/DynamicMesh3.h"
+#include "DynamicMesh/MeshNormals.h"
+#include "DynamicMesh/Operations/MergeCoincidentMeshEdges.h"
+#include "DynamicMesh/DynamicMeshAABBTree3.h"
+#include "Spatial/FastWinding.h"
+#include "Polygroups/PolygroupSet.h"
+#include "MeshBoundaryLoops.h"
+#include "DynamicMeshEditor.h"
+#include "Operations/RemoveOccludedTriangles.h"
+#include "Selections/MeshConnectedComponents.h"
+#include "Selections/MeshFaceSelection.h"
+#include "UDynamicMesh.h"
+
+#include "CleaningOps/HoleFillOp.h"
+
+#include "ExplicitUseGeometryMathTypes.h"		// using UE::Geometry::(math types)
+using namespace UE::Geometry;
+
+#define LOCTEXT_NAMESPACE "UGeometryScriptLibrary_MeshRepairFunctions"
+
+
+UDynamicMesh* UGeometryScriptLibrary_MeshRepairFunctions::CompactMesh(
+	UDynamicMesh* TargetMesh,
+	UGeometryScriptDebug* Debug)
+{
+	if (TargetMesh == nullptr)
+	{
+		UE::Geometry::AppendError(Debug, EGeometryScriptErrorType::InvalidInputs, LOCTEXT("CompactMesh_InvalidInput", "CompactMesh: TargetMesh is Null"));
+		return TargetMesh;
+	}
+
+	TargetMesh->EditMesh([&](FDynamicMesh3& EditMesh) 
+	{
+		EditMesh.CompactInPlace();
+
+	}, EDynamicMeshChangeType::GeneralEdit, EDynamicMeshAttributeChangeFlags::Unknown, false);
+
+	return TargetMesh;
+}
+
+
+
+UDynamicMesh* UGeometryScriptLibrary_MeshRepairFunctions::WeldMeshEdges(
+	UDynamicMesh* TargetMesh,
+	FGeometryScriptWeldEdgesOptions WeldOptions,
+	UGeometryScriptDebug* Debug)
+{
+	if (TargetMesh == nullptr)
+	{
+		UE::Geometry::AppendError(Debug, EGeometryScriptErrorType::InvalidInputs, LOCTEXT("WeldMeshEdges_InvalidInput", "WeldMeshEdges: TargetMesh is Null"));
+		return TargetMesh;
+	}
+
+	TargetMesh->EditMesh([&](FDynamicMesh3& EditMesh) 
+	{
+		FMergeCoincidentMeshEdges Welder(&EditMesh);
+		Welder.MergeVertexTolerance = WeldOptions.Tolerance;
+		Welder.OnlyUniquePairs = WeldOptions.bOnlyUniquePairs;
+		bool bOK = Welder.Apply();
+		if (!bOK)
+		{
+			UE::Geometry::AppendError(Debug, EGeometryScriptErrorType::OperationFailed, LOCTEXT("WeldMeshEdges_Error", "WeldMeshEdges: Weld Operation returned error flag"));
+		}
+
+	}, EDynamicMeshChangeType::GeneralEdit, EDynamicMeshAttributeChangeFlags::Unknown, false);
+
+	return TargetMesh;
+}
+
+
+
+UDynamicMesh* UGeometryScriptLibrary_MeshRepairFunctions::FillAllMeshHoles(
+	UDynamicMesh* TargetMesh,
+	FGeometryScriptFillHolesOptions FillOptions,
+	int32& NumFilledHoles,
+	int32& NumFailedHoleFills,
+	UGeometryScriptDebug* Debug)
+{
+	if (TargetMesh == nullptr)
+	{
+		UE::Geometry::AppendError(Debug, EGeometryScriptErrorType::InvalidInputs, LOCTEXT("FillAllMeshHoles_InvalidInput", "FillAllMeshHoles: TargetMesh is Null"));
+		return TargetMesh;
+	}
+
+	TSharedPtr<FDynamicMesh3> MeshCopy = MakeShared<FDynamicMesh3>();
+	TargetMesh->ProcessMesh([&](const FDynamicMesh3& EditMesh) { MeshCopy->Copy(EditMesh); });
+
+	FMeshBoundaryLoops Loops(MeshCopy.Get(), true);
+
+	FHoleFillOp HoleFiller;
+	HoleFiller.OriginalMesh = MeshCopy;
+	HoleFiller.Loops = MoveTemp(Loops.Loops);
+
+	switch (FillOptions.FillMethod)
+	{
+	case EGeometryScriptFillHolesMethod::Automatic:
+		// fall through
+	case EGeometryScriptFillHolesMethod::MinimalFill:
+		HoleFiller.FillType = EHoleFillOpFillType::Minimal;
+		break;
+	case EGeometryScriptFillHolesMethod::PolygonTriangulation:
+		HoleFiller.FillType = EHoleFillOpFillType::PolygonEarClipping;
+		break;
+	case EGeometryScriptFillHolesMethod::TriangleFan:
+		HoleFiller.FillType = EHoleFillOpFillType::TriangleFan;
+		break;
+	case EGeometryScriptFillHolesMethod::PlanarProjection:
+		HoleFiller.FillType = EHoleFillOpFillType::Planar;
+		break;
+	}
+	HoleFiller.FillOptions.bRemoveIsolatedTriangles = FillOptions.bDeleteIsolatedTriangles;
+	HoleFiller.FillOptions.bQuickFillSmallHoles = true;
+
+	HoleFiller.CalculateResult(nullptr);
+
+	NumFailedHoleFills = HoleFiller.NumFailedLoops;
+	NumFilledHoles = HoleFiller.Loops.Num() - NumFailedHoleFills;
+
+	TUniquePtr<FDynamicMesh3> NewResultMesh = HoleFiller.ExtractResult();
+	FDynamicMesh3* NewResultMeshPtr = NewResultMesh.Release();
+	TargetMesh->SetMesh(MoveTemp(*NewResultMeshPtr));
+	return TargetMesh;
+}
+
+
+
+
+UDynamicMesh* UGeometryScriptLibrary_MeshRepairFunctions::RemoveSmallComponents(
+	UDynamicMesh* TargetMesh,
+	FGeometryScriptRemoveSmallComponentOptions Options,
+	UGeometryScriptDebug* Debug)
+{
+	if (TargetMesh == nullptr)
+	{
+		UE::Geometry::AppendError(Debug, EGeometryScriptErrorType::InvalidInputs, LOCTEXT("RemoveSmallComponents_InvalidInput", "RemoveSmallComponents: TargetMesh is Null"));
+		return TargetMesh;
+	}
+
+	TargetMesh->EditMesh([&](FDynamicMesh3& EditMesh) 
+	{
+		FDynamicMeshEditor Editor(&EditMesh);
+		Editor.RemoveSmallComponents(Options.MinVolume, Options.MinArea, Options.MinTriangleCount);
+
+	}, EDynamicMeshChangeType::GeneralEdit, EDynamicMeshAttributeChangeFlags::Unknown, false);
+
+	return TargetMesh;
+}
+
+
+
+UDynamicMesh* UGeometryScriptLibrary_MeshRepairFunctions::RemoveHiddenTriangles(
+	UDynamicMesh* TargetMesh,
+	FGeometryScriptRemoveHiddenTrianglesOptions Options,
+	UGeometryScriptDebug* Debug)
+{
+	if (TargetMesh == nullptr)
+	{
+		UE::Geometry::AppendError(Debug, EGeometryScriptErrorType::InvalidInputs, LOCTEXT("RemoveHiddenTriangles_InvalidInput", "RemoveHiddenTriangles: TargetMesh is Null"));
+		return TargetMesh;
+	}
+
+	TargetMesh->EditMesh([&](FDynamicMesh3& EditMesh) 
+	{
+		TRemoveOccludedTriangles<FDynamicMesh3> Jacket(&EditMesh);
+
+		Jacket.InsideMode = static_cast<UE::Geometry::EOcclusionCalculationMode>(Options.Method);
+		//Jacket.TriangleSamplingMethod = UE::Geometry::EOcclusionTriangleSampling::Centroids
+		Jacket.TriangleSamplingMethod = UE::Geometry::EOcclusionTriangleSampling::VerticesAndCentroids;
+		Jacket.WindingIsoValue = Options.WindingIsoValue;
+		Jacket.NormalOffset = Options.NormalOffset;
+		Jacket.AddRandomRays = Options.RaysPerSample;
+		Jacket.AddTriangleSamples = Options.SamplesPerTriangle;
+
+		TArray<UE::Geometry::FTransform3d> NoTransforms;
+		NoTransforms.Add(UE::Geometry::FTransform3d::Identity());
+
+		//  set up AABBTree and FWNTree lists
+		FDynamicMeshAABBTree3 Spatial(&EditMesh);
+		TArray<FDynamicMeshAABBTree3*> OccluderTrees; 
+		OccluderTrees.Add(&Spatial);
+		
+		TFastWindingTree<FDynamicMesh3> FastWinding(&Spatial, false);
+		TArray<TFastWindingTree<FDynamicMesh3>*> OccluderWindings; 
+		if (Options.Method == EGeometryScriptRemoveHiddenTrianglesMethod::FastWindingNumber)
+		{
+			FastWinding.Build();
+		}
+		OccluderWindings.Add(&FastWinding);
+
+		Jacket.Select(NoTransforms, OccluderTrees, OccluderWindings, NoTransforms);
+
+		if (Options.ShrinkSelection > 0)
+		{
+			FMeshFaceSelection Selection(&EditMesh);
+			Selection.Select(Jacket.RemovedT);
+			Selection.ContractBorderByOneRingNeighbours(Options.ShrinkSelection);
+			Jacket.RemovedT = Selection.AsArray();
+		}
+		
+		if (Jacket.RemovedT.Num() > 0)
+		{
+			Jacket.RemoveSelected();
+
+			if (EditMesh.TriangleCount() == 0)
+			{
+				UE::Geometry::AppendError(Debug, EGeometryScriptErrorType::OperationFailed, LOCTEXT("RemoveHiddenTriangles_EmptyMesh", "RemoveHiddenTriangles: Operation resulted in an empty mesh"));
+			}
+
+			if (Options.bCompactResult)
+			{
+				EditMesh.CompactInPlace();
+			}
+		}
+
+	}, EDynamicMeshChangeType::GeneralEdit, EDynamicMeshAttributeChangeFlags::Unknown, false);
+
+	return TargetMesh;
+}
+
+
+
+#undef LOCTEXT_NAMESPACE
diff --git a/Engine/Plugins/Experimental/GeometryScripting/Source/GeometryScriptingCore/Private/MeshSimplifyFunctions.cpp b/Engine/Plugins/Experimental/GeometryScripting/Source/GeometryScriptingCore/Private/MeshSimplifyFunctions.cpp
new file mode 100644
--- /dev/null
+++ b/Engine/Plugins/Experimental/GeometryScripting/Source/GeometryScriptingCore/Private/MeshSimplifyFunctions.cpp
@@ -0,0 +1,216 @@
+// Copyright Epic Games, Inc. All Rights Reserved.
+
+#include "GeometryScript/MeshSimplifyFunctions.h"
+
+#include "DynamicMesh/DynamicMesh3.h"
+#include "DynamicMesh/MeshNormals.h"
+#include "DynamicMesh/DynamicMeshAABBTree3.h"
+#include "MeshSimplification.h"
+#include "MeshConstraintsUtil.h"
+#include "ProjectionTargets.h"
+
+#include "UDynamicMesh.h"
+
+
+#include "ExplicitUseGeometryMathTypes.h"		// using UE::Geometry::(math types)
+using namespace UE::Geometry;
+
+#define LOCTEXT_NAMESPACE "UGeometryScriptLibrary_MeshSimplifyFunctions"
+
+
+UDynamicMesh* UGeometryScriptLibrary_MeshSimplifyFunctions::ApplySimplifyToPlanar(
+	UDynamicMesh* TargetMesh,
+	FGeometryScriptPlanarSimplifyOptions Options,
+	UGeometryScriptDebug* Debug)
+{
+	if (TargetMesh == nullptr)
+	{
+		UE::Geometry::AppendError(Debug, EGeometryScriptErrorType::InvalidInputs, LOCTEXT("ApplySimplifyToPlanar_InvalidInput", "ApplySimplifyToPlanar: TargetMesh is Null"));
+		return TargetMesh;
+	}
+
+	TargetMesh->EditMesh([&](FDynamicMesh3& EditMesh) 
+	{
+		FQEMSimplification Simplifier(&EditMesh);
+
+		// todo: set up seam collapse etc?
+
+		Simplifier.CollapseMode = FQEMSimplification::ESimplificationCollapseModes::AverageVertexPosition;
+		Simplifier.SimplifyToMinimalPlanar( FMath::Max(0.00001, Options.AngleThreshold) );
+
+	}, EDynamicMeshChangeType::GeneralEdit, EDynamicMeshAttributeChangeFlags::Unknown, false);
+
+	return TargetMesh;
+}
+
+
+
+template<typename SimplificationType>
+void DoSimplifyMesh(
+	FDynamicMesh3& EditMesh, 
+	FGeometryScriptSimplifyMeshOptions Options,
+	bool bTargetIsTriCount,
+	int32 TargetCount,
+	FMeshProjectionTarget* ProjectionTarget = nullptr,
+	double GeometricTolerance = 0)
+{
+	SimplificationType Simplifier(&EditMesh);
+
+	Simplifier.ProjectionMode = SimplificationType::ETargetProjectionMode::NoProjection;
+	if (ProjectionTarget != nullptr)
+	{
+		Simplifier.SetProjectionTarget(ProjectionTarget);
+	}
+
+	Simplifier.DEBUG_CHECK_LEVEL = 0;
+	Simplifier.bRetainQuadricMemory = Options.bRetainQuadricMemory; 
+	Simplifier.bAllowSeamCollapse = Options.bAllowSeamCollapse;
+	if (Options.bAllowSeamCollapse)
+	{
+		Simplifier.SetEdgeFlipTolerance(1.e-5);
+		if (EditMesh.HasAttributes())
+		{
+			EditMesh.Attributes()->SplitAllBowties();	// eliminate any bowties that might have formed on attribute seams.
+		}
+	}
+
+	// do these flags matter here since we are not flipping??
+	EEdgeRefineFlags MeshBoundaryConstraints = EEdgeRefineFlags::NoFlip;
+	EEdgeRefineFlags GroupBorderConstraints = EEdgeRefineFlags::NoConstraint;
+	EEdgeRefineFlags MaterialBorderConstraints = EEdgeRefineFlags::NoConstraint;
+
+	FMeshConstraints Constraints;
+	FMeshConstraintsUtil::ConstrainAllBoundariesAndSeams(Constraints, EditMesh,
+		MeshBoundaryConstraints, GroupBorderConstraints, MaterialBorderConstraints,
+		Options.bAllowSeamSplits, Options.bAllowSeamSmoothing, Options.bAllowSeamCollapse);
+	Simplifier.SetExternalConstraints(MoveTemp(Constraints));
+
+	if ( ProjectionTarget != nullptr && GeometricTolerance > 0)
+	{
+		Simplifier.GeometricErrorConstraint = SimplificationType::EGeometricErrorCriteria::PredictedPointToProjectionTarget;
+		Simplifier.GeometricErrorTolerance = GeometricTolerance;
+	}
+
+	if (bTargetIsTriCount)
+	{
+		Simplifier.SimplifyToTriangleCount( FMath::Max(1,TargetCount) );
+	}
+	else
+	{
+		Simplifier.SimplifyToVertexCount( FMath::Max(1,TargetCount) );
+	}
+}
+
+
+
+UDynamicMesh* UGeometryScriptLibrary_MeshSimplifyFunctions::ApplySimplifyToTriangleCount(
+	UDynamicMesh* TargetMesh,
+	int32 TriangleCount,
+	FGeometryScriptSimplifyMeshOptions Options,
+	UGeometryScriptDebug* Debug)
+{
+	if (TargetMesh == nullptr)
+	{
+		UE::Geometry::AppendError(Debug, EGeometryScriptErrorType::InvalidInputs, LOCTEXT("ApplySimplifyToTriangleCount_InvalidInput", "ApplySimplifyToTriangleCount: TargetMesh is Null"));
+		return TargetMesh;
+	}
+
+	TargetMesh->EditMesh([&](FDynamicMesh3& EditMesh) 
+	{
+		if (Options.Method == EGeometryScriptRemoveMeshSimplificationType::AttributeAware)
+		{
+			DoSimplifyMesh<FAttrMeshSimplification>(EditMesh, Options, true, TriangleCount);
+		}
+		else if (Options.Method == EGeometryScriptRemoveMeshSimplificationType::VolumePreserving)
+		{
+			DoSimplifyMesh<FVolPresMeshSimplification>(EditMesh, Options, true, TriangleCount);
+		}
+		else
+		{
+			DoSimplifyMesh<FQEMSimplification>(EditMesh, Options, true, TriangleCount);
+		}
+
+	}, EDynamicMeshChangeType::GeneralEdit, EDynamicMeshAttributeChangeFlags::Unknown, false);
+
+	return TargetMesh;
+}
+
+
+
+
+UDynamicMesh* UGeometryScriptLibrary_MeshSimplifyFunctions::ApplySimplifyToVertexCount(
+	UDynamicMesh* TargetMesh,
+	int32 VertexCount,
+	FGeometryScriptSimplifyMeshOptions Options,
+	UGeometryScriptDebug* Debug)
+{
+	if (TargetMesh == nullptr)
+	{
+		UE::Geometry::AppendError(Debug, EGeometryScriptErrorType::InvalidInputs, LOCTEXT("ApplySimplifyToVertexCount_InvalidInput", "ApplySimplifyToVertexCount: TargetMesh is Null"));
+		return TargetMesh;
+	}
+
+	TargetMesh->EditMesh([&](FDynamicMesh3& EditMesh) 
+	{
+		if (Options.Method == EGeometryScriptRemoveMeshSimplificationType::AttributeAware)
+		{
+			DoSimplifyMesh<FAttrMeshSimplification>(EditMesh, Options, false, VertexCount);
+		}
+		else if (Options.Method == EGeometryScriptRemoveMeshSimplificationType::VolumePreserving)
+		{
+			DoSimplifyMesh<FVolPresMeshSimplification>(EditMesh, Options, false, VertexCount);
+		}
+		else
+		{
+			DoSimplifyMesh<FQEMSimplification>(EditMesh, Options, false, VertexCount);
+		}
+
+	}, EDynamicMeshChangeType::GeneralEdit, EDynamicMeshAttributeChangeFlags::Unknown, false);
+
+	return TargetMesh;
+}
+
+
+
+UDynamicMesh* UGeometryScriptLibrary_MeshSimplifyFunctions::ApplySimplifyToTolerance(
+	UDynamicMesh* TargetMesh,
+	float Tolerance,
+	FGeometryScriptSimplifyMeshOptions Options,
+	UGeometryScriptDebug* Debug)
+{
+	if (TargetMesh == nullptr)
+	{
+		UE::Geometry::AppendError(Debug, EGeometryScriptErrorType::InvalidInputs, LOCTEXT("ApplySimplifyToTolerance_InvalidInput", "ApplySimplifyToTolerance: TargetMesh is Null"));
+		return TargetMesh;
+	}
+
+	TargetMesh->EditMesh([&](FDynamicMesh3& EditMesh) 
+	{
+		FDynamicMesh3 TempCopy;
+		TempCopy.Copy(EditMesh, false, false, false, false);
+		FDynamicMeshAABBTree3 Spatial(&TempCopy, true);
+		FMeshProjectionTarget ProjTarget(&TempCopy, &Spatial);
+		float UseTolerance = FMath::Max(0.0, Tolerance);
+
+		if (Options.Method == EGeometryScriptRemoveMeshSimplificationType::AttributeAware)
+		{
+			DoSimplifyMesh<FAttrMeshSimplification>(EditMesh, Options, true, 1, &ProjTarget, UseTolerance);
+		}
+		else if (Options.Method == EGeometryScriptRemoveMeshSimplificationType::VolumePreserving)
+		{
+			DoSimplifyMesh<FVolPresMeshSimplification>(EditMesh, Options, true, 1, &ProjTarget, UseTolerance);
+		}
+		else
+		{
+			DoSimplifyMesh<FQEMSimplification>(EditMesh, Options, true, 1, &ProjTarget, UseTolerance);
+		}
+
+	}, EDynamicMeshChangeType::GeneralEdit, EDynamicMeshAttributeChangeFlags::Unknown, false);
+
+	return TargetMesh;
+}
+
+
+
+
+#undef LOCTEXT_NAMESPACE
\ No newline at end of file
diff --git a/Engine/Plugins/Experimental/GeometryScripting/Source/GeometryScriptingCore/Private/MeshTransformFunctions.cpp b/Engine/Plugins/Experimental/GeometryScripting/Source/GeometryScriptingCore/Private/MeshTransformFunctions.cpp
new file mode 100644
--- /dev/null
+++ b/Engine/Plugins/Experimental/GeometryScripting/Source/GeometryScriptingCore/Private/MeshTransformFunctions.cpp
@@ -0,0 +1,84 @@
+// Copyright Epic Games, Inc. All Rights Reserved.
+
+#include "GeometryScript/MeshTransformFunctions.h"
+
+#include "DynamicMesh/DynamicMesh3.h"
+#include "DynamicMesh/MeshTransforms.h"
+#include "UDynamicMesh.h"
+
+#include "ExplicitUseGeometryMathTypes.h"		// using UE::Geometry::(math types)
+using namespace UE::Geometry;
+
+#define LOCTEXT_NAMESPACE "UGeometryScriptLibrary_MeshTransformFunctions"
+
+
+UDynamicMesh* UGeometryScriptLibrary_MeshTransformFunctions::TransformMesh(
+	UDynamicMesh* TargetMesh,
+	FTransform Transform,
+	UGeometryScriptDebug* Debug)
+{
+	if (TargetMesh == nullptr)
+	{
+		UE::Geometry::AppendError(Debug, EGeometryScriptErrorType::InvalidInputs, LOCTEXT("TransformMesh_InvalidInput", "TransformMesh: TargetMesh is Null"));
+		return TargetMesh;
+	}
+
+	// todo: publish correct change types
+	TargetMesh->EditMesh([&](FDynamicMesh3& EditMesh) 
+	{
+		MeshTransforms::ApplyTransform(EditMesh, (UE::Geometry::FTransform3d)Transform);
+
+	}, EDynamicMeshChangeType::GeneralEdit, EDynamicMeshAttributeChangeFlags::Unknown, false);
+
+	return TargetMesh;
+}
+
+
+
+UDynamicMesh* UGeometryScriptLibrary_MeshTransformFunctions::TranslateMesh(
+	UDynamicMesh* TargetMesh,
+	FVector Translation,
+	UGeometryScriptDebug* Debug)
+{
+	if (TargetMesh == nullptr)
+	{
+		UE::Geometry::AppendError(Debug, EGeometryScriptErrorType::InvalidInputs, LOCTEXT("TranslateMesh_InvalidInput", "TranslateMesh: TargetMesh is Null"));
+		return TargetMesh;
+	}
+
+	// todo: publish correct change types
+	TargetMesh->EditMesh([&](FDynamicMesh3& EditMesh) 
+	{
+		MeshTransforms::Translate(EditMesh, (FVector3d)Translation);
+
+	}, EDynamicMeshChangeType::GeneralEdit, EDynamicMeshAttributeChangeFlags::Unknown, false);
+
+	return TargetMesh;
+}
+
+
+UDynamicMesh* UGeometryScriptLibrary_MeshTransformFunctions::ScaleMesh(
+	UDynamicMesh* TargetMesh,
+	FVector Scale,
+	UGeometryScriptDebug* Debug)
+{
+	if (TargetMesh == nullptr)
+	{
+		UE::Geometry::AppendError(Debug, EGeometryScriptErrorType::InvalidInputs, LOCTEXT("ScaleMesh_InvalidInput", "ScaleMesh: TargetMesh is Null"));
+		return TargetMesh;
+	}
+
+	// todo: publish correct change types
+	TargetMesh->EditMesh([&](FDynamicMesh3& EditMesh) 
+	{
+		MeshTransforms::Scale(EditMesh, (FVector3d)Scale, FVector3d::Zero());
+
+	}, EDynamicMeshChangeType::GeneralEdit, EDynamicMeshAttributeChangeFlags::Unknown, false);
+
+	return TargetMesh;
+}
+
+
+
+
+#undef LOCTEXT_NAMESPACE
diff --git a/Engine/Plugins/Experimental/GeometryScripting/Source/GeometryScriptingCore/Private/MeshUVFunctions.cpp b/Engine/Plugins/Experimental/GeometryScripting/Source/GeometryScriptingCore/Private/MeshUVFunctions.cpp
new file mode 100644
--- /dev/null
+++ b/Engine/Plugins/Experimental/GeometryScripting/Source/GeometryScriptingCore/Private/MeshUVFunctions.cpp
@@ -0,0 +1,50 @@
+// Copyright Epic Games, Inc. All Rights Reserved.
+
+#include "GeometryScript/MeshUVFunctions.h"
+
+#include "DynamicMesh/DynamicMesh3.h"
+#include "DynamicMesh/DynamicMeshAttributeSet.h"
+#include "Polygroups/PolygroupSet.h"
+#include "UDynamicMesh.h"
+
+#include "ExplicitUseGeometryMathTypes.h"		// using UE::Geometry::(math types)
+using namespace UE::Geometry;
+
+#define LOCTEXT_NAMESPACE "UGeometryScriptLibrary_MeshUVFunctions"
+
+
+UDynamicMesh* UGeometryScriptLibrary_MeshUVFunctions::SetMeshTriangleUVs(
+	UDynamicMesh* TargetMesh,
+	int UVSetIndex,
+	int TriangleID, 
+	FGeometryScriptUVTriangle UVs,
+	bool& bIsValidTriangle,
+	bool bDeferChangeNotifications)
+{
+	bIsValidTriangle = false;
+	if (TargetMesh)
+	{
+		TargetMesh->EditMesh([&](FDynamicMesh3& EditMesh)
+		{
+			if (EditMesh.IsTriangle(TriangleID) && EditMesh.HasAttributes() && UVSetIndex < EditMesh.Attributes()->NumUVLayers() )
+			{
+				FDynamicMeshUVOverlay* UVOverlay = EditMesh.Attributes()->GetUVLayer(UVSetIndex);
+				if (UVOverlay != nullptr)
+				{
+					bIsValidTriangle = true;
+					int32 Elem0 = UVOverlay->AppendElement((FVector2f)UVs.UV0);
+					int32 Elem1 = UVOverlay->AppendElement((FVector2f)UVs.UV1);
+					int32 Elem2 = UVOverlay->AppendElement((FVector2f)UVs.UV2);
+					UVOverlay->SetTriangle(TriangleID, FIndex3i(Elem0, Elem1, Elem2), true);
+				}
+			}
+		}, EDynamicMeshChangeType::GeneralEdit, EDynamicMeshAttributeChangeFlags::Unknown, bDeferChangeNotifications);
+	}
+	return TargetMesh;	
+}
+
+
+
+
+
+#undef LOCTEXT_NAMESPACE
\ No newline at end of file
diff --git a/Engine/Plugins/Experimental/GeometryScripting/Source/GeometryScriptingCore/Private/MeshVertexColorFunctions.cpp b/Engine/Plugins/Experimental/GeometryScripting/Source/GeometryScriptingCore/Private/MeshVertexColorFunctions.cpp
new file mode 100644
--- /dev/null
+++ b/Engine/Plugins/Experimental/GeometryScripting/Source/GeometryScriptingCore/Private/MeshVertexColorFunctions.cpp
@@ -0,0 +1,65 @@
+// Copyright Epic Games, Inc. All Rights Reserved.
+
+#include "GeometryScript/MeshVertexColorFunctions.h"
+
+#include "DynamicMesh/DynamicMesh3.h"
+#include "UDynamicMesh.h"
+
+#include "ExplicitUseGeometryMathTypes.h"		// using UE::Geometry::(math types)
+using namespace UE::Geometry;
+
+#define LOCTEXT_NAMESPACE "UGeometryScriptLibrary_MeshVertexColorFunctions"
+
+
+UDynamicMesh* UGeometryScriptLibrary_MeshVertexColorFunctions::SetMeshPerVertexColors(
+	UDynamicMesh* TargetMesh,
+	const TArray<FLinearColor>& VertexColors,
+	UGeometryScriptDebug* Debug)
+{
+	if (TargetMesh == nullptr)
+	{
+		UE::Geometry::AppendError(Debug, EGeometryScriptErrorType::InvalidInputs, LOCTEXT("SetMeshPerVertexColors_InvalidInput", "SetMeshPerVertexColors: TargetMesh is Null"));
+		return TargetMesh;
+	}
+
+	TargetMesh->EditMesh([&](FDynamicMesh3& EditMesh) 
+	{
+		if (VertexColors.Num() < EditMesh.MaxVertexID())
+		{
+			UE::Geometry::AppendError(Debug, EGeometryScriptErrorType::InvalidInputs, LOCTEXT("SetMeshPerVertexColors_IncorrectCount", "SetMeshPerVertexColors: size of provided VertexColors is smaller than MaxVertexID of Mesh"));
+		}
+		else
+		{
+			if (EditMesh.HasAttributes() == false)
+			{
+				EditMesh.EnableAttributes();
+			}
+			if (EditMesh.Attributes()->HasPrimaryColors() == false)
+			{
+				EditMesh.Attributes()->EnablePrimaryColors();
+			}
+			FDynamicMeshColorOverlay* Colors = EditMesh.Attributes()->PrimaryColors();
+			Colors->ClearElements();
+			TArray<int32> ElemIDs;
+			ElemIDs.SetNum(EditMesh.MaxVertexID());
+			for (int32 VertexID : EditMesh.VertexIndicesItr())
+			{
+				//const FLinearColor& FromColor = VertexColors[VertexID];
+				//FColor SRGBFColor = FromColor.ToFColorSRGB();
+				//FLinearColor Color = SRGBFColor.ReinterpretAsLinear();
+				const FLinearColor& Color = VertexColors[VertexID];
+				ElemIDs[VertexID] = Colors->AppendElement(FVector4f(Color.R, Color.G, Color.B, Color.A));
+			}
+			for (int32 TriangleID : EditMesh.TriangleIndicesItr())
+			{
+				FIndex3i Triangle = EditMesh.GetTriangle(TriangleID);
+				Colors->SetTriangle(TriangleID, FIndex3i(ElemIDs[Triangle.A], ElemIDs[Triangle.B], ElemIDs[Triangle.C]) );
+			}
+		}
+
+	}, EDynamicMeshChangeType::GeneralEdit, EDynamicMeshAttributeChangeFlags::Unknown, false);
+
+	return TargetMesh;
+}
+
+#undef LOCTEXT_NAMESPACE
\ No newline at end of file
diff --git a/Engine/Plugins/Experimental/GeometryScripting/Source/GeometryScriptingCore/Private/MeshVoxelFunctions.cpp b/Engine/Plugins/Experimental/GeometryScripting/Source/GeometryScriptingCore/Private/MeshVoxelFunctions.cpp
new file mode 100644
--- /dev/null
+++ b/Engine/Plugins/Experimental/GeometryScripting/Source/GeometryScriptingCore/Private/MeshVoxelFunctions.cpp
@@ -0,0 +1,174 @@
+// Copyright Epic Games, Inc. All Rights Reserved.
+
+#include "GeometryScript/MeshVoxelFunctions.h"
+
+#include "DynamicMesh/DynamicMesh3.h"
+#include "DynamicMesh/MeshNormals.h"
+#include "DynamicMesh/DynamicMeshAABBTree3.h"
+#include "Spatial/FastWinding.h"
+#include "Implicit/Solidify.h"
+#include "Implicit/Morphology.h"
+#include "UDynamicMesh.h"
+
+#include "ExplicitUseGeometryMathTypes.h"		// using UE::Geometry::(math types)
+using namespace UE::Geometry;
+
+#define LOCTEXT_NAMESPACE "UGeometryScriptLibrary_MeshVoxelFunctions"
+
+
+static TAutoConsoleVariable<bool> CVarGeometryScriptUnclampedGridDimensions(
+	TEXT("geometry.GeometryScript.AllowUnclampedGridDimensions"),
+	false,
+	TEXT("Flag to disable grid dimension clamping in GeometryScript grid-based operations"));
+
+
+
+static int32 GetClampedGridResolution(int32 TargetResolution, FString ToolName, UGeometryScriptDebug* Debug)
+{
+	int ClampResolution = FMath::Max(4, TargetResolution);
+	if (ClampResolution > 256 && CVarGeometryScriptUnclampedGridDimensions.GetValueOnGameThread() == false)
+	{
+		UE::Geometry::AppendError(Debug, EGeometryScriptErrorType::InvalidInputs,
+			FText::Format(LOCTEXT("ClampedGridResolution", "{0}: Requested GridResolution clamped to 256 to avoid significant memory/computation. Use CVar geometry.GeometryScript.AllowUnclampedGridDimensions to disable this limit."), FText::FromString(ToolName)));
+		ClampResolution = 256;
+	}
+	return ClampResolution;
+}
+
+static double GetClampedCellSize(double TargetCellSize, FAxisAlignedBox3d BoundingBox, FString ToolName, UGeometryScriptDebug* Debug)
+{
+	double UseCellSize = FMath::Max(0.001, TargetCellSize);
+	int MaxGridDimEstimate = (int32)(BoundingBox.MaxDim() / UseCellSize);
+	if (MaxGridDimEstimate > 256 && CVarGeometryScriptUnclampedGridDimensions.GetValueOnGameThread() == false)
+	{
+		UE::Geometry::AppendError(Debug, EGeometryScriptErrorType::InvalidInputs, 
+			FText::Format(LOCTEXT("ClampedCellSize", "{0}: Requested GridCellSize will require significant memory/computation, clamping to 256 cells. Use CVar geometry.GeometryScript.AllowUnclampedGridDimensions to disable this limit."), FText::FromString(ToolName)));
+		UseCellSize = (float)BoundingBox.MaxDim() / 256.0;
+	}
+	return UseCellSize;
+}
+
+
+
+
+
+
+
+UDynamicMesh* UGeometryScriptLibrary_MeshVoxelFunctions::ApplyMeshSolidify(
+	UDynamicMesh* TargetMesh,
+	FGeometryScriptSolidifyOptions Options,
+	UGeometryScriptDebug* Debug)
+{
+	if (TargetMesh == nullptr)
+	{
+		UE::Geometry::AppendError(Debug, EGeometryScriptErrorType::InvalidInputs, LOCTEXT("ApplyMeshSolidify_InvalidInput", "ApplyMeshSolidify: TargetMesh is Null"));
+		return TargetMesh;
+	}
+
+	TargetMesh->EditMesh([&](FDynamicMesh3& EditMesh) 
+	{
+		FDynamicMeshAABBTree3 Spatial(&EditMesh);
+		TFastWindingTree<FDynamicMesh3> FastWinding(&Spatial);
+
+		TImplicitSolidify<FDynamicMesh3> Solidify(&EditMesh, &Spatial, &FastWinding);
+
+		if (Options.GridParameters.SizeMethod == EGeometryScriptGridSizingMethod::GridResolution)
+		{
+			int ClampResolution = GetClampedGridResolution(Options.GridParameters.GridResolution, TEXT("ApplyMeshSolidify"), Debug);
+			Solidify.SetCellSizeAndExtendBounds(Spatial.GetBoundingBox(), Options.ExtendBounds, ClampResolution);
+		}
+		else
+		{
+			double UseCellSize = GetClampedCellSize(Options.GridParameters.GridCellSize, Spatial.GetBoundingBox(), TEXT("ApplyMeshSolidify"), Debug);
+			Solidify.ExtendBounds = Options.ExtendBounds;
+			Solidify.MeshCellSize = UseCellSize;
+		}
+
+		Solidify.WindingThreshold = Options.WindingThreshold;
+		Solidify.SurfaceSearchSteps = Options.SurfaceSearchSteps;
+		Solidify.bSolidAtBoundaries = Options.bSolidAtBoundaries;
+		Solidify.ExtendBounds = Options.ExtendBounds;
+		EditMesh.Copy(&Solidify.Generate());
+
+		if (EditMesh.TriangleCount() == 0)
+		{
+			UE::Geometry::AppendError(Debug, EGeometryScriptErrorType::OperationFailed, LOCTEXT("ApplyMeshSolidify_EmptyMesh", "ApplyMeshSolidify: Solidify operation output an empty mesh"));
+		}
+
+	}, EDynamicMeshChangeType::GeneralEdit, EDynamicMeshAttributeChangeFlags::Unknown, false);
+
+	return TargetMesh;
+}
+
+
+
+
+
+
+
+UDynamicMesh* UGeometryScriptLibrary_MeshVoxelFunctions::ApplyMeshMorphology(
+	UDynamicMesh* TargetMesh,
+	FGeometryScriptMorphologyOptions Options,
+	UGeometryScriptDebug* Debug)
+{
+	if (TargetMesh == nullptr)
+	{
+		UE::Geometry::AppendError(Debug, EGeometryScriptErrorType::InvalidInputs, LOCTEXT("ApplyMeshMorphology_InvalidInput", "ApplyMeshMorphology: TargetMesh is Null"));
+		return TargetMesh;
+	}
+
+	TargetMesh->EditMesh([&](FDynamicMesh3& EditMesh) 
+	{
+		FDynamicMeshAABBTree3 Spatial(&EditMesh);
+		TImplicitMorphology<FDynamicMesh3> Morphology;
+		Morphology.Source = &EditMesh;
+		Morphology.SourceSpatial = &Spatial;
+		Morphology.MorphologyOp = static_cast<TImplicitMorphology<FDynamicMesh3>::EMorphologyOp>(Options.Operation);
+		Morphology.Distance = Options.Distance;
+
+		FGeometryScript3DGridParameters SDFGridParameters = Options.SDFGridParameters;
+		if (SDFGridParameters.SizeMethod == EGeometryScriptGridSizingMethod::GridResolution)
+		{
+			int ClampResolution = GetClampedGridResolution(SDFGridParameters.GridResolution, TEXT("ApplyMeshMorphology"), Debug);
+			Morphology.SetGridCellSize(Spatial.GetBoundingBox(), Options.Distance, ClampResolution);
+		}
+		else
+		{
+			double UseCellSize = GetClampedCellSize(SDFGridParameters.GridCellSize, Spatial.GetBoundingBox(), TEXT("ApplyMeshMorphology"), Debug);
+			Morphology.GridCellSize = UseCellSize;
+		}
+
+
+		FGeometryScript3DGridParameters MeshGridParameters = Options.bUseSeparateMeshGrid ? Options.MeshGridParameters : SDFGridParameters;
+		if (MeshGridParameters.SizeMethod == EGeometryScriptGridSizingMethod::GridResolution)
+		{
+			int ClampResolution = GetClampedGridResolution(MeshGridParameters.GridResolution, TEXT("ApplyMeshMorphology"), Debug);
+			Morphology.SetMeshCellSize(Spatial.GetBoundingBox(), Options.Distance, ClampResolution);
+		}
+		else
+		{
+			double UseCellSize = GetClampedCellSize(MeshGridParameters.GridCellSize, Spatial.GetBoundingBox(), TEXT("ApplyMeshMorphology"), Debug);
+			Morphology.MeshCellSize = UseCellSize;
+		}
+
+		EditMesh.Copy(&Morphology.Generate());
+
+		if (EditMesh.TriangleCount() == 0)
+		{
+			UE::Geometry::AppendError(Debug, EGeometryScriptErrorType::OperationFailed, LOCTEXT("ApplyMeshMorphology_EmptyMesh", "ApplyMeshMorphology: Morphology operation output an empty mesh"));
+		}
+
+	}, EDynamicMeshChangeType::GeneralEdit, EDynamicMeshAttributeChangeFlags::Unknown, false);
+
+	return TargetMesh;
+}
+
+
+
+
+
+
+
+#undef LOCTEXT_NAMESPACE
+
+
diff --git a/Engine/Plugins/Experimental/GeometryScripting/Source/GeometryScriptingCore/Private/SceneUtilityFunctions.cpp b/Engine/Plugins/Experimental/GeometryScripting/Source/GeometryScriptingCore/Private/SceneUtilityFunctions.cpp
new file mode 100644
--- /dev/null
+++ b/Engine/Plugins/Experimental/GeometryScripting/Source/GeometryScriptingCore/Private/SceneUtilityFunctions.cpp
@@ -0,0 +1,82 @@
+// Copyright Epic Games, Inc. All Rights Reserved.
+
+#include "GeometryScript/SceneUtilityFunctions.h"
+
+#include "DynamicMesh/DynamicMesh3.h"
+#include "UDynamicMesh.h"
+#include "GeometryScript/MeshAssetFunctions.h"
+#include "DynamicMesh/MeshTransforms.h"
+
+#include "Components/StaticMeshComponent.h"
+#include "Components/DynamicMeshComponent.h"
+
+
+#include "ExplicitUseGeometryMathTypes.h"		// using UE::Geometry::(math types)
+using namespace UE::Geometry;
+
+#define LOCTEXT_NAMESPACE "UGeometryScriptLibrary_SceneUtilityFunctions"
+
+
+UDynamicMesh* UGeometryScriptLibrary_SceneUtilityFunctions::CopyMeshFromComponent(
+	USceneComponent* Component,
+	UDynamicMesh* ToDynamicMesh,
+	FGeometryScriptCopyMeshFromComponentOptions Options,
+	bool bTransformToWorld,
+	FTransform& LocalToWorld,
+	TEnumAsByte<EGeometryScriptOutcomePins>& Outcome,
+	UGeometryScriptDebug* Debug)
+{
+	Outcome = EGeometryScriptOutcomePins::Failure;
+
+	if (Cast<UStaticMeshComponent>(Component) != nullptr)
+	{
+		UStaticMeshComponent* StaticMeshComponent = Cast<UStaticMeshComponent>(Component);
+		LocalToWorld = StaticMeshComponent->GetComponentTransform();
+		UStaticMesh* StaticMesh = StaticMeshComponent->GetStaticMesh();
+		if (StaticMesh)
+		{
+			FGeometryScriptCopyMeshFromAssetOptions AssetOptions;
+			AssetOptions.bApplyBuildSettings = (Options.bWantNormals || Options.bWantTangents);
+			AssetOptions.bRequestTangents = Options.bWantTangents;
+			UGeometryScriptLibrary_StaticMeshFunctions::CopyMeshFromStaticMesh(
+				StaticMesh, ToDynamicMesh, AssetOptions, Options.RequestedLOD, Outcome, Debug);	// will set Outcome pin
+		}
+		else
+		{
+			UE::Geometry::AppendError(Debug, EGeometryScriptErrorType::InvalidInputs, LOCTEXT("CopyMeshFromComponent_MissingStaticMesh", "CopyMeshFromComponent: StaticMeshComponent has a null StaticMesh"));
+		}
+	}
+	else if (Cast<UDynamicMeshComponent>(Component) != nullptr)
+	{
+		UDynamicMeshComponent* DynamicMeshComponent = Cast<UDynamicMeshComponent>(Component);
+		LocalToWorld = DynamicMeshComponent->GetComponentTransform();
+		UDynamicMesh* CopyDynamicMesh = DynamicMeshComponent->GetDynamicMesh();
+		if (CopyDynamicMesh)
+		{
+			CopyDynamicMesh->ProcessMesh([&](const FDynamicMesh3& Mesh)
+			{
+				ToDynamicMesh->SetMesh(Mesh);
+			});
+			Outcome = EGeometryScriptOutcomePins::Success;
+		}
+		else
+		{
+			UE::Geometry::AppendError(Debug, EGeometryScriptErrorType::InvalidInputs, LOCTEXT("CopyMeshFromComponent_MissingDynamicMesh", "CopyMeshFromComponent: DynamicMeshComponent has a null DynamicMesh"));
+		}
+	}
+
+	// transform mesh to world
+	if (Outcome == EGeometryScriptOutcomePins::Success && bTransformToWorld)
+	{
+		ToDynamicMesh->EditMesh([&](FDynamicMesh3& EditMesh) 
+		{
+			MeshTransforms::ApplyTransform(EditMesh, (UE::Geometry::FTransform3d)LocalToWorld);
+
+		}, EDynamicMeshChangeType::GeneralEdit, EDynamicMeshAttributeChangeFlags::Unknown, false);	
+	}
+
+	return ToDynamicMesh;
+}
+
+
+#undef LOCTEXT_NAMESPACE
\ No newline at end of file
diff --git a/Engine/Plugins/Experimental/GeometryScripting/Source/GeometryScriptingCore/Public/GeometryScript/GeometryScriptTypes.h b/Engine/Plugins/Experimental/GeometryScripting/Source/GeometryScriptingCore/Public/GeometryScript/GeometryScriptTypes.h
new file mode 100644
--- /dev/null
+++ b/Engine/Plugins/Experimental/GeometryScripting/Source/GeometryScriptingCore/Public/GeometryScript/GeometryScriptTypes.h
@@ -0,0 +1,159 @@
+// Copyright Epic Games, Inc. All Rights Reserved.
+
+#pragma once
+
+#include "CoreMinimal.h"
+#include "GeometryScriptTypes.generated.h"
+
+
+UENUM(BlueprintType)
+enum EGeometryScriptOutcomePins
+{
+	Failure,
+	Success
+};
+
+
+
+UENUM(BlueprintType)
+enum class EGeometryScriptLODType : uint8
+{
+	MaxAvailable,
+	HiResSourceModel,
+	SourceModel,
+	RenderData
+};
+
+
+USTRUCT(BlueprintType)
+struct GEOMETRYSCRIPTINGCORE_API FGeometryScriptMeshReadLOD
+{
+	GENERATED_BODY()
+public:
+	UPROPERTY(BlueprintReadWrite, Category = LOD)
+	EGeometryScriptLODType LODType = EGeometryScriptLODType::MaxAvailable;
+
+	UPROPERTY(BlueprintReadWrite, Category = LOD)
+	int32 LODIndex = 0;
+};
+
+
+USTRUCT(BlueprintType)
+struct GEOMETRYSCRIPTINGCORE_API FGeometryScriptMeshWriteLOD
+{
+	GENERATED_BODY()
+public:
+	UPROPERTY(BlueprintReadWrite, Category = LOD)
+	bool bWriteHiResSource = false;
+
+	UPROPERTY(BlueprintReadWrite, Category = LOD)
+	int32 LODIndex = 0;
+};
+
+
+
+
+
+//
+// Triangles
+//
+
+USTRUCT(BlueprintType)
+struct GEOMETRYSCRIPTINGCORE_API FGeometryScriptTriangle
+{
+	GENERATED_BODY()
+public:
+	UPROPERTY(BlueprintReadWrite, Category = Options)
+	FVector Vector0;
+
+	UPROPERTY(BlueprintReadWrite, Category = Options)
+	FVector Vector1;
+
+	UPROPERTY(BlueprintReadWrite, Category = Options)
+	FVector Vector2;
+};
+
+
+USTRUCT(BlueprintType)
+struct GEOMETRYSCRIPTINGCORE_API FGeometryScriptUVTriangle
+{
+	GENERATED_BODY()
+public:
+	UPROPERTY(BlueprintReadWrite, Category = Options)
+	FVector2D UV0;
+
+	UPROPERTY(BlueprintReadWrite, Category = Options)
+	FVector2D UV1;
+
+	UPROPERTY(BlueprintReadWrite, Category = Options)
+	FVector2D UV2;
+};
+
+
+//
+// Errors/Debugging
+//
+
+
+
+UENUM(BlueprintType)
+enum class EGeometryScriptDebugMessageType : uint8
+{
+	ErrorMessage,
+	WarningMessage
+};
+
+
+UENUM(BlueprintType)
+enum class EGeometryScriptErrorType : uint8
+{
+	// warning: must only append members!
+	NoError,
+	UnknownError,
+	InvalidInputs,
+	OperationFailed
+};
+
+
+USTRUCT(BlueprintType)
+struct GEOMETRYSCRIPTINGCORE_API FGeometryScriptDebugMessage
+{
+	GENERATED_BODY()
+public:
+	UPROPERTY(BlueprintReadWrite, Category = Options)
+	EGeometryScriptDebugMessageType MessageType = EGeometryScriptDebugMessageType::ErrorMessage;
+
+	UPROPERTY(BlueprintReadWrite, Category = Options)
+	EGeometryScriptErrorType ErrorType = EGeometryScriptErrorType::UnknownError;
+
+	UPROPERTY(BlueprintReadWrite, Category = Options)
+	FText Message = FText();
+};
+
+
+UCLASS(BlueprintType)
+class GEOMETRYSCRIPTINGCORE_API UGeometryScriptDebug : public UObject
+{
+	GENERATED_BODY()
+public:
+	UPROPERTY(BlueprintReadWrite, Category = Messages)
+	TArray<FGeometryScriptDebugMessage> Messages;
+
+
+	void Append(const FGeometryScriptDebugMessage& MessageIn)
+	{
+		Messages.Add(MessageIn);
+	}
+};
+
+
+
+namespace UE
+{
+namespace Geometry
+{
+	FGeometryScriptDebugMessage MakeScriptError(EGeometryScriptErrorType ErrorTypeIn, const FText& MessageIn);
+
+	void AppendError(UGeometryScriptDebug* Debug, EGeometryScriptErrorType ErrorTypeIn, const FText& MessageIn);
+}
+}
\ No newline at end of file
diff --git a/Engine/Plugins/Experimental/GeometryScripting/Source/GeometryScriptingCore/Public/GeometryScript/MeshAssetFunctions.h b/Engine/Plugins/Experimental/GeometryScripting/Source/GeometryScriptingCore/Public/GeometryScript/MeshAssetFunctions.h
new file mode 100644
--- /dev/null
+++ b/Engine/Plugins/Experimental/GeometryScripting/Source/GeometryScriptingCore/Public/GeometryScript/MeshAssetFunctions.h
@@ -0,0 +1,78 @@
+// Copyright Epic Games, Inc. All Rights Reserved.
+
+#pragma once
+
+#include "CoreMinimal.h"
+#include "Kismet/BlueprintFunctionLibrary.h"
+#include "GeometryScript/GeometryScriptTypes.h"
+#include "MeshAssetFunctions.generated.h"
+
+class UStaticMesh;
+class UDynamicMesh;
+
+
+
+
+USTRUCT(BlueprintType)
+struct GEOMETRYSCRIPTINGCORE_API FGeometryScriptCopyMeshFromAssetOptions
+{
+	GENERATED_BODY()
+public:
+	UPROPERTY(BlueprintReadWrite, Category = Options)
+	bool bApplyBuildSettings = true;
+
+	UPROPERTY(BlueprintReadWrite, Category = Options)
+	bool bRequestTangents = true;
+
+	UPROPERTY(BlueprintReadWrite, Category = Options)
+	bool bIgnoreRemoveDegenerates = true;
+};
+
+
+
+USTRUCT(BlueprintType)
+struct GEOMETRYSCRIPTINGCORE_API FGeometryScriptCopyMeshToAssetOptions
+{
+	GENERATED_BODY()
+public:
+	UPROPERTY(BlueprintReadWrite, Category = Options)
+	bool bEmitTransaction = true;
+
+	UPROPERTY(BlueprintReadWrite, Category = Options)
+	bool bDeferMeshPostEditChange = false;
+};
+
+
+
+
+
+UCLASS()
+class GEOMETRYSCRIPTINGCORE_API UGeometryScriptLibrary_StaticMeshFunctions : public UBlueprintFunctionLibrary
+{
+	GENERATED_BODY()
+public:
+
+	UFUNCTION(BlueprintCallable, Category = "GeometryScript|StaticMesh", meta = (ExpandEnumAsExecs = "Outcome"))
+	static UPARAM(DisplayName = "Dynamic Mesh") UDynamicMesh* 
+	CopyMeshFromStaticMesh(
+		UStaticMesh* FromStaticMeshAsset, 
+		UDynamicMesh* ToDynamicMesh, 
+		FGeometryScriptCopyMeshFromAssetOptions AssetOptions,
+		FGeometryScriptMeshReadLOD RequestedLOD,
+		TEnumAsByte<EGeometryScriptOutcomePins>& Outcome,
+		UGeometryScriptDebug* Debug = nullptr);
+
+
+	UFUNCTION(BlueprintCallable, Category = "GeometryScript|StaticMesh", meta = (ExpandEnumAsExecs = "Outcome"))
+	static UPARAM(DisplayName = "Dynamic Mesh") UDynamicMesh* 
+	CopyMeshToStaticMesh(
+		UDynamicMesh* FromDynamicMesh, 
+		UStaticMesh* ToStaticMeshAsset, 
+		FGeometryScriptCopyMeshToAssetOptions Options,
+		FGeometryScriptMeshWriteLOD TargetLOD,
+		TEnumAsByte<EGeometryScriptOutcomePins>& Outcome,
+		UGeometryScriptDebug* Debug = nullptr);
+
+};
+
+
diff --git a/Engine/Plugins/Experimental/GeometryScripting/Source/GeometryScriptingCore/Public/GeometryScript/MeshBasicEditFunctions.h b/Engine/Plugins/Experimental/GeometryScripting/Source/GeometryScriptingCore/Public/GeometryScript/MeshBasicEditFunctions.h
new file mode 100644
--- /dev/null
+++ b/Engine/Plugins/Experimental/GeometryScripting/Source/GeometryScriptingCore/Public/GeometryScript/MeshBasicEditFunctions.h
@@ -0,0 +1,117 @@
+// Copyright Epic Games, Inc. All Rights Reserved.
+
+#pragma once
+
+#include "CoreMinimal.h"
+#include "Kismet/BlueprintFunctionLibrary.h"
+#include "GeometryScript/GeometryScriptTypes.h"
+#include "MeshBasicEditFunctions.generated.h"
+
+class UDynamicMesh;
+
+
+
+UCLASS()
+class GEOMETRYSCRIPTINGCORE_API UGeometryScriptLibrary_MeshBasicEditFunctions : public UBlueprintFunctionLibrary
+{
+	GENERATED_BODY()
+public:
+
+	UFUNCTION(BlueprintCallable, Category = "GeometryScript|MeshEdits")
+	static UPARAM(DisplayName = "Target Mesh") UDynamicMesh* 
+	DiscardMeshAttributes( 
+		UDynamicMesh* TargetMesh, 
+		bool bDeferChangeNotifications = false );
+
+
+
+	UFUNCTION(BlueprintCallable, Category = "GeometryScript|MeshEdits")
+	static UPARAM(DisplayName = "Target Mesh") UDynamicMesh* 
+	SetVertexPosition( 
+		UDynamicMesh* TargetMesh, 
+		int VertexID, 
+		FVector NewPosition, 
+		bool& bIsValidVertex, 
+		bool bDeferChangeNotifications = false );
+
+
+	UFUNCTION(BlueprintCallable, Category = "GeometryScript|MeshEdits")
+	static UPARAM(DisplayName = "Target Mesh") UDynamicMesh* 
+	AddVertexToMesh( 
+		UDynamicMesh* TargetMesh, 
+		FVector NewPosition, 
+		int& NewVertexIndex,
+		bool bDeferChangeNotifications = false );
+
+	UFUNCTION(BlueprintCallable, Category = "GeometryScript|MeshEdits")
+	static UPARAM(DisplayName = "Target Mesh") UDynamicMesh* 
+	AddVerticesToMesh( 
+		UDynamicMesh* TargetMesh, 
+		const TArray<FVector>& NewPositions, 
+		TArray<int>& NewIndices,
+		bool bDeferChangeNotifications = false );
+
+	UFUNCTION(BlueprintCallable, Category = "GeometryScript|MeshEdits")
+	static UPARAM(DisplayName = "Target Mesh") UDynamicMesh* 
+	DeleteVertexFromMesh( 
+		UDynamicMesh* TargetMesh, 
+		int VertexID,
+		bool& bWasVertexDeleted,
+		bool bDeferChangeNotifications = false );
+
+
+
+	UFUNCTION(BlueprintCallable, Category = "GeometryScript|MeshEdits")
+	static UPARAM(DisplayName = "Target Mesh") UDynamicMesh* 
+	AddTriangleToMesh( 
+		UDynamicMesh* TargetMesh, 
+		FIntVector NewTriangle,
+		int& NewTriangleIndex,
+		int NewTriangleGroupID = 0,
+		bool bDeferChangeNotifications = false,
+		UGeometryScriptDebug* Debug = nullptr );
+
+
+	UFUNCTION(BlueprintCallable, Category = "GeometryScript|MeshEdits")
+	static UPARAM(DisplayName = "Target Mesh") UDynamicMesh* 
+	AddTrianglesToMesh( 
+		UDynamicMesh* TargetMesh, 
+		const TArray<FIntVector>& NewTriangles,
+		TArray<int>& NewIndices,
+		int NewTriangleGroupID = 0,
+		bool bDeferChangeNotifications = false,
+		UGeometryScriptDebug* Debug = nullptr );
+
+	UFUNCTION(BlueprintCallable, Category = "GeometryScript|MeshEdits")
+	static UPARAM(DisplayName = "Target Mesh") UDynamicMesh* 
+	DeleteTriangleFromMesh( 
+		UDynamicMesh* TargetMesh, 
+		int TriangleID,
+		bool& bWasTriangleDeleted,
+		bool bDeferChangeNotifications = false );
+
+
+
+
+	UFUNCTION(BlueprintCallable, Category = "GeometryScript|MeshEdits")
+	static UPARAM(DisplayName = "Target Mesh") UDynamicMesh* 
+	AppendMesh( 
+		UDynamicMesh* TargetMesh, 
+		UDynamicMesh* AppendMesh, 
+		FTransform AppendTransform, 
+		bool bDeferChangeNotifications = false,
+		UGeometryScriptDebug* Debug = nullptr);
+
+	UFUNCTION(BlueprintCallable, Category = "GeometryScript|MeshEdits")
+	static UPARAM(DisplayName = "Target Mesh") UDynamicMesh* 
+	AppendMeshRepeated( 
+		UDynamicMesh* TargetMesh, 
+		UDynamicMesh* AppendMesh, 
+		FTransform AppendTransform, 
+		int RepeatCount = 1,
+		bool bApplyTransformToFirstInstance = true,
+		bool bDeferChangeNotifications = false,
+		UGeometryScriptDebug* Debug = nullptr);
+
+
+};
\ No newline at end of file
diff --git a/Engine/Plugins/Experimental/GeometryScripting/Source/GeometryScriptingCore/Public/GeometryScript/MeshBooleanFunctions.h b/Engine/Plugins/Experimental/GeometryScripting/Source/GeometryScriptingCore/Public/GeometryScript/MeshBooleanFunctions.h
new file mode 100644
--- /dev/null
+++ b/Engine/Plugins/Experimental/GeometryScripting/Source/GeometryScriptingCore/Public/GeometryScript/MeshBooleanFunctions.h
@@ -0,0 +1,140 @@
+// Copyright Epic Games, Inc. All Rights Reserved.
+
+#pragma once
+
+#include "CoreMinimal.h"
+#include "Kismet/BlueprintFunctionLibrary.h"
+#include "GeometryScript/GeometryScriptTypes.h"
+#include "MeshBooleanFunctions.generated.h"
+
+
+
+UENUM(BlueprintType)
+enum class EGeometryScriptBooleanOperation : uint8
+{
+	Union,
+	Intersection,
+	Subtract
+};
+
+
+
+USTRUCT(BlueprintType)
+struct GEOMETRYSCRIPTINGCORE_API FGeometryScriptMeshBooleanOptions
+{
+	GENERATED_BODY()
+public:
+	UPROPERTY(BlueprintReadWrite, Category = Options)
+	bool bFillHoles = true;
+
+	UPROPERTY(BlueprintReadWrite, Category = Options)
+	bool bSimplifyOutput = true;
+
+	UPROPERTY(BlueprintReadWrite, Category = Options)
+	float SimplifyPlanarTolerance = 0.01f;
+};
+
+
+USTRUCT(BlueprintType)
+struct GEOMETRYSCRIPTINGCORE_API FGeometryScriptMeshSelfUnionOptions
+{
+	GENERATED_BODY()
+public:
+	UPROPERTY(BlueprintReadWrite, Category = Options)
+	bool bFillHoles = true;
+
+	UPROPERTY(BlueprintReadWrite, Category = Options)
+	bool bTrimFlaps = true;
+
+	UPROPERTY(BlueprintReadWrite, Category = Options)
+	bool bSimplifyOutput = true;
+
+	UPROPERTY(BlueprintReadWrite, Category = Options)
+	float SimplifyPlanarTolerance = 0.01f;
+
+	UPROPERTY(BlueprintReadWrite, Category = Options)
+	float WindingThreshold = 0.5f;
+};
+
+
+
+USTRUCT(BlueprintType)
+struct GEOMETRYSCRIPTINGCORE_API FGeometryScriptMeshPlaneCutOptions
+{
+	GENERATED_BODY()
+public:
+	UPROPERTY(BlueprintReadWrite, Category = Options)
+	bool bFillHoles = true;
+
+	UPROPERTY(BlueprintReadWrite, Category = Options)
+	bool bFillSpans = true;
+
+	UPROPERTY(BlueprintReadWrite, Category = Options)
+	bool bFlipCutSide = false;
+};
+
+
+USTRUCT(BlueprintType)
+struct GEOMETRYSCRIPTINGCORE_API FGeometryScriptMeshMirrorOptions
+{
+	GENERATED_BODY()
+public:
+	UPROPERTY(BlueprintReadWrite, Category = Options)
+	bool bApplyPlaneCut = true;
+
+	UPROPERTY(BlueprintReadWrite, Category = Options)
+	bool bFlipCutSide = false;
+
+	UPROPERTY(BlueprintReadWrite, Category = Options)
+	bool bWeldAlongPlane = true;
+};
+
+
+
+UCLASS()
+class GEOMETRYSCRIPTINGCORE_API UGeometryScriptLibrary_MeshBooleanFunctions : public UBlueprintFunctionLibrary
+{
+	GENERATED_BODY()
+public:
+
+	UFUNCTION(BlueprintCallable, Category = "GeometryScript|Booleans")
+	static UPARAM(DisplayName = "Target Mesh") UDynamicMesh* 
+	ApplyMeshBoolean( 
+		UDynamicMesh* TargetMesh, 
+		FTransform TargetTransform,
+		UDynamicMesh* ToolMesh, 
+		FTransform ToolTransform,
+		EGeometryScriptBooleanOperation Operation,
+		FGeometryScriptMeshBooleanOptions Options,
+		UGeometryScriptDebug* Debug = nullptr);
+
+
+	UFUNCTION(BlueprintCallable, Category = "GeometryScript|Booleans")
+	static UPARAM(DisplayName = "Target Mesh") UDynamicMesh* 
+	ApplyMeshSelfUnion( 
+		UDynamicMesh* TargetMesh, 
+		FGeometryScriptMeshSelfUnionOptions Options,
+		UGeometryScriptDebug* Debug = nullptr);
+
+
+
+	UFUNCTION(BlueprintCallable, Category = "GeometryScript|Booleans")
+	static UPARAM(DisplayName = "Target Mesh") UDynamicMesh* 
+	ApplyMeshPlaneCut( 
+		UDynamicMesh* TargetMesh, 
+		FVector CutPlaneOrigin,
+		FVector CutPlaneNormal,
+		FGeometryScriptMeshPlaneCutOptions Options,
+		UGeometryScriptDebug* Debug = nullptr);
+
+
+	UFUNCTION(BlueprintCallable, Category = "GeometryScript|Booleans")
+	static UPARAM(DisplayName = "Target Mesh") UDynamicMesh* 
+	ApplyMeshMirror( 
+		UDynamicMesh* TargetMesh, 
+		FVector MirrorPlaneOrigin,
+		FVector MirrorPlaneNormal,
+		FGeometryScriptMeshMirrorOptions Options,
+		UGeometryScriptDebug* Debug = nullptr);
+
+};
diff --git a/Engine/Plugins/Experimental/GeometryScripting/Source/GeometryScriptingCore/Public/GeometryScript/MeshDeformFunctions.h b/Engine/Plugins/Experimental/GeometryScripting/Source/GeometryScriptingCore/Public/GeometryScript/MeshDeformFunctions.h
new file mode 100644
--- /dev/null
+++ b/Engine/Plugins/Experimental/GeometryScripting/Source/GeometryScriptingCore/Public/GeometryScript/MeshDeformFunctions.h
@@ -0,0 +1,111 @@
+// Copyright Epic Games, Inc. All Rights Reserved.
+
+#pragma once
+
+#include "CoreMinimal.h"
+#include "Kismet/BlueprintFunctionLibrary.h"
+#include "GeometryScript/GeometryScriptTypes.h"
+#include "MeshDeformFunctions.generated.h"
+
+class UDynamicMesh;
+
+
+USTRUCT(BlueprintType)
+struct GEOMETRYSCRIPTINGCORE_API FGeometryScriptBendWarpOptions
+{
+	GENERATED_BODY()
+public:
+	/** Origin of the Bend Warp */
+	UPROPERTY(BlueprintReadWrite, Category = Options)
+	FVector BendOrigin = FVector(0,0,0);
+
+	/** Bend will be along this axis, ie if bending a long tube, this axis would be aligned with the tube axis */
+	UPROPERTY(BlueprintReadWrite, Category = Options)
+	FVector BendAxis = FVector(0, 0, 1);
+
+	/** Bend will bend towards this direction, eg if bending a vertical (+Z) bar with BendAxis = (0,0,1), A positive bend angle would bend the bar towards this direction (in the XY plane). Automatically calculated if this vector is Zero. */
+	UPROPERTY(BlueprintReadWrite, Category = Options)
+	FVector TowardAxis = FVector(0, 0, 0);
+
+	/** Rotation in Degrees around the Bend Axis, this is applied after the TowardAxis rotation */
+	UPROPERTY(BlueprintReadWrite, Category = Options)
+	float BendAxisRotation = 0;
+
+	/** Symmetric extents are [-BendExtent,BendExtent], if disabled, then [-LowerExtent,BendExtent] is used  */
+	UPROPERTY(BlueprintReadWrite, Category = Options)
+	bool bSymmetricExtents = true;
+
+	/** Lower extent used when bSymmetricExtents = false */
+	UPROPERTY(BlueprintReadWrite, Category = Options, meta = (EditCondition = "bSymmetricExtents == false"))
+	float LowerExtent = 10;
+
+	/** If true, the Bend is "centered" at the Origin, ie the regions on either side of the extents are rigidly transformed. If false, the Bend begins at the start of the Lower Extents, and the "lower" region is not affected. */
+	UPROPERTY(BlueprintReadWrite, Category = Options)
+	bool bBidirectional = true;
+};
+
+
+
+USTRUCT(BlueprintType)
+struct GEOMETRYSCRIPTINGCORE_API FGeometryScriptTwistWarpOptions
+{
+	GENERATED_BODY()
+public:
+	/** Origin of the Bend Warp */
+	UPROPERTY(BlueprintReadWrite, Category = Options)
+	FVector TwistOrigin = FVector(0,0,0);
+
+	/** Twist will be along this axis, ie if twisting a long tube, this the tube would twist "around" this axis */
+	UPROPERTY(BlueprintReadWrite, Category = Options)
+	FVector TwistAxis = FVector(0, 0, 1);
+
+	/** Twist will begin pointed towards this direction (perpendicular to TwistAxis). Automatically calculated if this vector is Zero. */
+	UPROPERTY(BlueprintReadWrite, Category = Options)
+	FVector TowardAxis = FVector(0, 0, 0);
+
+	/** Rotation in Degrees around the Twist Axis, this is applied after the TowardAxis rotation */
+	UPROPERTY(BlueprintReadWrite, Category = Options)
+	float TwistAxisRotation = 0;
+
+	/** Symmetric extents are [-BendExtent,BendExtent], if disabled, then [-LowerExtent,BendExtent] is used  */
+	UPROPERTY(BlueprintReadWrite, Category = Options)
+	bool bSymmetricExtents = true;
+
+	/** Lower extent used when bSymmetricExtents = false */
+	UPROPERTY(BlueprintReadWrite, Category = Options, meta = (EditCondition = "bSymmetricExtents == false"))
+	float LowerExtent = 10;
+
+	/** If true, the Twist is "centered" at the Origin, ie the regions on either side of the extents are rigidly transformed. If false, the Twist begins at the start of the Lower Extents, and the "lower" region is not affected. */
+	UPROPERTY(BlueprintReadWrite, Category = Options)
+	bool bBidirectional = true;
+};
+
+
+
+
+UCLASS()
+class GEOMETRYSCRIPTINGCORE_API UGeometryScriptLibrary_MeshDeformFunctions : public UBlueprintFunctionLibrary
+{
+	GENERATED_BODY()
+public:
+
+	UFUNCTION(BlueprintCallable, Category = "GeometryScript|Deformations")
+	static UPARAM(DisplayName = "Target Mesh") UDynamicMesh* 
+	ApplyBendWarpToMesh(  
+		UDynamicMesh* TargetMesh, 
+		FGeometryScriptBendWarpOptions Options,
+		float BendAngle = 45,
+		float BendExtent = 50,
+		UGeometryScriptDebug* Debug = nullptr);
+
+	UFUNCTION(BlueprintCallable, Category = "GeometryScript|Deformations")
+	static UPARAM(DisplayName = "Target Mesh") UDynamicMesh* 
+	ApplyTwistWarpToMesh(  
+		UDynamicMesh* TargetMesh, 
+		FGeometryScriptTwistWarpOptions Options,
+		float TwistAngle = 45,
+		float TwistExtent = 50,
+		UGeometryScriptDebug* Debug = nullptr);
+
+
+};
\ No newline at end of file
diff --git a/Engine/Plugins/Experimental/GeometryScripting/Source/GeometryScriptingCore/Public/GeometryScript/MeshNormalsFunctions.h b/Engine/Plugins/Experimental/GeometryScripting/Source/GeometryScriptingCore/Public/GeometryScript/MeshNormalsFunctions.h
new file mode 100644
--- /dev/null
+++ b/Engine/Plugins/Experimental/GeometryScripting/Source/GeometryScriptingCore/Public/GeometryScript/MeshNormalsFunctions.h
@@ -0,0 +1,102 @@
+// Copyright Epic Games, Inc. All Rights Reserved.
+
+#pragma once
+
+#include "CoreMinimal.h"
+#include "Kismet/BlueprintFunctionLibrary.h"
+#include "GeometryScript/GeometryScriptTypes.h"
+#include "MeshNormalsFunctions.generated.h"
+
+class UDynamicMesh;
+
+USTRUCT(BlueprintType)
+struct GEOMETRYSCRIPTINGCORE_API FGeometryScriptCalculateNormalsOptions
+{
+	GENERATED_BODY()
+public:
+	UPROPERTY(BlueprintReadWrite, Category = Options)
+	bool bAngleWeighted = true;
+
+	UPROPERTY(BlueprintReadWrite, Category = Options)
+	bool bAreaWeighted = true;
+};
+
+
+
+USTRUCT(BlueprintType)
+struct GEOMETRYSCRIPTINGCORE_API FGeometryScriptSplitNormalsOptions
+{
+	GENERATED_BODY()
+public:
+	UPROPERTY(BlueprintReadWrite, Category = Options)
+	bool bSplitByOpeningAngle = true;
+
+	UPROPERTY(BlueprintReadWrite, Category = Options)
+	float OpeningAngleDeg = 15.0f;
+
+	UPROPERTY(BlueprintReadWrite, Category = Options)
+	bool bSplitByFaceGroup = false;
+
+	UPROPERTY(BlueprintReadWrite, Category = Options)
+	FName FaceGroupLayerName = FName();
+};
+
+
+UCLASS()
+class GEOMETRYSCRIPTINGCORE_API UGeometryScriptLibrary_MeshNormalsFunctions : public UBlueprintFunctionLibrary
+{
+	GENERATED_BODY()
+public:
+
+	UFUNCTION(BlueprintCallable, Category = "GeometryScript|Normals")
+	static UPARAM(DisplayName = "Target Mesh") UDynamicMesh* 
+	FlipNormals( 
+		UDynamicMesh* TargetMesh, 
+		UGeometryScriptDebug* Debug = nullptr);
+
+	UFUNCTION(BlueprintCallable, Category = "GeometryScript|Normals")
+	static UPARAM(DisplayName = "Target Mesh") UDynamicMesh* 
+	AutoRepairNormals( 
+		UDynamicMesh* TargetMesh, 
+		UGeometryScriptDebug* Debug = nullptr);
+
+	UFUNCTION(BlueprintCallable, Category = "GeometryScript|Normals")
+	static UPARAM(DisplayName = "Target Mesh") UDynamicMesh* 
+	SetPerVertexNormals( 
+		UDynamicMesh* TargetMesh, 
+		UGeometryScriptDebug* Debug = nullptr);
+
+	UFUNCTION(BlueprintCallable, Category = "GeometryScript|Normals")
+	static UPARAM(DisplayName = "Target Mesh") UDynamicMesh* 
+	SetPerFaceNormals( 
+		UDynamicMesh* TargetMesh, 
+		UGeometryScriptDebug* Debug = nullptr);
+
+	UFUNCTION(BlueprintCallable, Category = "GeometryScript|Normals")
+	static UPARAM(DisplayName = "Target Mesh") UDynamicMesh* 
+	RecomputeNormals(  
+		UDynamicMesh* TargetMesh, 
+		FGeometryScriptCalculateNormalsOptions CalculateOptions,
+		UGeometryScriptDebug* Debug = nullptr);
+
+	UFUNCTION(BlueprintCallable, Category = "GeometryScript|Normals")
+	static UPARAM(DisplayName = "Target Mesh") UDynamicMesh* 
+	ComputeSplitNormals( 
+		UDynamicMesh* TargetMesh, 
+		FGeometryScriptSplitNormalsOptions SplitOptions,
+		FGeometryScriptCalculateNormalsOptions CalculateOptions,
+		UGeometryScriptDebug* Debug = nullptr);
+
+
+
+
+	UFUNCTION(BlueprintCallable, Category = "GeometryScript|Normals")
+	static UPARAM(DisplayName = "Target Mesh") UDynamicMesh* 
+	SetMeshTriangleNormals( 
+		UDynamicMesh* TargetMesh, 
+		int TriangleID, 
+		FGeometryScriptTriangle Normals,
+		bool& bIsValidTriangle, 
+		bool bDeferChangeNotifications = false );
+
+};
\ No newline at end of file
diff --git a/Engine/Plugins/Experimental/GeometryScripting/Source/GeometryScriptingCore/Public/GeometryScript/MeshPrimitiveFunctions.h b/Engine/Plugins/Experimental/GeometryScripting/Source/GeometryScriptingCore/Public/GeometryScript/MeshPrimitiveFunctions.h
new file mode 100644
--- /dev/null
+++ b/Engine/Plugins/Experimental/GeometryScripting/Source/GeometryScriptingCore/Public/GeometryScript/MeshPrimitiveFunctions.h
@@ -0,0 +1,203 @@
+// Copyright Epic Games, Inc. All Rights Reserved.
+
+#pragma once
+
+#include "CoreMinimal.h"
+#include "Kismet/BlueprintFunctionLibrary.h"
+#include "GeometryScript/GeometryScriptTypes.h"
+#include "MeshPrimitiveFunctions.generated.h"
+
+
+UENUM(BlueprintType)
+enum class EGeometryScriptPrimitivePolygroupMode : uint8
+{
+	SingleGroup = 0,
+	PerFace = 1,
+	PerQuad = 2
+};
+
+
+
+USTRUCT(BlueprintType)
+struct GEOMETRYSCRIPTINGCORE_API FGeometryScriptPrimitiveOptions
+{
+	GENERATED_BODY()
+public:
+	UPROPERTY(BlueprintReadWrite, Category = Options)
+	EGeometryScriptPrimitivePolygroupMode PolygroupMode = EGeometryScriptPrimitivePolygroupMode::PerFace;
+
+	UPROPERTY(BlueprintReadWrite, Category = Options)
+	bool bFlipOrientation = false;
+};
+
+
+UCLASS()
+class GEOMETRYSCRIPTINGCORE_API UGeometryScriptLibrary_MeshPrimitiveFunctions : public UBlueprintFunctionLibrary
+{
+	GENERATED_BODY()
+public:
+
+
+	UFUNCTION(BlueprintCallable, Category = "GeometryScript|Primitives")
+	static UPARAM(DisplayName = "Target Mesh") UDynamicMesh* 
+	AppendBox( 
+		UDynamicMesh* TargetMesh, 
+		FGeometryScriptPrimitiveOptions PrimitiveOptions,
+		FTransform Transform,
+		FBox Box,
+		int32 StepsX = 0,
+		int32 StepsY = 0,
+		int32 StepsZ = 0,
+		UGeometryScriptDebug* Debug = nullptr);
+
+
+	UFUNCTION(BlueprintCallable, Category = "GeometryScript|Primitives")
+	static UPARAM(DisplayName = "Target Mesh") UDynamicMesh* 
+	AppendSphereLatLong( 
+		UDynamicMesh* TargetMesh, 
+		FGeometryScriptPrimitiveOptions PrimitiveOptions,
+		FTransform Transform,
+		float Radius = 50,
+		int32 StepsPhi = 8,
+		int32 StepsTheta = 8,
+		UGeometryScriptDebug* Debug = nullptr);
+
+
+	UFUNCTION(BlueprintCallable, Category = "GeometryScript|Primitives")
+	static UPARAM(DisplayName = "Target Mesh") UDynamicMesh* 
+	AppendSphereBox( 
+		UDynamicMesh* TargetMesh, 
+		FGeometryScriptPrimitiveOptions PrimitiveOptions,
+		FTransform Transform,
+		float Radius = 50,
+		int32 StepsX = 0,
+		int32 StepsY = 0,
+		int32 StepsZ = 0,
+		UGeometryScriptDebug* Debug = nullptr);
+
+
+	UFUNCTION(BlueprintCallable, Category = "GeometryScript|Primitives")
+	static UPARAM(DisplayName = "Target Mesh") UDynamicMesh* 
+	AppendCapsule( 
+		UDynamicMesh* TargetMesh, 
+		FGeometryScriptPrimitiveOptions PrimitiveOptions,
+		FTransform Transform,
+		float Radius = 10,
+		float LineLength = 50,
+		int32 HemisphereSteps = 4,
+		int32 CircleSteps = 8,
+		UGeometryScriptDebug* Debug = nullptr);
+
+
+	UFUNCTION(BlueprintCallable, Category = "GeometryScript|Primitives")
+	static UPARAM(DisplayName = "Target Mesh") UDynamicMesh* 
+	AppendCylinder( 
+		UDynamicMesh* TargetMesh, 
+		FGeometryScriptPrimitiveOptions PrimitiveOptions,
+		FTransform Transform,
+		float Radius = 25,
+		float Height = 50,
+		int32 RadialSteps = 8,
+		int32 HeightSteps = 0,
+		bool bCapped = true,
+		UGeometryScriptDebug* Debug = nullptr);
+
+
+	UFUNCTION(BlueprintCallable, Category = "GeometryScript|Primitives")
+	static UPARAM(DisplayName = "Target Mesh") UDynamicMesh* 
+	AppendCone( 
+		UDynamicMesh* TargetMesh, 
+		FGeometryScriptPrimitiveOptions PrimitiveOptions,
+		FTransform Transform,
+		float BaseRadius = 25,
+		float TopRadius = 5,
+		float Height = 50,
+		int32 RadialSteps = 8,
+		int32 HeightSteps = 0,
+		bool bCapped = true,
+		UGeometryScriptDebug* Debug = nullptr);
+
+
+	UFUNCTION(BlueprintCallable, Category = "GeometryScript|Primitives")
+	static UPARAM(DisplayName = "Target Mesh") UDynamicMesh* 
+	AppendTorus( 
+		UDynamicMesh* TargetMesh, 
+		FGeometryScriptPrimitiveOptions PrimitiveOptions,
+		FTransform Transform,
+		float MajorRadius = 100,
+		float MinorRadius = 20,
+		int32 MajorSteps = 8,
+		int32 MinorSteps = 4,
+		UGeometryScriptDebug* Debug = nullptr);
+
+	/**
+	 * In the coordinate system of the revolve polygon, +X is towards the "outside" of the revolve donut, and +Y is "up" (ie +Z in local space)
+	 * Polygon should be oriented counter-clockwise to produce a correctly-oriented shape, otherwise it will be inside-out
+	 * Polygon endpoint is not repeated.
+	 */
+	UFUNCTION(BlueprintCallable, Category = "GeometryScript|Primitives")
+	static UPARAM(DisplayName = "Target Mesh") UDynamicMesh* 
+	AppendSimpleRevolvePolygon( 
+		UDynamicMesh* TargetMesh, 
+		FGeometryScriptPrimitiveOptions PrimitiveOptions,
+		FTransform Transform,
+		const TArray<FVector2D>& PolygonVertices,
+		float Radius = 100,
+		int32 Steps = 8,
+		UGeometryScriptDebug* Debug = nullptr);
+
+
+	UFUNCTION(BlueprintCallable, Category = "GeometryScript|Primitives")
+	static UPARAM(DisplayName = "Target Mesh") UDynamicMesh* 
+	AppendSimpleSweptPolygon( 
+		UDynamicMesh* TargetMesh, 
+		FGeometryScriptPrimitiveOptions PrimitiveOptions,
+		FTransform Transform,
+		const TArray<FVector2D>& PolygonVertices,
+		const TArray<FVector>& SweepPath,
+		bool bLoop = false,
+		bool bCapped = true,
+		float StartScale = 1.0f,
+		float EndScale = 1.0f,
+		UGeometryScriptDebug* Debug = nullptr);
+
+
+
+	UFUNCTION(BlueprintCallable, Category = "GeometryScript|Primitives")
+	static UPARAM(DisplayName = "Target Mesh") UDynamicMesh* 
+	AppendRectangle( 
+		UDynamicMesh* TargetMesh, 
+		FGeometryScriptPrimitiveOptions PrimitiveOptions,
+		FTransform Transform,
+		FBox2D Box,
+		int32 StepsWidth = 0,
+		int32 StepsHeight = 0,
+		UGeometryScriptDebug* Debug = nullptr);
+
+	UFUNCTION(BlueprintCallable, Category = "GeometryScript|Primitives")
+	static UPARAM(DisplayName = "Target Mesh") UDynamicMesh* 
+	AppendRoundRectangle( 
+		UDynamicMesh* TargetMesh, 
+		FGeometryScriptPrimitiveOptions PrimitiveOptions,
+		FTransform Transform,
+		FBox2D Box,
+		float CornerRadius = 5,
+		int32 StepsWidth = 0,
+		int32 StepsHeight = 0,
+		int32 StepsRound = 4,
+		UGeometryScriptDebug* Debug = nullptr);
+
+	UFUNCTION(BlueprintCallable, Category = "GeometryScript|Primitives")
+	static UPARAM(DisplayName = "Target Mesh") UDynamicMesh* 
+	AppendDisc( 
+		UDynamicMesh* TargetMesh, 
+		FGeometryScriptPrimitiveOptions PrimitiveOptions,
+		FTransform Transform,
+		float Radius = 50,
+		int32 AngleSteps = 16,
+		int32 SpokeSteps = 0,
+		float StartAngle = 0,
+		float EndAngle = 360,
+		UGeometryScriptDebug* Debug = nullptr);
+
+};
\ No newline at end of file
diff --git a/Engine/Plugins/Experimental/GeometryScripting/Source/GeometryScriptingCore/Public/GeometryScript/MeshQueryFunctions.h b/Engine/Plugins/Experimental/GeometryScripting/Source/GeometryScriptingCore/Public/GeometryScript/MeshQueryFunctions.h
new file mode 100644
--- /dev/null
+++ b/Engine/Plugins/Experimental/GeometryScripting/Source/GeometryScriptingCore/Public/GeometryScript/MeshQueryFunctions.h
@@ -0,0 +1,96 @@
+// Copyright Epic Games, Inc. All Rights Reserved.
+
+#pragma once
+
+#include "CoreMinimal.h"
+#include "Kismet/BlueprintFunctionLibrary.h"
+#include "GeometryScript/GeometryScriptTypes.h"
+#include "MeshQueryFunctions.generated.h"
+
+class UDynamicMesh;
+
+
+UCLASS()
+class GEOMETRYSCRIPTINGCORE_API UGeometryScriptLibrary_MeshQueryFunctions : public UBlueprintFunctionLibrary
+{
+	GENERATED_BODY()
+public:
+
+	UFUNCTION(BlueprintPure, Category = "GeometryScript|MeshQueries")
+	static bool GetIsDenseMesh( UDynamicMesh* TargetMesh );
+
+	UFUNCTION(BlueprintPure, Category = "GeometryScript|MeshQueries")
+	static bool GetMeshHasAttributeSet( UDynamicMesh* TargetMesh );
+
+	UFUNCTION(BlueprintPure, Category = "GeometryScript|MeshQueries")
+	static UPARAM(DisplayName = "Bounding Box") FBox GetMeshBoundingBox( UDynamicMesh* TargetMesh );
+
+	UFUNCTION(BlueprintPure, Category = "GeometryScript|MeshQueries")
+	static void GetMeshVolumeArea( UDynamicMesh* TargetMesh, float& SurfaceArea, float& Volume );
+
+	UFUNCTION(BlueprintPure, Category = "GeometryScript|MeshQueries")
+	static bool GetIsClosedMesh( UDynamicMesh* TargetMesh );
+
+	UFUNCTION(BlueprintPure, Category = "GeometryScript|MeshQueries")
+	static UPARAM(DisplayName = "Num Loops") int32 GetNumOpenBorderLoops( UDynamicMesh* TargetMesh, bool& bAmbiguousTopologyFound );
+
+	UFUNCTION(BlueprintPure, Category = "GeometryScript|MeshQueries")
+	static UPARAM(DisplayName = "Num Edges") int32 GetNumOpenBorderEdges( UDynamicMesh* TargetMesh );
+
+
+
+	UFUNCTION(BlueprintPure, Category = "GeometryScript|MeshQueries")
+	static UPARAM(DisplayName = "Triangle Count") int32 GetTriangleCount( UDynamicMesh* TargetMesh );
+
+	UFUNCTION(BlueprintPure, Category = "GeometryScript|MeshQueries")
+	static UPARAM(DisplayName = "Num Triangle IDs") int32 GetNumTriangleIDs( UDynamicMesh* TargetMesh );
+
+	UFUNCTION(BlueprintPure, Category = "GeometryScript|MeshQueries")
+	static bool GetHasTriangleIDGaps( UDynamicMesh* TargetMesh );
+
+	UFUNCTION(BlueprintPure, Category = "GeometryScript|MeshQueries")
+	static bool IsValidTriangleID( UDynamicMesh* TargetMesh, int32 TriangleID );
+
+	UFUNCTION(BlueprintPure, Category = "GeometryScript|MeshQueries")
+	static UPARAM(DisplayName = "Indices") FIntVector GetTriangleIndices( UDynamicMesh* TargetMesh, int32 TriangleID, bool& bIsValidTriangle );
+
+	UFUNCTION(BlueprintPure, Category = "GeometryScript|MeshQueries")
+	static void GetTrianglePositions( UDynamicMesh* TargetMesh, int32 TriangleID, bool& bIsValidTriangle, FVector& Vertex1, FVector& Vertex2, FVector& Vertex3 );
+
+	UFUNCTION(BlueprintPure, Category = "GeometryScript|MeshQueries")
+	static UPARAM(DisplayName = "Normal") FVector GetTriangleFaceNormal( UDynamicMesh* TargetMesh, int32 TriangleID, bool& bIsValidTriangle );
+
+
+
+
+	UFUNCTION(BlueprintPure, Category = "GeometryScript|MeshQueries")
+	static UPARAM(DisplayName = "Vertex Count") int32 GetVertexCount( UDynamicMesh* TargetMesh );
+
+	UFUNCTION(BlueprintPure, Category = "GeometryScript|MeshQueries")
+	static UPARAM(DisplayName = "Num Vertex IDs") int32 GetNumVertexIDs( UDynamicMesh* TargetMesh );
+
+	UFUNCTION(BlueprintPure, Category = "GeometryScript|MeshQueries")
+	static bool GetHasVertexIDGaps( UDynamicMesh* TargetMesh );
+
+	UFUNCTION(BlueprintPure, Category = "GeometryScript|MeshQueries")
+	static bool IsValidVertexID( UDynamicMesh* TargetMesh, int32 VertexID );
+
+	UFUNCTION(BlueprintPure, Category = "GeometryScript|MeshQueries")
+	static UPARAM(DisplayName = "Position") FVector GetVertexPosition( UDynamicMesh* TargetMesh, int32 VertexID, bool& bIsValidVertex );
+
+	UFUNCTION(BlueprintPure, Category = "GeometryScript|MeshQueries")
+	static UPARAM(DisplayName = "Positions") TArray<FVector> GetAllVertexPositions( UDynamicMesh* TargetMesh, bool& bHasVertexIDGaps );
+
+
+
+	UFUNCTION(BlueprintPure, Category = "GeometryScript|MeshQueries")
+	static UPARAM(DisplayName = "Num UV Sets") int32 GetNumUVSets( UDynamicMesh* TargetMesh );
+
+	UFUNCTION(BlueprintPure, Category = "GeometryScript|MeshQueries")
+	static UPARAM(DisplayName = "Bounding Box") FBox2D GetUVSetBoundingBox( UDynamicMesh* TargetMesh, int UVSetIndex, bool& bIsValidUVSet, bool& bUVSetIsEmpty );
+
+	UFUNCTION(BlueprintPure, Category = "GeometryScript|MeshQueries")
+	static void GetTriangleUVs( UDynamicMesh* TargetMesh, int32 UVSetIndex, int32 TriangleID, FVector2D& UV1, FVector2D& UV2, FVector2D& UV3, bool& bHaveValidUVs );
+
+
+};
\ No newline at end of file
diff --git a/Engine/Plugins/Experimental/GeometryScripting/Source/GeometryScriptingCore/Public/GeometryScript/MeshRepairFunctions.h b/Engine/Plugins/Experimental/GeometryScripting/Source/GeometryScriptingCore/Public/GeometryScript/MeshRepairFunctions.h
new file mode 100644
--- /dev/null
+++ b/Engine/Plugins/Experimental/GeometryScripting/Source/GeometryScriptingCore/Public/GeometryScript/MeshRepairFunctions.h
@@ -0,0 +1,159 @@
+// Copyright Epic Games, Inc. All Rights Reserved.
+
+#pragma once
+
+#include "CoreMinimal.h"
+#include "Kismet/BlueprintFunctionLibrary.h"
+#include "GeometryScript/GeometryScriptTypes.h"
+#include "MeshRepairFunctions.generated.h"
+
+class UDynamicMesh;
+
+USTRUCT(BlueprintType)
+struct GEOMETRYSCRIPTINGCORE_API FGeometryScriptWeldEdgesOptions
+{
+	GENERATED_BODY()
+public:
+	/** Edges are coincident if both pairs of endpoint vertices, and their midpoint, are closer than this distance */
+	UPROPERTY(BlueprintReadWrite, Category = Options)
+	float Tolerance = 1e-06f;
+
+	/** Only merge unambiguous pairs that have unique duplicate-edge matches */
+	UPROPERTY(BlueprintReadWrite, Category = Options)
+	bool bOnlyUniquePairs = true;
+};
+
+
+UENUM(BlueprintType)
+enum class EGeometryScriptFillHolesMethod : uint8
+{
+	Automatic = 0,
+	MinimalFill = 1,
+	PolygonTriangulation = 2,
+	TriangleFan = 3,
+	PlanarProjection = 4
+};
+
+USTRUCT(BlueprintType)
+struct GEOMETRYSCRIPTINGCORE_API FGeometryScriptFillHolesOptions
+{
+	GENERATED_BODY()
+public:
+	UPROPERTY(BlueprintReadWrite, Category = Options)
+	EGeometryScriptFillHolesMethod FillMethod = EGeometryScriptFillHolesMethod::Automatic;
+
+	/** Delete floating, disconnected triangles, as they produce a "hole" that cannot be filled */
+	bool bDeleteIsolatedTriangles = true;
+};
+
+
+
+USTRUCT(BlueprintType)
+struct GEOMETRYSCRIPTINGCORE_API FGeometryScriptRemoveSmallComponentOptions
+{
+	GENERATED_BODY()
+public:
+	/**  */
+	UPROPERTY(BlueprintReadWrite, Category = Options)
+	float MinVolume = 0.0001;
+
+	/**  */
+	UPROPERTY(BlueprintReadWrite, Category = Options)
+	float MinArea = 0.0001;
+
+	UPROPERTY(BlueprintReadWrite, Category = Options)
+	int MinTriangleCount = 1;
+};
+
+
+
+
+UENUM(BlueprintType)
+enum class EGeometryScriptRemoveHiddenTrianglesMethod : uint8
+{
+	FastWindingNumber = 0,
+	RaycastOcclusionTest = 1
+};
+
+
+USTRUCT(BlueprintType)
+struct GEOMETRYSCRIPTINGCORE_API FGeometryScriptRemoveHiddenTrianglesOptions
+{
+	GENERATED_BODY()
+public:
+
+	UPROPERTY(BlueprintReadWrite, Category = Options)
+	EGeometryScriptRemoveHiddenTrianglesMethod Method = EGeometryScriptRemoveHiddenTrianglesMethod::FastWindingNumber;
+
+	// add triangle samples per triangle (in addition to TriangleSamplingMethod)
+	UPROPERTY(BlueprintReadWrite, Category = Options)
+	int SamplesPerTriangle = 0;
+
+	// once triangles to remove are identified, do iterations of boundary erosion, ie contract selection by boundary vertex one-rings
+	UPROPERTY(BlueprintReadWrite, Category = Options)
+	int ShrinkSelection = 0;
+
+	// use this as winding isovalue for WindingNumber mode
+	UPROPERTY(BlueprintReadWrite, Category = Options)
+	float WindingIsoValue = 0.5;
+
+	// random rays to add beyond +/- major axes, for raycast sampling
+	UPROPERTY(BlueprintReadWrite, Category = Options)
+	int RaysPerSample = 0;
+
+
+	/** Nudge sample points out by this amount to try to counteract numerical issues */
+	UPROPERTY(BlueprintReadWrite, Category = Options)
+	float NormalOffset = 1e-6f;
+
+	/**  */
+	UPROPERTY(BlueprintReadWrite, Category = Options)
+	bool bCompactResult = true;
+};
+
+
+
+
+UCLASS()
+class GEOMETRYSCRIPTINGCORE_API UGeometryScriptLibrary_MeshRepairFunctions : public UBlueprintFunctionLibrary
+{
+	GENERATED_BODY()
+public:
+
+	UFUNCTION(BlueprintCallable, Category = "GeometryScript|Repair")
+	static UPARAM(DisplayName = "Target Mesh") UDynamicMesh* 
+	CompactMesh(  
+		UDynamicMesh* TargetMesh, 
+		UGeometryScriptDebug* Debug = nullptr);
+
+	UFUNCTION(BlueprintCallable, Category = "GeometryScript|Repair")
+	static UPARAM(DisplayName = "Target Mesh") UDynamicMesh* 
+	WeldMeshEdges(  
+		UDynamicMesh* TargetMesh, 
+		FGeometryScriptWeldEdgesOptions WeldOptions,
+		UGeometryScriptDebug* Debug = nullptr);
+
+	UFUNCTION(BlueprintCallable, Category = "GeometryScript|Repair")
+	static UPARAM(DisplayName = "Target Mesh") UDynamicMesh* 
+	FillAllMeshHoles(  
+		UDynamicMesh* TargetMesh, 
+		FGeometryScriptFillHolesOptions FillOptions,
+		int32& NumFilledHoles,
+		int32& NumFailedHoleFills,
+		UGeometryScriptDebug* Debug = nullptr);
+
+	UFUNCTION(BlueprintCallable, Category = "GeometryScript|Repair")
+	static UPARAM(DisplayName = "Target Mesh") UDynamicMesh* 
+	RemoveSmallComponents(  
+		UDynamicMesh* TargetMesh, 
+		FGeometryScriptRemoveSmallComponentOptions Options,
+		UGeometryScriptDebug* Debug = nullptr);
+
+	UFUNCTION(BlueprintCallable, Category = "GeometryScript|Repair")
+	static UPARAM(DisplayName = "Target Mesh") UDynamicMesh* 
+	RemoveHiddenTriangles(  
+		UDynamicMesh* TargetMesh, 
+		FGeometryScriptRemoveHiddenTrianglesOptions Options,
+		UGeometryScriptDebug* Debug = nullptr);
+
+};
\ No newline at end of file
diff --git a/Engine/Plugins/Experimental/GeometryScripting/Source/GeometryScriptingCore/Public/GeometryScript/MeshSimplifyFunctions.h b/Engine/Plugins/Experimental/GeometryScripting/Source/GeometryScriptingCore/Public/GeometryScript/MeshSimplifyFunctions.h
new file mode 100644
--- /dev/null
+++ b/Engine/Plugins/Experimental/GeometryScripting/Source/GeometryScriptingCore/Public/GeometryScript/MeshSimplifyFunctions.h
@@ -0,0 +1,98 @@
+// Copyright Epic Games, Inc. All Rights Reserved.
+
+#pragma once
+
+#include "CoreMinimal.h"
+#include "Kismet/BlueprintFunctionLibrary.h"
+#include "GeometryScript/GeometryScriptTypes.h"
+#include "MeshSimplifyFunctions.generated.h"
+
+class UDynamicMesh;
+
+USTRUCT(BlueprintType)
+struct GEOMETRYSCRIPTINGCORE_API FGeometryScriptPlanarSimplifyOptions
+{
+	GENERATED_BODY()
+public:
+	/** */
+	UPROPERTY(BlueprintReadWrite, Category = Options)
+	float AngleThreshold = 0.001;
+};
+
+
+
+UENUM(BlueprintType)
+enum class EGeometryScriptRemoveMeshSimplificationType : uint8
+{
+	StandardQEM = 0,
+	VolumePreserving = 1,
+	AttributeAware = 2
+};
+
+
+USTRUCT(Blueprintable)
+struct GEOMETRYSCRIPTINGCORE_API FGeometryScriptSimplifyMeshOptions
+{
+	GENERATED_BODY()
+public:
+
+	UPROPERTY(BlueprintReadWrite, Category = Options)
+	EGeometryScriptRemoveMeshSimplificationType Method = EGeometryScriptRemoveMeshSimplificationType::AttributeAware;
+
+	UPROPERTY(BlueprintReadWrite, Category = Options)
+	bool bAllowSeamCollapse = true;
+
+	UPROPERTY(BlueprintReadWrite, Category = Options)
+	bool bAllowSeamSmoothing = true;
+
+	UPROPERTY(BlueprintReadWrite, Category = Options)
+	bool bAllowSeamSplits = true;
+
+	UPROPERTY(BlueprintReadWrite, Category = Options)
+	bool bPreserveVertexPositions = false;
+
+	UPROPERTY(BlueprintReadWrite, Category = Options)
+	bool bRetainQuadricMemory = false;
+};
+
+
+
+
+UCLASS()
+class GEOMETRYSCRIPTINGCORE_API UGeometryScriptLibrary_MeshSimplifyFunctions : public UBlueprintFunctionLibrary
+{
+	GENERATED_BODY()
+public:
+
+	UFUNCTION(BlueprintCallable, Category = "GeometryScript|Simplification")
+	static UPARAM(DisplayName = "Target Mesh") UDynamicMesh* 
+	ApplySimplifyToPlanar(  
+		UDynamicMesh* TargetMesh, 
+		FGeometryScriptPlanarSimplifyOptions Options,
+		UGeometryScriptDebug* Debug = nullptr);
+
+	UFUNCTION(BlueprintCallable, Category = "GeometryScript|Simplification")
+	static UPARAM(DisplayName = "Target Mesh") UDynamicMesh* 
+	ApplySimplifyToTriangleCount(  
+		UDynamicMesh* TargetMesh, 
+		int32 TriangleCount,
+		FGeometryScriptSimplifyMeshOptions Options,
+		UGeometryScriptDebug* Debug = nullptr);
+
+	UFUNCTION(BlueprintCallable, Category = "GeometryScript|Simplification")
+	static UPARAM(DisplayName = "Target Mesh") UDynamicMesh* 
+	ApplySimplifyToVertexCount(  
+		UDynamicMesh* TargetMesh, 
+		int32 VertexCount,
+		FGeometryScriptSimplifyMeshOptions Options,
+		UGeometryScriptDebug* Debug = nullptr);
+
+	UFUNCTION(BlueprintCallable, Category = "GeometryScript|Simplification")
+	static UPARAM(DisplayName = "Target Mesh") UDynamicMesh* 
+	ApplySimplifyToTolerance(  
+		UDynamicMesh* TargetMesh, 
+		float Tolerance,
+		FGeometryScriptSimplifyMeshOptions Options,
+		UGeometryScriptDebug* Debug = nullptr);
+
+};
\ No newline at end of file
diff --git a/Engine/Plugins/Experimental/GeometryScripting/Source/GeometryScriptingCore/Public/GeometryScript/MeshTransformFunctions.h b/Engine/Plugins/Experimental/GeometryScripting/Source/GeometryScriptingCore/Public/GeometryScript/MeshTransformFunctions.h
new file mode 100644
--- /dev/null
+++ b/Engine/Plugins/Experimental/GeometryScripting/Source/GeometryScriptingCore/Public/GeometryScript/MeshTransformFunctions.h
@@ -0,0 +1,40 @@
+// Copyright Epic Games, Inc. All Rights Reserved.
+
+#pragma once
+
+#include "CoreMinimal.h"
+#include "Kismet/BlueprintFunctionLibrary.h"
+#include "GeometryScript/GeometryScriptTypes.h"
+#include "MeshTransformFunctions.generated.h"
+
+class UDynamicMesh;
+
+
+UCLASS()
+class GEOMETRYSCRIPTINGCORE_API UGeometryScriptLibrary_MeshTransformFunctions : public UBlueprintFunctionLibrary
+{
+	GENERATED_BODY()
+public:
+
+	UFUNCTION(BlueprintCallable, Category = "GeometryScript|Transforms")
+	static UPARAM(DisplayName = "Target Mesh") UDynamicMesh* 
+	TransformMesh(
+		UDynamicMesh* TargetMesh,
+		FTransform Transform,
+		UGeometryScriptDebug* Debug = nullptr );
+
+	UFUNCTION(BlueprintCallable, Category = "GeometryScript|Transforms")
+	static UPARAM(DisplayName = "Target Mesh") UDynamicMesh* 
+	TranslateMesh(
+		UDynamicMesh* TargetMesh,
+		FVector Translation,
+		UGeometryScriptDebug* Debug = nullptr );
+
+	UFUNCTION(BlueprintCallable, Category = "GeometryScript|Transforms")
+	static UPARAM(DisplayName = "Target Mesh") UDynamicMesh* 
+	ScaleMesh(
+		UDynamicMesh* TargetMesh,
+		FVector Scale = FVector(1,1,1),
+		UGeometryScriptDebug* Debug = nullptr );
+
+};
\ No newline at end of file
diff --git a/Engine/Plugins/Experimental/GeometryScripting/Source/GeometryScriptingCore/Public/GeometryScript/MeshUVFunctions.h b/Engine/Plugins/Experimental/GeometryScripting/Source/GeometryScriptingCore/Public/GeometryScript/MeshUVFunctions.h
new file mode 100644
--- /dev/null
+++ b/Engine/Plugins/Experimental/GeometryScripting/Source/GeometryScriptingCore/Public/GeometryScript/MeshUVFunctions.h
@@ -0,0 +1,31 @@
+// Copyright Epic Games, Inc. All Rights Reserved.
+
+#pragma once
+
+#include "CoreMinimal.h"
+#include "Kismet/BlueprintFunctionLibrary.h"
+#include "GeometryScript/GeometryScriptTypes.h"
+#include "MeshUVFunctions.generated.h"
+
+class UDynamicMesh;
+
+
+UCLASS()
+class GEOMETRYSCRIPTINGCORE_API UGeometryScriptLibrary_MeshUVFunctions : public UBlueprintFunctionLibrary
+{
+	GENERATED_BODY()
+public:
+
+
+
+	UFUNCTION(BlueprintCallable, Category = "GeometryScript|UVs")
+	static UPARAM(DisplayName = "Target Mesh") UDynamicMesh* 
+	SetMeshTriangleUVs( 
+		UDynamicMesh* TargetMesh, 
+		int UVSetIndex,
+		int TriangleID, 
+		FGeometryScriptUVTriangle UVs,
+		bool& bIsValidTriangle, 
+		bool bDeferChangeNotifications = false );
+
+};
\ No newline at end of file
diff --git a/Engine/Plugins/Experimental/GeometryScripting/Source/GeometryScriptingCore/Public/GeometryScript/MeshVertexColorFunctions.h b/Engine/Plugins/Experimental/GeometryScripting/Source/GeometryScriptingCore/Public/GeometryScript/MeshVertexColorFunctions.h
new file mode 100644
--- /dev/null
+++ b/Engine/Plugins/Experimental/GeometryScripting/Source/GeometryScriptingCore/Public/GeometryScript/MeshVertexColorFunctions.h
@@ -0,0 +1,28 @@
+// Copyright Epic Games, Inc. All Rights Reserved.
+
+#pragma once
+
+#include "CoreMinimal.h"
+#include "Kismet/BlueprintFunctionLibrary.h"
+#include "GeometryScript/GeometryScriptTypes.h"
+#include "MeshVertexColorFunctions.generated.h"
+
+class UDynamicMesh;
+
+
+
+UCLASS()
+class GEOMETRYSCRIPTINGCORE_API UGeometryScriptLibrary_MeshVertexColorFunctions : public UBlueprintFunctionLibrary
+{
+	GENERATED_BODY()
+public:
+
+	UFUNCTION(BlueprintCallable, Category = "GeometryScript|Repair")
+	static UPARAM(DisplayName = "Target Mesh") UDynamicMesh*
+	SetMeshPerVertexColors(
+		UDynamicMesh* TargetMesh,
+		const TArray<FLinearColor>& VertexColors,
+		UGeometryScriptDebug* Debug = nullptr);
+
+};
+
diff --git a/Engine/Plugins/Experimental/GeometryScripting/Source/GeometryScriptingCore/Public/GeometryScript/MeshVoxelFunctions.h b/Engine/Plugins/Experimental/GeometryScripting/Source/GeometryScriptingCore/Public/GeometryScript/MeshVoxelFunctions.h
new file mode 100644
--- /dev/null
+++ b/Engine/Plugins/Experimental/GeometryScripting/Source/GeometryScriptingCore/Public/GeometryScript/MeshVoxelFunctions.h
@@ -0,0 +1,127 @@
+// Copyright Epic Games, Inc. All Rights Reserved.
+
+#pragma once
+
+#include "CoreMinimal.h"
+#include "Kismet/BlueprintFunctionLibrary.h"
+#include "GeometryScript/GeometryScriptTypes.h"
+#include "MeshVoxelFunctions.generated.h"
+
+class UDynamicMesh;
+
+UENUM(BlueprintType)
+enum class EGeometryScriptGridSizingMethod : uint8
+{
+	GridCellSize = 0,
+	GridResolution = 1
+};
+
+
+/***
+ * Parameters for 3D grids, eg grids used for sampling, SDFs, voxelization, etc
+ */
+USTRUCT(BlueprintType)
+struct GEOMETRYSCRIPTINGCORE_API FGeometryScript3DGridParameters
+{
+	GENERATED_BODY()
+public:
+	/** SizeMethod determines how the parameters below will be interpreted to define the size of a 3D sampling/voxel grid */
+	UPROPERTY(BlueprintReadWrite, Category = Options)
+	EGeometryScriptGridSizingMethod SizeMethod = EGeometryScriptGridSizingMethod::GridResolution;
+
+	/** Use a specific grid cell size, and construct a grid with dimensions large enough to contain the target object */
+	UPROPERTY(BlueprintReadWrite, Category = Options, meta = (EditCondition = "SizeMethod == EGeometryScriptGridSizingMethod::GridCellSize"));
+	float GridCellSize = 0.5;
+
+	/** Use a specific grid resolution, with the grid cell size derived form the target object bounds such that this is the number of cells along the longest box dimension */
+	UPROPERTY(BlueprintReadWrite, Category = Options, meta = (EditCondition = "SizeMethod == EGeometryScriptGridSizingMethod::GridResolution"))
+	int GridResolution = 64;
+};
+
+
+USTRUCT(BlueprintType)
+struct GEOMETRYSCRIPTINGCORE_API FGeometryScriptSolidifyOptions
+{
+	GENERATED_BODY()
+public:
+	UPROPERTY(BlueprintReadWrite, Category = Options)
+	FGeometryScript3DGridParameters GridParameters;
+
+	UPROPERTY(BlueprintReadWrite, Category = Options)
+	float WindingThreshold = 0.5;
+
+	UPROPERTY(BlueprintReadWrite, Category = Options)
+	bool bSolidAtBoundaries;
+
+	UPROPERTY(BlueprintReadWrite, Category = Options)
+	float ExtendBounds = 1.0;
+
+	UPROPERTY(BlueprintReadWrite, Category = Options)
+	int SurfaceSearchSteps = 3;
+};
+
+
+UENUM(BlueprintType)
+enum class EGeometryScriptMorphologicalOpType : uint8
+{
+	/** Expand the shapes outward */
+	Dilate = 0,
+	/** Shrink the shapes inward */
+	Contract = 1,
+	/** Dilate and then contract, to delete small negative features (sharp inner corners, small holes) */
+	Close = 2,
+	/** Contract and then dilate, to delete small positive features (sharp outer corners, small isolated pieces) */
+	Open = 3
+
+	// note: keep above compatible with TImplicitMorphology
+};
+
+
+
+
+USTRUCT(BlueprintType)
+struct GEOMETRYSCRIPTINGCORE_API FGeometryScriptMorphologyOptions
+{
+	GENERATED_BODY()
+public:
+	UPROPERTY(BlueprintReadWrite, Category = Options)
+	FGeometryScript3DGridParameters SDFGridParameters;
+
+	UPROPERTY(BlueprintReadWrite, Category = Options)
+	bool bUseSeparateMeshGrid = false;
+
+	UPROPERTY(BlueprintReadWrite, Category = Options)
+	FGeometryScript3DGridParameters MeshGridParameters;
+
+	UPROPERTY(BlueprintReadWrite, Category = Options)
+	EGeometryScriptMorphologicalOpType Operation = EGeometryScriptMorphologicalOpType::Dilate;
+
+	UPROPERTY(BlueprintReadWrite, Category = Options)
+	float Distance = 1.0;
+
+};
+
+
+
+
+UCLASS()
+class GEOMETRYSCRIPTINGCORE_API UGeometryScriptLibrary_MeshVoxelFunctions : public UBlueprintFunctionLibrary
+{
+	GENERATED_BODY()
+public:
+
+	UFUNCTION(BlueprintCallable, Category = "GeometryScript|Voxel")
+	static UPARAM(DisplayName = "Target Mesh") UDynamicMesh*
+	ApplyMeshSolidify(
+		UDynamicMesh* TargetMesh,
+		FGeometryScriptSolidifyOptions Options,
+		UGeometryScriptDebug* Debug = nullptr);
+
+	UFUNCTION(BlueprintCallable, Category = "GeometryScript|Voxel")
+	static UPARAM(DisplayName = "Target Mesh") UDynamicMesh*
+	ApplyMeshMorphology(
+		UDynamicMesh* TargetMesh,
+		FGeometryScriptMorphologyOptions Options,
+		UGeometryScriptDebug* Debug = nullptr);
+
+};
\ No newline at end of file
diff --git a/Engine/Plugins/Experimental/GeometryScripting/Source/GeometryScriptingCore/Public/GeometryScript/SceneUtilityFunctions.h b/Engine/Plugins/Experimental/GeometryScripting/Source/GeometryScriptingCore/Public/GeometryScript/SceneUtilityFunctions.h
new file mode 100644
--- /dev/null
+++ b/Engine/Plugins/Experimental/GeometryScripting/Source/GeometryScriptingCore/Public/GeometryScript/SceneUtilityFunctions.h
@@ -0,0 +1,47 @@
+// Copyright Epic Games, Inc. All Rights Reserved.
+
+#pragma once
+
+#include "CoreMinimal.h"
+#include "Kismet/BlueprintFunctionLibrary.h"
+#include "GeometryScript/GeometryScriptTypes.h"
+#include "SceneUtilityFunctions.generated.h"
+
+class UStaticMesh;
+class UDynamicMesh;
+
+
+USTRUCT(BlueprintType)
+struct GEOMETRYSCRIPTINGCORE_API FGeometryScriptCopyMeshFromComponentOptions
+{
+	GENERATED_BODY()
+public:
+	UPROPERTY(BlueprintReadWrite, Category = Options)
+	bool bWantNormals = true;
+
+	UPROPERTY(BlueprintReadWrite, Category = Options)
+	bool bWantTangents = true;
+
+	UPROPERTY(BlueprintReadWrite, Category = Options)
+	FGeometryScriptMeshReadLOD RequestedLOD = FGeometryScriptMeshReadLOD();
+};
+
+
+UCLASS()
+class GEOMETRYSCRIPTINGCORE_API UGeometryScriptLibrary_SceneUtilityFunctions : public UBlueprintFunctionLibrary
+{
+	GENERATED_BODY()
+public:
+
+	UFUNCTION(BlueprintCallable, Category = "GeometryScript|Scene", meta = (ExpandEnumAsExecs = "Outcome"))
+	static UPARAM(DisplayName = "Dynamic Mesh") UDynamicMesh*
+	CopyMeshFromComponent(
+		USceneComponent* Component,
+		UDynamicMesh* ToDynamicMesh,
+		FGeometryScriptCopyMeshFromComponentOptions Options,
+		bool bTransformToWorld,
+		FTransform& LocalToWorld,
+		TEnumAsByte<EGeometryScriptOutcomePins>& Outcome,
+		UGeometryScriptDebug* Debug = nullptr);
+
+};
\ No newline at end of file
diff --git a/Engine/Plugins/Experimental/GeometryScripting/Source/GeometryScriptingCore/Public/GeometryScriptingCoreModule.h b/Engine/Plugins/Experimental/GeometryScripting/Source/GeometryScriptingCore/Public/GeometryScriptingCoreModule.h
new file mode 100644
--- /dev/null
+++ b/Engine/Plugins/Experimental/GeometryScripting/Source/GeometryScriptingCore/Public/GeometryScriptingCoreModule.h
@@ -0,0 +1,17 @@
+// Copyright Epic Games, Inc. All Rights Reserved.
+
+#pragma once
+
+#include "CoreMinimal.h"
+#include "Modules/ModuleManager.h"
+
+class FGeometryScriptingCoreModule : public IModuleInterface
+{
+public:
+
+	/** IModuleInterface implementation */
+	virtual void StartupModule() override;
+	virtual void ShutdownModule() override;
+};
+
+
