GeometryScripting: fix LWC FVector cast issue
#rnx

diff --git a/Engine/Plugins/Experimental/GeometryScripting/Source/GeometryScriptingCore/Private/MeshDeformFunctions.cpp b/Engine/Plugins/Experimental/GeometryScripting/Source/GeometryScriptingCore/Private/MeshDeformFunctions.cpp
--- a/Engine/Plugins/Experimental/GeometryScripting/Source/GeometryScriptingCore/Private/MeshDeformFunctions.cpp
+++ b/Engine/Plugins/Experimental/GeometryScripting/Source/GeometryScriptingCore/Private/MeshDeformFunctions.cpp
@@ -33,11 +33,11 @@ UDynamicMesh* UGeometryScriptLibrary_MeshDeformFunctions::ApplyBendWarpToMesh(
 		UE::Geometry::AppendError(Debug, EGeometryScriptErrorType::InvalidInputs, LOCTEXT("ApplyBendWarpToMesh_InvalidAxis", "ApplyBendWarpToMesh: BendAxis is degenerate"));
 		UseAxis = FVector3d::UnitZ();
 	}
-	FFrame3d WarpFrame = FFrame3d(Options.BendOrigin, UseAxis);
+	FFrame3d WarpFrame = FFrame3d((FVector3d)Options.BendOrigin, (FVector3d)UseAxis);
 
 	if (Options.TowardAxis.Length() > 0)
 	{
-		FVector UseTowardAxis = Options.TowardAxis.GetSafeNormal();
+		FVector3d UseTowardAxis = (FVector3d)Options.TowardAxis.GetSafeNormal();
 		if ( WarpFrame.Z().Dot(UseTowardAxis) < 0.95 )
 		{
 			WarpFrame.ConstrainedAlignAxis(1, UseTowardAxis, WarpFrame.Z());
@@ -98,11 +98,11 @@ UDynamicMesh* UGeometryScriptLibrary_MeshDeformFunctions::ApplyTwistWarpToMesh(
 		UE::Geometry::AppendError(Debug, EGeometryScriptErrorType::InvalidInputs, LOCTEXT("ApplyTwistWarpToMesh_InvalidAxis", "ApplyTwistWarpToMesh: TwistAxis is degenerate"));
 		UseAxis = FVector3d::UnitZ();
 	}
-	FFrame3d WarpFrame = FFrame3d(Options.TwistOrigin, UseAxis);
+	FFrame3d WarpFrame = FFrame3d((FVector3d)Options.TwistOrigin, UseAxis);
 
 	if (Options.TowardAxis.Length() > 0)
 	{
-		FVector UseTowardAxis = Options.TowardAxis.GetSafeNormal();
+		FVector3d UseTowardAxis = (FVector3d)Options.TowardAxis.GetSafeNormal();
 		if ( WarpFrame.Z().Dot(UseTowardAxis) < 0.95 )
 		{
 			WarpFrame.ConstrainedAlignAxis(1, UseTowardAxis, WarpFrame.Z());
