GeometryScripting: add ScriptNames to the various BP libraries to make using Python more pleasant
#rnx

diff --git a/Engine/Plugins/Experimental/GeometryScripting/Source/GeometryScriptingCore/Public/GeometryScript/MeshAssetFunctions.h b/Engine/Plugins/Experimental/GeometryScripting/Source/GeometryScriptingCore/Public/GeometryScript/MeshAssetFunctions.h
--- a/Engine/Plugins/Experimental/GeometryScripting/Source/GeometryScriptingCore/Public/GeometryScript/MeshAssetFunctions.h
+++ b/Engine/Plugins/Experimental/GeometryScripting/Source/GeometryScriptingCore/Public/GeometryScript/MeshAssetFunctions.h
@@ -46,7 +46,7 @@ public:
 
 
 
-UCLASS()
+UCLASS(meta = (ScriptName = "GeometryScript_AssetUtils"))
 class GEOMETRYSCRIPTINGCORE_API UGeometryScriptLibrary_StaticMeshFunctions : public UBlueprintFunctionLibrary
 {
 	GENERATED_BODY()
diff --git a/Engine/Plugins/Experimental/GeometryScripting/Source/GeometryScriptingCore/Public/GeometryScript/MeshBasicEditFunctions.h b/Engine/Plugins/Experimental/GeometryScripting/Source/GeometryScriptingCore/Public/GeometryScript/MeshBasicEditFunctions.h
--- a/Engine/Plugins/Experimental/GeometryScripting/Source/GeometryScriptingCore/Public/GeometryScript/MeshBasicEditFunctions.h
+++ b/Engine/Plugins/Experimental/GeometryScripting/Source/GeometryScriptingCore/Public/GeometryScript/MeshBasicEditFunctions.h
@@ -11,7 +11,7 @@ class UDynamicMesh;
 
 
 
-UCLASS()
+UCLASS(meta = (ScriptName = "GeometryScript_MeshEdits"))
 class GEOMETRYSCRIPTINGCORE_API UGeometryScriptLibrary_MeshBasicEditFunctions : public UBlueprintFunctionLibrary
 {
 	GENERATED_BODY()
diff --git a/Engine/Plugins/Experimental/GeometryScripting/Source/GeometryScriptingCore/Public/GeometryScript/MeshBooleanFunctions.h b/Engine/Plugins/Experimental/GeometryScripting/Source/GeometryScriptingCore/Public/GeometryScript/MeshBooleanFunctions.h
--- a/Engine/Plugins/Experimental/GeometryScripting/Source/GeometryScriptingCore/Public/GeometryScript/MeshBooleanFunctions.h
+++ b/Engine/Plugins/Experimental/GeometryScripting/Source/GeometryScriptingCore/Public/GeometryScript/MeshBooleanFunctions.h
@@ -91,7 +91,7 @@ public:
 
 
 
-UCLASS()
+UCLASS(meta = (ScriptName = "GeometryScript_MeshBooleans"))
 class GEOMETRYSCRIPTINGCORE_API UGeometryScriptLibrary_MeshBooleanFunctions : public UBlueprintFunctionLibrary
 {
 	GENERATED_BODY()
diff --git a/Engine/Plugins/Experimental/GeometryScripting/Source/GeometryScriptingCore/Public/GeometryScript/MeshDeformFunctions.h b/Engine/Plugins/Experimental/GeometryScripting/Source/GeometryScriptingCore/Public/GeometryScript/MeshDeformFunctions.h
--- a/Engine/Plugins/Experimental/GeometryScripting/Source/GeometryScriptingCore/Public/GeometryScript/MeshDeformFunctions.h
+++ b/Engine/Plugins/Experimental/GeometryScripting/Source/GeometryScriptingCore/Public/GeometryScript/MeshDeformFunctions.h
@@ -83,7 +83,7 @@ public:
 
 
 
-UCLASS()
+UCLASS(meta = (ScriptName = "GeometryScript_MeshDeformers"))
 class GEOMETRYSCRIPTINGCORE_API UGeometryScriptLibrary_MeshDeformFunctions : public UBlueprintFunctionLibrary
 {
 	GENERATED_BODY()
diff --git a/Engine/Plugins/Experimental/GeometryScripting/Source/GeometryScriptingCore/Public/GeometryScript/MeshNormalsFunctions.h b/Engine/Plugins/Experimental/GeometryScripting/Source/GeometryScriptingCore/Public/GeometryScript/MeshNormalsFunctions.h
--- a/Engine/Plugins/Experimental/GeometryScripting/Source/GeometryScriptingCore/Public/GeometryScript/MeshNormalsFunctions.h
+++ b/Engine/Plugins/Experimental/GeometryScripting/Source/GeometryScriptingCore/Public/GeometryScript/MeshNormalsFunctions.h
@@ -42,7 +42,7 @@ public:
 };
 
 
-UCLASS()
+UCLASS(meta = (ScriptName = "GeometryScript_Normals"))
 class GEOMETRYSCRIPTINGCORE_API UGeometryScriptLibrary_MeshNormalsFunctions : public UBlueprintFunctionLibrary
 {
 	GENERATED_BODY()
diff --git a/Engine/Plugins/Experimental/GeometryScripting/Source/GeometryScriptingCore/Public/GeometryScript/MeshPrimitiveFunctions.h b/Engine/Plugins/Experimental/GeometryScripting/Source/GeometryScriptingCore/Public/GeometryScript/MeshPrimitiveFunctions.h
--- a/Engine/Plugins/Experimental/GeometryScripting/Source/GeometryScriptingCore/Public/GeometryScript/MeshPrimitiveFunctions.h
+++ b/Engine/Plugins/Experimental/GeometryScripting/Source/GeometryScriptingCore/Public/GeometryScript/MeshPrimitiveFunctions.h
@@ -31,7 +31,7 @@ public:
 };
 
 
-UCLASS()
+UCLASS(meta = (ScriptName = "GeometryScript_Primitives"))
 class GEOMETRYSCRIPTINGCORE_API UGeometryScriptLibrary_MeshPrimitiveFunctions : public UBlueprintFunctionLibrary
 {
 	GENERATED_BODY()
diff --git a/Engine/Plugins/Experimental/GeometryScripting/Source/GeometryScriptingCore/Public/GeometryScript/MeshQueryFunctions.h b/Engine/Plugins/Experimental/GeometryScripting/Source/GeometryScriptingCore/Public/GeometryScript/MeshQueryFunctions.h
--- a/Engine/Plugins/Experimental/GeometryScripting/Source/GeometryScriptingCore/Public/GeometryScript/MeshQueryFunctions.h
+++ b/Engine/Plugins/Experimental/GeometryScripting/Source/GeometryScriptingCore/Public/GeometryScript/MeshQueryFunctions.h
@@ -10,7 +10,7 @@
 class UDynamicMesh;
 
 
-UCLASS()
+UCLASS(meta = (ScriptName = "GeometryScript_MeshQueries"))
 class GEOMETRYSCRIPTINGCORE_API UGeometryScriptLibrary_MeshQueryFunctions : public UBlueprintFunctionLibrary
 {
 	GENERATED_BODY()
diff --git a/Engine/Plugins/Experimental/GeometryScripting/Source/GeometryScriptingCore/Public/GeometryScript/MeshRepairFunctions.h b/Engine/Plugins/Experimental/GeometryScripting/Source/GeometryScriptingCore/Public/GeometryScript/MeshRepairFunctions.h
--- a/Engine/Plugins/Experimental/GeometryScripting/Source/GeometryScriptingCore/Public/GeometryScript/MeshRepairFunctions.h
+++ b/Engine/Plugins/Experimental/GeometryScripting/Source/GeometryScriptingCore/Public/GeometryScript/MeshRepairFunctions.h
@@ -114,7 +114,7 @@ public:
 
 
 
-UCLASS()
+UCLASS(meta = (ScriptName = "GeometryScript_MeshRepair"))
 class GEOMETRYSCRIPTINGCORE_API UGeometryScriptLibrary_MeshRepairFunctions : public UBlueprintFunctionLibrary
 {
 	GENERATED_BODY()
diff --git a/Engine/Plugins/Experimental/GeometryScripting/Source/GeometryScriptingCore/Public/GeometryScript/MeshSimplifyFunctions.h b/Engine/Plugins/Experimental/GeometryScripting/Source/GeometryScriptingCore/Public/GeometryScript/MeshSimplifyFunctions.h
--- a/Engine/Plugins/Experimental/GeometryScripting/Source/GeometryScriptingCore/Public/GeometryScript/MeshSimplifyFunctions.h
+++ b/Engine/Plugins/Experimental/GeometryScripting/Source/GeometryScriptingCore/Public/GeometryScript/MeshSimplifyFunctions.h
@@ -58,7 +58,7 @@ public:
 
 
 
-UCLASS()
+UCLASS(meta = (ScriptName = "GeometryScript_MeshSimplification"))
 class GEOMETRYSCRIPTINGCORE_API UGeometryScriptLibrary_MeshSimplifyFunctions : public UBlueprintFunctionLibrary
 {
 	GENERATED_BODY()
diff --git a/Engine/Plugins/Experimental/GeometryScripting/Source/GeometryScriptingCore/Public/GeometryScript/MeshTransformFunctions.h b/Engine/Plugins/Experimental/GeometryScripting/Source/GeometryScriptingCore/Public/GeometryScript/MeshTransformFunctions.h
--- a/Engine/Plugins/Experimental/GeometryScripting/Source/GeometryScriptingCore/Public/GeometryScript/MeshTransformFunctions.h
+++ b/Engine/Plugins/Experimental/GeometryScripting/Source/GeometryScriptingCore/Public/GeometryScript/MeshTransformFunctions.h
@@ -10,7 +10,7 @@
 class UDynamicMesh;
 
 
-UCLASS()
+UCLASS(meta = (ScriptName = "GeometryScript_MeshTransforms"))
 class GEOMETRYSCRIPTINGCORE_API UGeometryScriptLibrary_MeshTransformFunctions : public UBlueprintFunctionLibrary
 {
 	GENERATED_BODY()
diff --git a/Engine/Plugins/Experimental/GeometryScripting/Source/GeometryScriptingCore/Public/GeometryScript/MeshUVFunctions.h b/Engine/Plugins/Experimental/GeometryScripting/Source/GeometryScriptingCore/Public/GeometryScript/MeshUVFunctions.h
--- a/Engine/Plugins/Experimental/GeometryScripting/Source/GeometryScriptingCore/Public/GeometryScript/MeshUVFunctions.h
+++ b/Engine/Plugins/Experimental/GeometryScripting/Source/GeometryScriptingCore/Public/GeometryScript/MeshUVFunctions.h
@@ -10,7 +10,7 @@
 class UDynamicMesh;
 
 
-UCLASS()
+UCLASS(meta = (ScriptName = "GeometryScript_UVs"))
 class GEOMETRYSCRIPTINGCORE_API UGeometryScriptLibrary_MeshUVFunctions : public UBlueprintFunctionLibrary
 {
 	GENERATED_BODY()
diff --git a/Engine/Plugins/Experimental/GeometryScripting/Source/GeometryScriptingCore/Public/GeometryScript/MeshVertexColorFunctions.h b/Engine/Plugins/Experimental/GeometryScripting/Source/GeometryScriptingCore/Public/GeometryScript/MeshVertexColorFunctions.h
--- a/Engine/Plugins/Experimental/GeometryScripting/Source/GeometryScriptingCore/Public/GeometryScript/MeshVertexColorFunctions.h
+++ b/Engine/Plugins/Experimental/GeometryScripting/Source/GeometryScriptingCore/Public/GeometryScript/MeshVertexColorFunctions.h
@@ -11,7 +11,7 @@ class UDynamicMesh;
 
 
 
-UCLASS()
+UCLASS(meta = (ScriptName = "GeometryScript_VertexColors"))
 class GEOMETRYSCRIPTINGCORE_API UGeometryScriptLibrary_MeshVertexColorFunctions : public UBlueprintFunctionLibrary
 {
 	GENERATED_BODY()
diff --git a/Engine/Plugins/Experimental/GeometryScripting/Source/GeometryScriptingCore/Public/GeometryScript/MeshVoxelFunctions.h b/Engine/Plugins/Experimental/GeometryScripting/Source/GeometryScriptingCore/Public/GeometryScript/MeshVoxelFunctions.h
--- a/Engine/Plugins/Experimental/GeometryScripting/Source/GeometryScriptingCore/Public/GeometryScript/MeshVoxelFunctions.h
+++ b/Engine/Plugins/Experimental/GeometryScripting/Source/GeometryScriptingCore/Public/GeometryScript/MeshVoxelFunctions.h
@@ -104,7 +104,7 @@ public:
 
 
 
-UCLASS()
+UCLASS(meta = (ScriptName = "GeometryScript_MeshVoxelProcessing"))
 class GEOMETRYSCRIPTINGCORE_API UGeometryScriptLibrary_MeshVoxelFunctions : public UBlueprintFunctionLibrary
 {
 	GENERATED_BODY()
diff --git a/Engine/Plugins/Experimental/GeometryScripting/Source/GeometryScriptingCore/Public/GeometryScript/SceneUtilityFunctions.h b/Engine/Plugins/Experimental/GeometryScripting/Source/GeometryScriptingCore/Public/GeometryScript/SceneUtilityFunctions.h
--- a/Engine/Plugins/Experimental/GeometryScripting/Source/GeometryScriptingCore/Public/GeometryScript/SceneUtilityFunctions.h
+++ b/Engine/Plugins/Experimental/GeometryScripting/Source/GeometryScriptingCore/Public/GeometryScript/SceneUtilityFunctions.h
@@ -27,7 +27,7 @@ public:
 };
 
 
-UCLASS()
+UCLASS(meta = (ScriptName = "GeometryScript_SceneUtils"))
 class GEOMETRYSCRIPTINGCORE_API UGeometryScriptLibrary_SceneUtilityFunctions : public UBlueprintFunctionLibrary
 {
 	GENERATED_BODY()
