GeometryScript: add ScriptMethod meta tags to simplify Python usage
#rnx

diff --git a/Engine/Plugins/Experimental/GeometryScripting/Source/GeometryScriptingCore/Public/GeometryScript/MeshBasicEditFunctions.h b/Engine/Plugins/Experimental/GeometryScripting/Source/GeometryScriptingCore/Public/GeometryScript/MeshBasicEditFunctions.h
--- a/Engine/Plugins/Experimental/GeometryScripting/Source/GeometryScriptingCore/Public/GeometryScript/MeshBasicEditFunctions.h
+++ b/Engine/Plugins/Experimental/GeometryScripting/Source/GeometryScriptingCore/Public/GeometryScript/MeshBasicEditFunctions.h
@@ -17,7 +17,7 @@ class GEOMETRYSCRIPTINGCORE_API UGeometryScriptLibrary_MeshBasicEditFunctions :
 	GENERATED_BODY()
 public:
 
-	UFUNCTION(BlueprintCallable, Category = "GeometryScript|MeshEdits")
+	UFUNCTION(BlueprintCallable, Category = "GeometryScript|MeshEdits", meta=(ScriptMethod))
 	static UPARAM(DisplayName = "Target Mesh") UDynamicMesh* 
 	DiscardMeshAttributes( 
 		UDynamicMesh* TargetMesh, 
@@ -25,7 +25,7 @@ public:
 
 
 
-	UFUNCTION(BlueprintCallable, Category = "GeometryScript|MeshEdits")
+	UFUNCTION(BlueprintCallable, Category = "GeometryScript|MeshEdits", meta=(ScriptMethod))
 	static UPARAM(DisplayName = "Target Mesh") UDynamicMesh* 
 	SetVertexPosition( 
 		UDynamicMesh* TargetMesh, 
@@ -35,7 +35,7 @@ public:
 		bool bDeferChangeNotifications = false );
 
 
-	UFUNCTION(BlueprintCallable, Category = "GeometryScript|MeshEdits")
+	UFUNCTION(BlueprintCallable, Category = "GeometryScript|MeshEdits", meta=(ScriptMethod))
 	static UPARAM(DisplayName = "Target Mesh") UDynamicMesh* 
 	AddVertexToMesh( 
 		UDynamicMesh* TargetMesh, 
@@ -43,7 +43,7 @@ public:
 		int& NewVertexIndex,
 		bool bDeferChangeNotifications = false );
 
-	UFUNCTION(BlueprintCallable, Category = "GeometryScript|MeshEdits")
+	UFUNCTION(BlueprintCallable, Category = "GeometryScript|MeshEdits", meta=(ScriptMethod))
 	static UPARAM(DisplayName = "Target Mesh") UDynamicMesh* 
 	AddVerticesToMesh( 
 		UDynamicMesh* TargetMesh, 
@@ -51,7 +51,7 @@ public:
 		TArray<int>& NewIndices,
 		bool bDeferChangeNotifications = false );
 
-	UFUNCTION(BlueprintCallable, Category = "GeometryScript|MeshEdits")
+	UFUNCTION(BlueprintCallable, Category = "GeometryScript|MeshEdits", meta=(ScriptMethod))
 	static UPARAM(DisplayName = "Target Mesh") UDynamicMesh* 
 	DeleteVertexFromMesh( 
 		UDynamicMesh* TargetMesh, 
@@ -61,7 +61,7 @@ public:
 
 
 
-	UFUNCTION(BlueprintCallable, Category = "GeometryScript|MeshEdits")
+	UFUNCTION(BlueprintCallable, Category = "GeometryScript|MeshEdits", meta=(ScriptMethod))
 	static UPARAM(DisplayName = "Target Mesh") UDynamicMesh* 
 	AddTriangleToMesh( 
 		UDynamicMesh* TargetMesh, 
@@ -72,7 +72,7 @@ public:
 		UGeometryScriptDebug* Debug = nullptr );
 
 
-	UFUNCTION(BlueprintCallable, Category = "GeometryScript|MeshEdits")
+	UFUNCTION(BlueprintCallable, Category = "GeometryScript|MeshEdits", meta=(ScriptMethod))
 	static UPARAM(DisplayName = "Target Mesh") UDynamicMesh* 
 	AddTrianglesToMesh( 
 		UDynamicMesh* TargetMesh, 
@@ -82,7 +82,7 @@ public:
 		bool bDeferChangeNotifications = false,
 		UGeometryScriptDebug* Debug = nullptr );
 
-	UFUNCTION(BlueprintCallable, Category = "GeometryScript|MeshEdits")
+	UFUNCTION(BlueprintCallable, Category = "GeometryScript|MeshEdits", meta=(ScriptMethod))
 	static UPARAM(DisplayName = "Target Mesh") UDynamicMesh* 
 	DeleteTriangleFromMesh( 
 		UDynamicMesh* TargetMesh, 
@@ -93,7 +93,7 @@ public:
 
 
 
-	UFUNCTION(BlueprintCallable, Category = "GeometryScript|MeshEdits")
+	UFUNCTION(BlueprintCallable, Category = "GeometryScript|MeshEdits", meta=(ScriptMethod))
 	static UPARAM(DisplayName = "Target Mesh") UDynamicMesh* 
 	AppendMesh( 
 		UDynamicMesh* TargetMesh, 
@@ -102,7 +102,7 @@ public:
 		bool bDeferChangeNotifications = false,
 		UGeometryScriptDebug* Debug = nullptr);
 
-	UFUNCTION(BlueprintCallable, Category = "GeometryScript|MeshEdits")
+	UFUNCTION(BlueprintCallable, Category = "GeometryScript|MeshEdits", meta=(ScriptMethod))
 	static UPARAM(DisplayName = "Target Mesh") UDynamicMesh* 
 	AppendMeshRepeated( 
 		UDynamicMesh* TargetMesh, 
diff --git a/Engine/Plugins/Experimental/GeometryScripting/Source/GeometryScriptingCore/Public/GeometryScript/MeshBooleanFunctions.h b/Engine/Plugins/Experimental/GeometryScripting/Source/GeometryScriptingCore/Public/GeometryScript/MeshBooleanFunctions.h
--- a/Engine/Plugins/Experimental/GeometryScripting/Source/GeometryScriptingCore/Public/GeometryScript/MeshBooleanFunctions.h
+++ b/Engine/Plugins/Experimental/GeometryScripting/Source/GeometryScriptingCore/Public/GeometryScript/MeshBooleanFunctions.h
@@ -97,7 +97,7 @@ class GEOMETRYSCRIPTINGCORE_API UGeometryScriptLibrary_MeshBooleanFunctions : pu
 	GENERATED_BODY()
 public:
 
-	UFUNCTION(BlueprintCallable, Category = "GeometryScript|Booleans")
+	UFUNCTION(BlueprintCallable, Category = "GeometryScript|Booleans", meta=(ScriptMethod))
 	static UPARAM(DisplayName = "Target Mesh") UDynamicMesh* 
 	ApplyMeshBoolean( 
 		UDynamicMesh* TargetMesh, 
@@ -109,7 +109,7 @@ public:
 		UGeometryScriptDebug* Debug = nullptr);
 
 
-	UFUNCTION(BlueprintCallable, Category = "GeometryScript|Booleans")
+	UFUNCTION(BlueprintCallable, Category = "GeometryScript|Booleans", meta=(ScriptMethod))
 	static UPARAM(DisplayName = "Target Mesh") UDynamicMesh* 
 	ApplyMeshSelfUnion( 
 		UDynamicMesh* TargetMesh, 
@@ -118,7 +118,7 @@ public:
 
 
 
-	UFUNCTION(BlueprintCallable, Category = "GeometryScript|Booleans")
+	UFUNCTION(BlueprintCallable, Category = "GeometryScript|Booleans", meta=(ScriptMethod))
 	static UPARAM(DisplayName = "Target Mesh") UDynamicMesh* 
 	ApplyMeshPlaneCut( 
 		UDynamicMesh* TargetMesh, 
@@ -128,7 +128,7 @@ public:
 		UGeometryScriptDebug* Debug = nullptr);
 
 
-	UFUNCTION(BlueprintCallable, Category = "GeometryScript|Booleans")
+	UFUNCTION(BlueprintCallable, Category = "GeometryScript|Booleans", meta=(ScriptMethod))
 	static UPARAM(DisplayName = "Target Mesh") UDynamicMesh* 
 	ApplyMeshMirror( 
 		UDynamicMesh* TargetMesh, 
diff --git a/Engine/Plugins/Experimental/GeometryScripting/Source/GeometryScriptingCore/Public/GeometryScript/MeshDeformFunctions.h b/Engine/Plugins/Experimental/GeometryScripting/Source/GeometryScriptingCore/Public/GeometryScript/MeshDeformFunctions.h
--- a/Engine/Plugins/Experimental/GeometryScripting/Source/GeometryScriptingCore/Public/GeometryScript/MeshDeformFunctions.h
+++ b/Engine/Plugins/Experimental/GeometryScripting/Source/GeometryScriptingCore/Public/GeometryScript/MeshDeformFunctions.h
@@ -89,7 +89,7 @@ class GEOMETRYSCRIPTINGCORE_API UGeometryScriptLibrary_MeshDeformFunctions : pub
 	GENERATED_BODY()
 public:
 
-	UFUNCTION(BlueprintCallable, Category = "GeometryScript|Deformations")
+	UFUNCTION(BlueprintCallable, Category = "GeometryScript|Deformations", meta=(ScriptMethod))
 	static UPARAM(DisplayName = "Target Mesh") UDynamicMesh* 
 	ApplyBendWarpToMesh(  
 		UDynamicMesh* TargetMesh, 
@@ -98,7 +98,7 @@ public:
 		float BendExtent = 50,
 		UGeometryScriptDebug* Debug = nullptr);
 
-	UFUNCTION(BlueprintCallable, Category = "GeometryScript|Deformations")
+	UFUNCTION(BlueprintCallable, Category = "GeometryScript|Deformations", meta=(ScriptMethod))
 	static UPARAM(DisplayName = "Target Mesh") UDynamicMesh* 
 	ApplyTwistWarpToMesh(  
 		UDynamicMesh* TargetMesh, 
diff --git a/Engine/Plugins/Experimental/GeometryScripting/Source/GeometryScriptingCore/Public/GeometryScript/MeshNormalsFunctions.h b/Engine/Plugins/Experimental/GeometryScripting/Source/GeometryScriptingCore/Public/GeometryScript/MeshNormalsFunctions.h
--- a/Engine/Plugins/Experimental/GeometryScripting/Source/GeometryScriptingCore/Public/GeometryScript/MeshNormalsFunctions.h
+++ b/Engine/Plugins/Experimental/GeometryScripting/Source/GeometryScriptingCore/Public/GeometryScript/MeshNormalsFunctions.h
@@ -48,38 +48,38 @@ class GEOMETRYSCRIPTINGCORE_API UGeometryScriptLibrary_MeshNormalsFunctions : pu
 	GENERATED_BODY()
 public:
 
-	UFUNCTION(BlueprintCallable, Category = "GeometryScript|Normals")
+	UFUNCTION(BlueprintCallable, Category = "GeometryScript|Normals", meta=(ScriptMethod))
 	static UPARAM(DisplayName = "Target Mesh") UDynamicMesh* 
 	FlipNormals( 
 		UDynamicMesh* TargetMesh, 
 		UGeometryScriptDebug* Debug = nullptr);
 
-	UFUNCTION(BlueprintCallable, Category = "GeometryScript|Normals")
+	UFUNCTION(BlueprintCallable, Category = "GeometryScript|Normals", meta=(ScriptMethod))
 	static UPARAM(DisplayName = "Target Mesh") UDynamicMesh* 
 	AutoRepairNormals( 
 		UDynamicMesh* TargetMesh, 
 		UGeometryScriptDebug* Debug = nullptr);
 
-	UFUNCTION(BlueprintCallable, Category = "GeometryScript|Normals")
+	UFUNCTION(BlueprintCallable, Category = "GeometryScript|Normals", meta=(ScriptMethod))
 	static UPARAM(DisplayName = "Target Mesh") UDynamicMesh* 
 	SetPerVertexNormals( 
 		UDynamicMesh* TargetMesh, 
 		UGeometryScriptDebug* Debug = nullptr);
 
-	UFUNCTION(BlueprintCallable, Category = "GeometryScript|Normals")
+	UFUNCTION(BlueprintCallable, Category = "GeometryScript|Normals", meta=(ScriptMethod))
 	static UPARAM(DisplayName = "Target Mesh") UDynamicMesh* 
 	SetPerFaceNormals( 
 		UDynamicMesh* TargetMesh, 
 		UGeometryScriptDebug* Debug = nullptr);
 
-	UFUNCTION(BlueprintCallable, Category = "GeometryScript|Normals")
+	UFUNCTION(BlueprintCallable, Category = "GeometryScript|Normals", meta=(ScriptMethod))
 	static UPARAM(DisplayName = "Target Mesh") UDynamicMesh* 
 	RecomputeNormals(  
 		UDynamicMesh* TargetMesh, 
 		FGeometryScriptCalculateNormalsOptions CalculateOptions,
 		UGeometryScriptDebug* Debug = nullptr);
 
-	UFUNCTION(BlueprintCallable, Category = "GeometryScript|Normals")
+	UFUNCTION(BlueprintCallable, Category = "GeometryScript|Normals", meta=(ScriptMethod))
 	static UPARAM(DisplayName = "Target Mesh") UDynamicMesh* 
 	ComputeSplitNormals( 
 		UDynamicMesh* TargetMesh, 
@@ -90,7 +90,7 @@ public:
 
 
 
-	UFUNCTION(BlueprintCallable, Category = "GeometryScript|Normals")
+	UFUNCTION(BlueprintCallable, Category = "GeometryScript|Normals", meta=(ScriptMethod))
 	static UPARAM(DisplayName = "Target Mesh") UDynamicMesh* 
 	SetMeshTriangleNormals( 
 		UDynamicMesh* TargetMesh, 
diff --git a/Engine/Plugins/Experimental/GeometryScripting/Source/GeometryScriptingCore/Public/GeometryScript/MeshPrimitiveFunctions.h b/Engine/Plugins/Experimental/GeometryScripting/Source/GeometryScriptingCore/Public/GeometryScript/MeshPrimitiveFunctions.h
--- a/Engine/Plugins/Experimental/GeometryScripting/Source/GeometryScriptingCore/Public/GeometryScript/MeshPrimitiveFunctions.h
+++ b/Engine/Plugins/Experimental/GeometryScripting/Source/GeometryScriptingCore/Public/GeometryScript/MeshPrimitiveFunctions.h
@@ -38,7 +38,7 @@ class GEOMETRYSCRIPTINGCORE_API UGeometryScriptLibrary_MeshPrimitiveFunctions :
 public:
 
 
-	UFUNCTION(BlueprintCallable, Category = "GeometryScript|Primitives")
+	UFUNCTION(BlueprintCallable, Category = "GeometryScript|Primitives", meta=(ScriptMethod))
 	static UPARAM(DisplayName = "Target Mesh") UDynamicMesh* 
 	AppendBox( 
 		UDynamicMesh* TargetMesh, 
@@ -51,7 +51,7 @@ public:
 		UGeometryScriptDebug* Debug = nullptr);
 
 
-	UFUNCTION(BlueprintCallable, Category = "GeometryScript|Primitives")
+	UFUNCTION(BlueprintCallable, Category = "GeometryScript|Primitives", meta=(ScriptMethod))
 	static UPARAM(DisplayName = "Target Mesh") UDynamicMesh* 
 	AppendSphereLatLong( 
 		UDynamicMesh* TargetMesh, 
@@ -63,7 +63,7 @@ public:
 		UGeometryScriptDebug* Debug = nullptr);
 
 
-	UFUNCTION(BlueprintCallable, Category = "GeometryScript|Primitives")
+	UFUNCTION(BlueprintCallable, Category = "GeometryScript|Primitives", meta=(ScriptMethod))
 	static UPARAM(DisplayName = "Target Mesh") UDynamicMesh* 
 	AppendSphereBox( 
 		UDynamicMesh* TargetMesh, 
@@ -76,7 +76,7 @@ public:
 		UGeometryScriptDebug* Debug = nullptr);
 
 
-	UFUNCTION(BlueprintCallable, Category = "GeometryScript|Primitives")
+	UFUNCTION(BlueprintCallable, Category = "GeometryScript|Primitives", meta=(ScriptMethod))
 	static UPARAM(DisplayName = "Target Mesh") UDynamicMesh* 
 	AppendCapsule( 
 		UDynamicMesh* TargetMesh, 
@@ -89,7 +89,7 @@ public:
 		UGeometryScriptDebug* Debug = nullptr);
 
 
-	UFUNCTION(BlueprintCallable, Category = "GeometryScript|Primitives")
+	UFUNCTION(BlueprintCallable, Category = "GeometryScript|Primitives", meta=(ScriptMethod))
 	static UPARAM(DisplayName = "Target Mesh") UDynamicMesh* 
 	AppendCylinder( 
 		UDynamicMesh* TargetMesh, 
@@ -103,7 +103,7 @@ public:
 		UGeometryScriptDebug* Debug = nullptr);
 
 
-	UFUNCTION(BlueprintCallable, Category = "GeometryScript|Primitives")
+	UFUNCTION(BlueprintCallable, Category = "GeometryScript|Primitives", meta=(ScriptMethod))
 	static UPARAM(DisplayName = "Target Mesh") UDynamicMesh* 
 	AppendCone( 
 		UDynamicMesh* TargetMesh, 
@@ -118,7 +118,7 @@ public:
 		UGeometryScriptDebug* Debug = nullptr);
 
 
-	UFUNCTION(BlueprintCallable, Category = "GeometryScript|Primitives")
+	UFUNCTION(BlueprintCallable, Category = "GeometryScript|Primitives", meta=(ScriptMethod))
 	static UPARAM(DisplayName = "Target Mesh") UDynamicMesh* 
 	AppendTorus( 
 		UDynamicMesh* TargetMesh, 
@@ -135,7 +135,7 @@ public:
 	 * Polygon should be oriented counter-clockwise to produce a correctly-oriented shape, otherwise it will be inside-out
 	 * Polygon endpoint is not repeated.
 	 */
-	UFUNCTION(BlueprintCallable, Category = "GeometryScript|Primitives")
+	UFUNCTION(BlueprintCallable, Category = "GeometryScript|Primitives", meta=(ScriptMethod))
 	static UPARAM(DisplayName = "Target Mesh") UDynamicMesh* 
 	AppendSimpleRevolvePolygon( 
 		UDynamicMesh* TargetMesh, 
@@ -147,7 +147,7 @@ public:
 		UGeometryScriptDebug* Debug = nullptr);
 
 
-	UFUNCTION(BlueprintCallable, Category = "GeometryScript|Primitives")
+	UFUNCTION(BlueprintCallable, Category = "GeometryScript|Primitives", meta=(ScriptMethod))
 	static UPARAM(DisplayName = "Target Mesh") UDynamicMesh* 
 	AppendSimpleSweptPolygon( 
 		UDynamicMesh* TargetMesh, 
@@ -163,7 +163,7 @@ public:
 
 
 
-	UFUNCTION(BlueprintCallable, Category = "GeometryScript|Primitives")
+	UFUNCTION(BlueprintCallable, Category = "GeometryScript|Primitives", meta=(ScriptMethod))
 	static UPARAM(DisplayName = "Target Mesh") UDynamicMesh* 
 	AppendRectangle( 
 		UDynamicMesh* TargetMesh, 
@@ -174,7 +174,7 @@ public:
 		int32 StepsHeight = 0,
 		UGeometryScriptDebug* Debug = nullptr);
 
-	UFUNCTION(BlueprintCallable, Category = "GeometryScript|Primitives")
+	UFUNCTION(BlueprintCallable, Category = "GeometryScript|Primitives", meta=(ScriptMethod))
 	static UPARAM(DisplayName = "Target Mesh") UDynamicMesh* 
 	AppendRoundRectangle( 
 		UDynamicMesh* TargetMesh, 
@@ -187,7 +187,7 @@ public:
 		int32 StepsRound = 4,
 		UGeometryScriptDebug* Debug = nullptr);
 
-	UFUNCTION(BlueprintCallable, Category = "GeometryScript|Primitives")
+	UFUNCTION(BlueprintCallable, Category = "GeometryScript|Primitives", meta=(ScriptMethod))
 	static UPARAM(DisplayName = "Target Mesh") UDynamicMesh* 
 	AppendDisc( 
 		UDynamicMesh* TargetMesh, 
diff --git a/Engine/Plugins/Experimental/GeometryScripting/Source/GeometryScriptingCore/Public/GeometryScript/MeshQueryFunctions.h b/Engine/Plugins/Experimental/GeometryScripting/Source/GeometryScriptingCore/Public/GeometryScript/MeshQueryFunctions.h
--- a/Engine/Plugins/Experimental/GeometryScripting/Source/GeometryScriptingCore/Public/GeometryScript/MeshQueryFunctions.h
+++ b/Engine/Plugins/Experimental/GeometryScripting/Source/GeometryScriptingCore/Public/GeometryScript/MeshQueryFunctions.h
@@ -16,80 +16,80 @@ class GEOMETRYSCRIPTINGCORE_API UGeometryScriptLibrary_MeshQueryFunctions : publ
 	GENERATED_BODY()
 public:
 
-	UFUNCTION(BlueprintPure, Category = "GeometryScript|MeshQueries")
+	UFUNCTION(BlueprintPure, Category = "GeometryScript|MeshQueries", meta=(ScriptMethod))
 	static bool GetIsDenseMesh( UDynamicMesh* TargetMesh );
 
-	UFUNCTION(BlueprintPure, Category = "GeometryScript|MeshQueries")
+	UFUNCTION(BlueprintPure, Category = "GeometryScript|MeshQueries", meta=(ScriptMethod))
 	static bool GetMeshHasAttributeSet( UDynamicMesh* TargetMesh );
 
-	UFUNCTION(BlueprintPure, Category = "GeometryScript|MeshQueries")
+	UFUNCTION(BlueprintPure, Category = "GeometryScript|MeshQueries", meta=(ScriptMethod))
 	static UPARAM(DisplayName = "Bounding Box") FBox GetMeshBoundingBox( UDynamicMesh* TargetMesh );
 
-	UFUNCTION(BlueprintPure, Category = "GeometryScript|MeshQueries")
+	UFUNCTION(BlueprintPure, Category = "GeometryScript|MeshQueries", meta=(ScriptMethod))
 	static void GetMeshVolumeArea( UDynamicMesh* TargetMesh, float& SurfaceArea, float& Volume );
 
-	UFUNCTION(BlueprintPure, Category = "GeometryScript|MeshQueries")
+	UFUNCTION(BlueprintPure, Category = "GeometryScript|MeshQueries", meta=(ScriptMethod))
 	static bool GetIsClosedMesh( UDynamicMesh* TargetMesh );
 
-	UFUNCTION(BlueprintPure, Category = "GeometryScript|MeshQueries")
+	UFUNCTION(BlueprintPure, Category = "GeometryScript|MeshQueries", meta=(ScriptMethod))
 	static UPARAM(DisplayName = "Num Loops") int32 GetNumOpenBorderLoops( UDynamicMesh* TargetMesh, bool& bAmbiguousTopologyFound );
 
-	UFUNCTION(BlueprintPure, Category = "GeometryScript|MeshQueries")
+	UFUNCTION(BlueprintPure, Category = "GeometryScript|MeshQueries", meta=(ScriptMethod))
 	static UPARAM(DisplayName = "Num Edges") int32 GetNumOpenBorderEdges( UDynamicMesh* TargetMesh );
 
 
 
-	UFUNCTION(BlueprintPure, Category = "GeometryScript|MeshQueries")
+	UFUNCTION(BlueprintPure, Category = "GeometryScript|MeshQueries", meta=(ScriptMethod))
 	static UPARAM(DisplayName = "Triangle Count") int32 GetTriangleCount( UDynamicMesh* TargetMesh );
 
-	UFUNCTION(BlueprintPure, Category = "GeometryScript|MeshQueries")
+	UFUNCTION(BlueprintPure, Category = "GeometryScript|MeshQueries", meta=(ScriptMethod))
 	static UPARAM(DisplayName = "Num Triangle IDs") int32 GetNumTriangleIDs( UDynamicMesh* TargetMesh );
 
-	UFUNCTION(BlueprintPure, Category = "GeometryScript|MeshQueries")
+	UFUNCTION(BlueprintPure, Category = "GeometryScript|MeshQueries", meta=(ScriptMethod))
 	static bool GetHasTriangleIDGaps( UDynamicMesh* TargetMesh );
 
-	UFUNCTION(BlueprintPure, Category = "GeometryScript|MeshQueries")
+	UFUNCTION(BlueprintPure, Category = "GeometryScript|MeshQueries", meta=(ScriptMethod))
 	static bool IsValidTriangleID( UDynamicMesh* TargetMesh, int32 TriangleID );
 
-	UFUNCTION(BlueprintPure, Category = "GeometryScript|MeshQueries")
+	UFUNCTION(BlueprintPure, Category = "GeometryScript|MeshQueries", meta=(ScriptMethod))
 	static UPARAM(DisplayName = "Indices") FIntVector GetTriangleIndices( UDynamicMesh* TargetMesh, int32 TriangleID, bool& bIsValidTriangle );
 
-	UFUNCTION(BlueprintPure, Category = "GeometryScript|MeshQueries")
+	UFUNCTION(BlueprintPure, Category = "GeometryScript|MeshQueries", meta=(ScriptMethod))
 	static void GetTrianglePositions( UDynamicMesh* TargetMesh, int32 TriangleID, bool& bIsValidTriangle, FVector& Vertex1, FVector& Vertex2, FVector& Vertex3 );
 
-	UFUNCTION(BlueprintPure, Category = "GeometryScript|MeshQueries")
+	UFUNCTION(BlueprintPure, Category = "GeometryScript|MeshQueries", meta=(ScriptMethod))
 	static UPARAM(DisplayName = "Normal") FVector GetTriangleFaceNormal( UDynamicMesh* TargetMesh, int32 TriangleID, bool& bIsValidTriangle );
 
 
 
 
-	UFUNCTION(BlueprintPure, Category = "GeometryScript|MeshQueries")
+	UFUNCTION(BlueprintPure, Category = "GeometryScript|MeshQueries", meta=(ScriptMethod))
 	static UPARAM(DisplayName = "Vertex Count") int32 GetVertexCount( UDynamicMesh* TargetMesh );
 
-	UFUNCTION(BlueprintPure, Category = "GeometryScript|MeshQueries")
+	UFUNCTION(BlueprintPure, Category = "GeometryScript|MeshQueries", meta=(ScriptMethod))
 	static UPARAM(DisplayName = "Num Vertex IDs") int32 GetNumVertexIDs( UDynamicMesh* TargetMesh );
 
-	UFUNCTION(BlueprintPure, Category = "GeometryScript|MeshQueries")
+	UFUNCTION(BlueprintPure, Category = "GeometryScript|MeshQueries", meta=(ScriptMethod))
 	static bool GetHasVertexIDGaps( UDynamicMesh* TargetMesh );
 
-	UFUNCTION(BlueprintPure, Category = "GeometryScript|MeshQueries")
+	UFUNCTION(BlueprintPure, Category = "GeometryScript|MeshQueries", meta=(ScriptMethod))
 	static bool IsValidVertexID( UDynamicMesh* TargetMesh, int32 VertexID );
 
-	UFUNCTION(BlueprintPure, Category = "GeometryScript|MeshQueries")
+	UFUNCTION(BlueprintPure, Category = "GeometryScript|MeshQueries", meta=(ScriptMethod))
 	static UPARAM(DisplayName = "Position") FVector GetVertexPosition( UDynamicMesh* TargetMesh, int32 VertexID, bool& bIsValidVertex );
 
-	UFUNCTION(BlueprintPure, Category = "GeometryScript|MeshQueries")
+	UFUNCTION(BlueprintPure, Category = "GeometryScript|MeshQueries", meta=(ScriptMethod))
 	static UPARAM(DisplayName = "Positions") TArray<FVector> GetAllVertexPositions( UDynamicMesh* TargetMesh, bool& bHasVertexIDGaps );
 
 
 
-	UFUNCTION(BlueprintPure, Category = "GeometryScript|MeshQueries")
+	UFUNCTION(BlueprintPure, Category = "GeometryScript|MeshQueries", meta=(ScriptMethod))
 	static UPARAM(DisplayName = "Num UV Sets") int32 GetNumUVSets( UDynamicMesh* TargetMesh );
 
-	UFUNCTION(BlueprintPure, Category = "GeometryScript|MeshQueries")
+	UFUNCTION(BlueprintPure, Category = "GeometryScript|MeshQueries", meta=(ScriptMethod))
 	static UPARAM(DisplayName = "Bounding Box") FBox2D GetUVSetBoundingBox( UDynamicMesh* TargetMesh, int UVSetIndex, bool& bIsValidUVSet, bool& bUVSetIsEmpty );
 
-	UFUNCTION(BlueprintPure, Category = "GeometryScript|MeshQueries")
+	UFUNCTION(BlueprintPure, Category = "GeometryScript|MeshQueries", meta=(ScriptMethod))
 	static void GetTriangleUVs( UDynamicMesh* TargetMesh, int32 UVSetIndex, int32 TriangleID, FVector2D& UV1, FVector2D& UV2, FVector2D& UV3, bool& bHaveValidUVs );
 
 
diff --git a/Engine/Plugins/Experimental/GeometryScripting/Source/GeometryScriptingCore/Public/GeometryScript/MeshRepairFunctions.h b/Engine/Plugins/Experimental/GeometryScripting/Source/GeometryScriptingCore/Public/GeometryScript/MeshRepairFunctions.h
--- a/Engine/Plugins/Experimental/GeometryScripting/Source/GeometryScriptingCore/Public/GeometryScript/MeshRepairFunctions.h
+++ b/Engine/Plugins/Experimental/GeometryScripting/Source/GeometryScriptingCore/Public/GeometryScript/MeshRepairFunctions.h
@@ -120,20 +120,20 @@ class GEOMETRYSCRIPTINGCORE_API UGeometryScriptLibrary_MeshRepairFunctions : pub
 	GENERATED_BODY()
 public:
 
-	UFUNCTION(BlueprintCallable, Category = "GeometryScript|Repair")
+	UFUNCTION(BlueprintCallable, Category = "GeometryScript|Repair", meta=(ScriptMethod))
 	static UPARAM(DisplayName = "Target Mesh") UDynamicMesh* 
 	CompactMesh(  
 		UDynamicMesh* TargetMesh, 
 		UGeometryScriptDebug* Debug = nullptr);
 
-	UFUNCTION(BlueprintCallable, Category = "GeometryScript|Repair")
+	UFUNCTION(BlueprintCallable, Category = "GeometryScript|Repair", meta=(ScriptMethod))
 	static UPARAM(DisplayName = "Target Mesh") UDynamicMesh* 
 	WeldMeshEdges(  
 		UDynamicMesh* TargetMesh, 
 		FGeometryScriptWeldEdgesOptions WeldOptions,
 		UGeometryScriptDebug* Debug = nullptr);
 
-	UFUNCTION(BlueprintCallable, Category = "GeometryScript|Repair")
+	UFUNCTION(BlueprintCallable, Category = "GeometryScript|Repair", meta=(ScriptMethod))
 	static UPARAM(DisplayName = "Target Mesh") UDynamicMesh* 
 	FillAllMeshHoles(  
 		UDynamicMesh* TargetMesh, 
@@ -142,14 +142,14 @@ public:
 		int32& NumFailedHoleFills,
 		UGeometryScriptDebug* Debug = nullptr);
 
-	UFUNCTION(BlueprintCallable, Category = "GeometryScript|Repair")
+	UFUNCTION(BlueprintCallable, Category = "GeometryScript|Repair", meta=(ScriptMethod))
 	static UPARAM(DisplayName = "Target Mesh") UDynamicMesh* 
 	RemoveSmallComponents(  
 		UDynamicMesh* TargetMesh, 
 		FGeometryScriptRemoveSmallComponentOptions Options,
 		UGeometryScriptDebug* Debug = nullptr);
 
-	UFUNCTION(BlueprintCallable, Category = "GeometryScript|Repair")
+	UFUNCTION(BlueprintCallable, Category = "GeometryScript|Repair", meta=(ScriptMethod))
 	static UPARAM(DisplayName = "Target Mesh") UDynamicMesh* 
 	RemoveHiddenTriangles(  
 		UDynamicMesh* TargetMesh, 
diff --git a/Engine/Plugins/Experimental/GeometryScripting/Source/GeometryScriptingCore/Public/GeometryScript/MeshSimplifyFunctions.h b/Engine/Plugins/Experimental/GeometryScripting/Source/GeometryScriptingCore/Public/GeometryScript/MeshSimplifyFunctions.h
--- a/Engine/Plugins/Experimental/GeometryScripting/Source/GeometryScriptingCore/Public/GeometryScript/MeshSimplifyFunctions.h
+++ b/Engine/Plugins/Experimental/GeometryScripting/Source/GeometryScriptingCore/Public/GeometryScript/MeshSimplifyFunctions.h
@@ -64,14 +64,14 @@ class GEOMETRYSCRIPTINGCORE_API UGeometryScriptLibrary_MeshSimplifyFunctions : p
 	GENERATED_BODY()
 public:
 
-	UFUNCTION(BlueprintCallable, Category = "GeometryScript|Simplification")
+	UFUNCTION(BlueprintCallable, Category = "GeometryScript|Simplification", meta=(ScriptMethod))
 	static UPARAM(DisplayName = "Target Mesh") UDynamicMesh* 
 	ApplySimplifyToPlanar(  
 		UDynamicMesh* TargetMesh, 
 		FGeometryScriptPlanarSimplifyOptions Options,
 		UGeometryScriptDebug* Debug = nullptr);
 
-	UFUNCTION(BlueprintCallable, Category = "GeometryScript|Simplification")
+	UFUNCTION(BlueprintCallable, Category = "GeometryScript|Simplification", meta=(ScriptMethod))
 	static UPARAM(DisplayName = "Target Mesh") UDynamicMesh* 
 	ApplySimplifyToTriangleCount(  
 		UDynamicMesh* TargetMesh, 
@@ -79,7 +79,7 @@ public:
 		FGeometryScriptSimplifyMeshOptions Options,
 		UGeometryScriptDebug* Debug = nullptr);
 
-	UFUNCTION(BlueprintCallable, Category = "GeometryScript|Simplification")
+	UFUNCTION(BlueprintCallable, Category = "GeometryScript|Simplification", meta=(ScriptMethod))
 	static UPARAM(DisplayName = "Target Mesh") UDynamicMesh* 
 	ApplySimplifyToVertexCount(  
 		UDynamicMesh* TargetMesh, 
@@ -87,7 +87,7 @@ public:
 		FGeometryScriptSimplifyMeshOptions Options,
 		UGeometryScriptDebug* Debug = nullptr);
 
-	UFUNCTION(BlueprintCallable, Category = "GeometryScript|Simplification")
+	UFUNCTION(BlueprintCallable, Category = "GeometryScript|Simplification", meta=(ScriptMethod))
 	static UPARAM(DisplayName = "Target Mesh") UDynamicMesh* 
 	ApplySimplifyToTolerance(  
 		UDynamicMesh* TargetMesh, 
diff --git a/Engine/Plugins/Experimental/GeometryScripting/Source/GeometryScriptingCore/Public/GeometryScript/MeshTransformFunctions.h b/Engine/Plugins/Experimental/GeometryScripting/Source/GeometryScriptingCore/Public/GeometryScript/MeshTransformFunctions.h
--- a/Engine/Plugins/Experimental/GeometryScripting/Source/GeometryScriptingCore/Public/GeometryScript/MeshTransformFunctions.h
+++ b/Engine/Plugins/Experimental/GeometryScripting/Source/GeometryScriptingCore/Public/GeometryScript/MeshTransformFunctions.h
@@ -16,21 +16,21 @@ class GEOMETRYSCRIPTINGCORE_API UGeometryScriptLibrary_MeshTransformFunctions :
 	GENERATED_BODY()
 public:
 
-	UFUNCTION(BlueprintCallable, Category = "GeometryScript|Transforms")
+	UFUNCTION(BlueprintCallable, Category = "GeometryScript|Transforms", meta=(ScriptMethod))
 	static UPARAM(DisplayName = "Target Mesh") UDynamicMesh* 
 	TransformMesh(
 		UDynamicMesh* TargetMesh,
 		FTransform Transform,
 		UGeometryScriptDebug* Debug = nullptr );
 
-	UFUNCTION(BlueprintCallable, Category = "GeometryScript|Transforms")
+	UFUNCTION(BlueprintCallable, Category = "GeometryScript|Transforms", meta=(ScriptMethod))
 	static UPARAM(DisplayName = "Target Mesh") UDynamicMesh* 
 	TranslateMesh(
 		UDynamicMesh* TargetMesh,
 		FVector Translation,
 		UGeometryScriptDebug* Debug = nullptr );
 
-	UFUNCTION(BlueprintCallable, Category = "GeometryScript|Transforms")
+	UFUNCTION(BlueprintCallable, Category = "GeometryScript|Transforms", meta=(ScriptMethod))
 	static UPARAM(DisplayName = "Target Mesh") UDynamicMesh* 
 	ScaleMesh(
 		UDynamicMesh* TargetMesh,
diff --git a/Engine/Plugins/Experimental/GeometryScripting/Source/GeometryScriptingCore/Public/GeometryScript/MeshUVFunctions.h b/Engine/Plugins/Experimental/GeometryScripting/Source/GeometryScriptingCore/Public/GeometryScript/MeshUVFunctions.h
--- a/Engine/Plugins/Experimental/GeometryScripting/Source/GeometryScriptingCore/Public/GeometryScript/MeshUVFunctions.h
+++ b/Engine/Plugins/Experimental/GeometryScripting/Source/GeometryScriptingCore/Public/GeometryScript/MeshUVFunctions.h
@@ -18,7 +18,7 @@ public:
 
 
 
-	UFUNCTION(BlueprintCallable, Category = "GeometryScript|UVs")
+	UFUNCTION(BlueprintCallable, Category = "GeometryScript|UVs", meta=(ScriptMethod))
 	static UPARAM(DisplayName = "Target Mesh") UDynamicMesh* 
 	SetMeshTriangleUVs( 
 		UDynamicMesh* TargetMesh, 
diff --git a/Engine/Plugins/Experimental/GeometryScripting/Source/GeometryScriptingCore/Public/GeometryScript/MeshVertexColorFunctions.h b/Engine/Plugins/Experimental/GeometryScripting/Source/GeometryScriptingCore/Public/GeometryScript/MeshVertexColorFunctions.h
--- a/Engine/Plugins/Experimental/GeometryScripting/Source/GeometryScriptingCore/Public/GeometryScript/MeshVertexColorFunctions.h
+++ b/Engine/Plugins/Experimental/GeometryScripting/Source/GeometryScriptingCore/Public/GeometryScript/MeshVertexColorFunctions.h
@@ -17,7 +17,7 @@ class GEOMETRYSCRIPTINGCORE_API UGeometryScriptLibrary_MeshVertexColorFunctions
 	GENERATED_BODY()
 public:
 
-	UFUNCTION(BlueprintCallable, Category = "GeometryScript|Repair")
+	UFUNCTION(BlueprintCallable, Category = "GeometryScript|Repair", meta=(ScriptMethod))
 	static UPARAM(DisplayName = "Target Mesh") UDynamicMesh*
 	SetMeshPerVertexColors(
 		UDynamicMesh* TargetMesh,
diff --git a/Engine/Plugins/Experimental/GeometryScripting/Source/GeometryScriptingCore/Public/GeometryScript/MeshVoxelFunctions.h b/Engine/Plugins/Experimental/GeometryScripting/Source/GeometryScriptingCore/Public/GeometryScript/MeshVoxelFunctions.h
--- a/Engine/Plugins/Experimental/GeometryScripting/Source/GeometryScriptingCore/Public/GeometryScript/MeshVoxelFunctions.h
+++ b/Engine/Plugins/Experimental/GeometryScripting/Source/GeometryScriptingCore/Public/GeometryScript/MeshVoxelFunctions.h
@@ -110,14 +110,14 @@ class GEOMETRYSCRIPTINGCORE_API UGeometryScriptLibrary_MeshVoxelFunctions : publ
 	GENERATED_BODY()
 public:
 
-	UFUNCTION(BlueprintCallable, Category = "GeometryScript|Voxel")
+	UFUNCTION(BlueprintCallable, Category = "GeometryScript|Voxel", meta=(ScriptMethod))
 	static UPARAM(DisplayName = "Target Mesh") UDynamicMesh*
 	ApplyMeshSolidify(
 		UDynamicMesh* TargetMesh,
 		FGeometryScriptSolidifyOptions Options,
 		UGeometryScriptDebug* Debug = nullptr);
 
-	UFUNCTION(BlueprintCallable, Category = "GeometryScript|Voxel")
+	UFUNCTION(BlueprintCallable, Category = "GeometryScript|Voxel", meta=(ScriptMethod))
 	static UPARAM(DisplayName = "Target Mesh") UDynamicMesh*
 	ApplyMeshMorphology(
 		UDynamicMesh* TargetMesh,
