GeometryScript: continue building out library (still at the trivial-wrappers stage)
- ComputeMeshConvexHull, ComputeMeshSweptHull
- CreateNewVolumeFromMesh, CreateNewStaticMeshAssetFromMesh
- DeleteVerticesFromMesh, DeleteTrianglesFromMesh, AppendBuffersToMesh
- AppendSimpleExtrudePolygon, AppendTriangulatedPolygon, hole parameter for Append Disc (to make punctured disc)
- IsSameMeshAs, MeasureDistancesBetweenMeshes, IsIntersectingMesh
- SplitMeshByComponents, SplitMeshByMaterialIDs, GetSubMeshFromMesh, CopyMeshToMesh
- ApplyFlareWarpToMesh, changed other warps to use FTransform for orientation instead of multiple axis parameters
- GetHasMaterialIDs, GetMaxMaterialID, EnableMaterialIDs, ClearMaterialIDs, RemapMaterialIDs, GetTriangleMaterialID, GetAllTriangleMaterialIDs, SetTriangleMaterialID, SetAllTriangleMaterialIDs
- SetNumUVSets, TranslateMeshUVs, ScaleMeshUVs, RotateMeshUVs, SetMeshUVsFromPlanarProjection, SetMeshUVsFromBoxProjection, RepackMeshUVs
- GetNumConnectedComponents, GetAllVertexPositions
- aded FGeometryScriptDynamicMeshBVH wrapper for AABBTree & FWNTree. Somewhat risky but avoids a UObject wrapper, and these objects are (currently) intended to have a lifetime that does not extend outside a single BP
- BuildBVHForMesh, IsBVHValidForMesh, RebuildBVHForMesh, FindNearestPointOnMesh, FindNearestRayIntersectionWithMesh, IsPointInsideMesh

#rnx

diff --git a/Engine/Plugins/Experimental/GeometryScripting/GeometryScripting.uplugin b/Engine/Plugins/Experimental/GeometryScripting/GeometryScripting.uplugin
--- a/Engine/Plugins/Experimental/GeometryScripting/GeometryScripting.uplugin
+++ b/Engine/Plugins/Experimental/GeometryScripting/GeometryScripting.uplugin
@@ -18,6 +18,12 @@
 			"Name": "GeometryScriptingCore",
 			"Type": "Runtime",
 			"LoadingPhase": "Default"
+		},
+		{
+			"Name": "GeometryScriptingEditor",
+			"Type": "Editor",
+			"LoadingPhase": "Default",
+			"WhitelistPlatforms": [ "Win64" ]
 		}
 	],
 	"Plugins": [
diff --git a/Engine/Plugins/Experimental/GeometryScripting/Source/GeometryScriptingCore/Private/ContainmentFunctions.cpp b/Engine/Plugins/Experimental/GeometryScripting/Source/GeometryScriptingCore/Private/ContainmentFunctions.cpp
new file mode 100644
--- /dev/null
+++ b/Engine/Plugins/Experimental/GeometryScripting/Source/GeometryScriptingCore/Private/ContainmentFunctions.cpp
@@ -0,0 +1,124 @@
+// Copyright Epic Games, Inc. All Rights Reserved.
+
+#include "GeometryScript/ContainmentFunctions.h"
+
+#include "DynamicMesh/DynamicMesh3.h"
+#include "DynamicMesh/MeshNormals.h"
+#include "Operations/MeshConvexHull.h"
+#include "Operations/MeshProjectionHull.h"
+#include "UDynamicMesh.h"
+
+#include "ExplicitUseGeometryMathTypes.h"		// using UE::Geometry::(math types)
+using namespace UE::Geometry;
+
+#define LOCTEXT_NAMESPACE "UGeometryScriptLibrary_MeshRepairFunctions"
+
+
+UDynamicMesh* UGeometryScriptLibrary_ContainmentFunctions::ComputeMeshConvexHull(
+	UDynamicMesh* TargetMesh,
+	UPARAM(DisplayName = "Hull Mesh", ref) UDynamicMesh* CopyToMesh, 
+	UPARAM(DisplayName = "Hull Mesh") UDynamicMesh*& CopyToMeshOut, 
+	FGeometryScriptConvexHullOptions Options,
+	UGeometryScriptDebug* Debug)
+{
+	if (TargetMesh == nullptr)
+	{
+		UE::Geometry::AppendError(Debug, EGeometryScriptErrorType::InvalidInputs, LOCTEXT("ComputeMeshConvexHull_InvalidInput", "ComputeMeshConvexHull: TargetMesh is Null"));
+		return TargetMesh;
+	}
+	if (CopyToMesh == nullptr)
+	{
+		UE::Geometry::AppendError(Debug, EGeometryScriptErrorType::InvalidInputs, LOCTEXT("ComputeMeshConvexHull_InvalidInput2", "ComputeMeshConvexHull: CopyToMesh is Null"));
+		return TargetMesh;
+	}
+	CopyToMeshOut = CopyToMesh;
+
+	FDynamicMesh3 HullMesh;
+	TargetMesh->ProcessMesh([&](const FDynamicMesh3& EditMesh) 
+	{
+		FMeshConvexHull Hull(&EditMesh);
+
+		if (Options.bPrefilterVertices)
+		{
+			FMeshConvexHull::GridSample(EditMesh, FMath::Max(32, Options.PrefilterGridResolution), Hull.VertexSet);
+		}
+
+		Hull.bPostSimplify = (Options.SimplifyToFaceCount > 4);
+		Hull.MaxTargetFaceCount = Options.SimplifyToFaceCount;
+		if (Hull.Compute(nullptr))
+		{
+			HullMesh = MoveTemp(Hull.ConvexHull);
+			HullMesh.EnableAttributes();
+			FMeshNormals::InitializeOverlayToPerTriangleNormals(HullMesh.Attributes()->PrimaryNormals());
+		}
+	});
+
+	if ( HullMesh.TriangleCount() == 0 )
+	{
+		// todo: replace output with bounding box mesh?
+		UE::Geometry::AppendError(Debug, EGeometryScriptErrorType::OperationFailed, LOCTEXT("ComputeMeshConvexHull_Failed", "ComputeMeshConvexHull: Hull Computation Failed"));
+		CopyToMesh->ResetToCube();
+	}
+	else
+	{
+		CopyToMesh->SetMesh(MoveTemp(HullMesh));
+	}
+
+	return TargetMesh;
+}
+
+
+
+
+UDynamicMesh* UGeometryScriptLibrary_ContainmentFunctions::ComputeMeshSweptHull(
+	UDynamicMesh* TargetMesh,
+	UPARAM(DisplayName = "Hull Mesh", ref) UDynamicMesh* CopyToMesh, 
+	UPARAM(DisplayName = "Hull Mesh") UDynamicMesh*& CopyToMeshOut, 
+	FTransform ProjectionFrame,
+	FGeometryScriptSweptHullOptions Options,
+	UGeometryScriptDebug* Debug)
+{
+	if (TargetMesh == nullptr)
+	{
+		UE::Geometry::AppendError(Debug, EGeometryScriptErrorType::InvalidInputs, LOCTEXT("ComputeMeshSweptHull_InvalidInput", "ComputeMeshSweptHull: TargetMesh is Null"));
+		return TargetMesh;
+	}
+	if (CopyToMesh == nullptr)
+	{
+		UE::Geometry::AppendError(Debug, EGeometryScriptErrorType::InvalidInputs, LOCTEXT("ComputeMeshSweptHulll_InvalidInput2", "ComputeMeshSweptHull: CopyToMesh is Null"));
+		return TargetMesh;
+	}
+	CopyToMeshOut = CopyToMesh;
+
+	FDynamicMesh3 HullMesh;
+	TargetMesh->ProcessMesh([&](const FDynamicMesh3& EditMesh) 
+	{
+		FMeshProjectionHull Hull(&EditMesh);
+		Hull.ProjectionFrame = FFrame3d(ProjectionFrame);
+		Hull.MinThickness = FMathd::Max(Options.MinThickness, 0);
+		Hull.bSimplifyPolygon = Options.bSimplify;
+		Hull.MinEdgeLength = Options.MinEdgeLength;
+		Hull.DeviationTolerance = Options.SimplifyTolerance;
+
+		if (Hull.Compute())
+		{
+			HullMesh = MoveTemp(Hull.ConvexHull3D);
+		}
+	});
+
+	if ( HullMesh.TriangleCount() == 0 )
+	{
+		// todo: replace output with bounding box mesh?
+		UE::Geometry::AppendError(Debug, EGeometryScriptErrorType::OperationFailed, LOCTEXT("ComputeMeshSweptHull_Failed", "ComputeMeshSweptHull: Hull Computation Failed"));
+		CopyToMesh->ResetToCube();
+	}
+	else
+	{
+		CopyToMesh->SetMesh(MoveTemp(HullMesh));
+	}
+
+	return TargetMesh;
+}
+
+
+#undef LOCTEXT_NAMESPACE
\ No newline at end of file
diff --git a/Engine/Plugins/Experimental/GeometryScripting/Source/GeometryScriptingCore/Private/MeshBasicEditFunctions.cpp b/Engine/Plugins/Experimental/GeometryScripting/Source/GeometryScriptingCore/Private/MeshBasicEditFunctions.cpp
--- a/Engine/Plugins/Experimental/GeometryScripting/Source/GeometryScriptingCore/Private/MeshBasicEditFunctions.cpp
+++ b/Engine/Plugins/Experimental/GeometryScripting/Source/GeometryScriptingCore/Private/MeshBasicEditFunctions.cpp
@@ -112,6 +112,32 @@ UDynamicMesh* UGeometryScriptLibrary_MeshBasicEditFunctions::DeleteVertexFromMes
 }
 
 
+UDynamicMesh* UGeometryScriptLibrary_MeshBasicEditFunctions::DeleteVerticesFromMesh(
+	UDynamicMesh* TargetMesh,
+	const TArray<int32>& VertexList,
+	int& NumDeleted,
+	bool bDeferChangeNotifications)
+{
+	NumDeleted = 0;
+	if (TargetMesh)
+	{
+		TargetMesh->EditMesh([&](FDynamicMesh3& EditMesh)
+		{
+			for (int32 VertexID : VertexList)
+			{
+				EMeshResult Result = EditMesh.RemoveVertex(VertexID);
+				if (Result == EMeshResult::Ok)
+				{
+					NumDeleted++;
+				}
+			}
+		}, EDynamicMeshChangeType::GeneralEdit, EDynamicMeshAttributeChangeFlags::Unknown, bDeferChangeNotifications);
+	}
+	return TargetMesh;
+}
+
+
+
 
 
 UDynamicMesh* UGeometryScriptLibrary_MeshBasicEditFunctions::AddTriangleToMesh(
@@ -228,6 +254,30 @@ UDynamicMesh* UGeometryScriptLibrary_MeshBasicEditFunctions::DeleteTriangleFromM
 }
 
 
+UDynamicMesh* UGeometryScriptLibrary_MeshBasicEditFunctions::DeleteTrianglesFromMesh(
+	UDynamicMesh* TargetMesh,
+	const TArray<int32>& TriangleList,
+	int& NumDeleted,
+	bool bDeferChangeNotifications)
+{
+	NumDeleted = 0;
+	if (TargetMesh)
+	{
+		TargetMesh->EditMesh([&](FDynamicMesh3& EditMesh)
+		{
+			for (int32 TriangleID : TriangleList)
+			{
+				EMeshResult Result = EditMesh.RemoveTriangle(TriangleID);
+				if (Result == EMeshResult::Ok)
+				{
+					NumDeleted++;
+				}
+			}
+		}, EDynamicMeshChangeType::GeneralEdit, EDynamicMeshAttributeChangeFlags::Unknown, bDeferChangeNotifications);
+	}
+	return TargetMesh;
+}
+
 
 
 UDynamicMesh* UGeometryScriptLibrary_MeshBasicEditFunctions::AppendMesh(
@@ -255,7 +305,14 @@ UDynamicMesh* UGeometryScriptLibrary_MeshBasicEditFunctions::AppendMesh(
 			UE::Geometry::FTransform3d XForm(AppendTransform);
 			FMeshIndexMappings TmpMappings;
 			FDynamicMeshEditor Editor(&AppendToMesh);
-			Editor.AppendMesh(&OtherMesh, TmpMappings,
+			const FDynamicMesh3* UseOtherMesh = &OtherMesh;
+			FDynamicMesh3 TmpMesh;
+			if (UseOtherMesh == &AppendToMesh)
+			{
+				TmpMesh = OtherMesh;	// need  to make a copy if we are appending to ourself
+				UseOtherMesh = &TmpMesh;
+			}
+			Editor.AppendMesh(UseOtherMesh, TmpMappings,
 				[&](int, const FVector3d& Position) { return XForm.TransformPosition(Position); },
 				[&](int, const FVector3d& Normal) { return XForm.TransformNormal(Normal); });
 		});
@@ -318,4 +375,160 @@ UDynamicMesh* UGeometryScriptLibrary_MeshBasicEditFunctions::AppendMeshRepeated(
 
 
 
+UDynamicMesh* UGeometryScriptLibrary_MeshBasicEditFunctions::AppendBuffersToMesh(
+	UDynamicMesh* TargetMesh,
+	const FGeometryScriptSimpleMeshBuffers& Buffers,
+	TArray<int>& NewTriangleIndices,
+	int MaterialID,
+	bool bDeferChangeNotifications,
+	UGeometryScriptDebug* Debug)
+{
+	NewTriangleIndices.Reset();
+	if (TargetMesh == nullptr)
+	{
+		UE::Geometry::AppendError(Debug, EGeometryScriptErrorType::InvalidInputs, LOCTEXT("AppendMeshRepeated_InvalidInput1", "AppendMeshRepeated: TargetMesh is Null"));
+		return TargetMesh;
+	}
+
+	TargetMesh->EditMesh([&](FDynamicMesh3& EditMesh)
+	{
+		if (!EditMesh.HasAttributes())
+		{
+			EditMesh.EnableAttributes();
+		}
+		if (!EditMesh.HasTriangleGroups())
+		{
+			EditMesh.EnableTriangleGroups();
+		}
+
+		if (!EditMesh.Attributes()->HasMaterialID())
+		{
+			EditMesh.Attributes()->EnableMaterialID();
+		}
+		FDynamicMeshMaterialAttribute* MaterialIDs = EditMesh.Attributes()->GetMaterialID();
+
+		TArray<int32> VertexIDMap;
+		int32 NumVertices = Buffers.Vertices.Num();
+		VertexIDMap.SetNum(NumVertices);
+		for (int32 k = 0; k < NumVertices; ++k)
+		{
+			int32 NewVertexID = EditMesh.AppendVertex((FVector3d)Buffers.Vertices[k]);
+			VertexIDMap[k] = NewVertexID;
+		}
+
+		auto MapTriangleFunc = [&VertexIDMap](FIntVector Triangle)
+		{
+			return FIndex3i(VertexIDMap[Triangle.X], VertexIDMap[Triangle.Y], VertexIDMap[Triangle.Z]);
+		};
+
+		int32 NumTriangles = Buffers.Triangles.Num();
+		bool bHaveGroups = Buffers.TriGroupIDs.Num() == NumTriangles;
+		int32 ConstantGroupID = EditMesh.AllocateTriangleGroup();
+		for (int32 k = 0; k < NumTriangles; ++k)
+		{
+			int32 UseGroupID = bHaveGroups ? Buffers.TriGroupIDs[k] : ConstantGroupID;
+			int32 NewTriangleID = EditMesh.AppendTriangle(MapTriangleFunc(Buffers.Triangles[k]), UseGroupID);
+			if (NewTriangleID < 0)
+			{
+				if (NewTriangleID == FDynamicMesh3::NonManifoldID)
+				{
+					UE::Geometry::AppendError(Debug, EGeometryScriptErrorType::InvalidInputs, LOCTEXT("AppendBuffersToMesh_NonManifold", "AppendBuffersToMesh: Triangle cannot be added because it would create invalid Non-Manifold Mesh Topology"));
+				}
+				else if (NewTriangleID == FDynamicMesh3::DuplicateTriangleID)
+				{
+					UE::Geometry::AppendError(Debug, EGeometryScriptErrorType::InvalidInputs, LOCTEXT("AppendBuffersToMesh_Duplicate", "AppendBuffersToMesh: Triangle cannot be added because it is a duplicate of an existing Triangle"));
+				}
+				else
+				{
+					UE::Geometry::AppendError(Debug, EGeometryScriptErrorType::OperationFailed, LOCTEXT("AppendBuffersToMesh_Unknown", "AppendBuffersToMesh: adding Triangle Failed"));
+				}
+				NewTriangleID = INDEX_NONE;
+			}
+			else
+			{
+				MaterialIDs->SetValue(NewTriangleID, MaterialID);
+			}
+			NewTriangleIndices.Add(NewTriangleID);
+		}
+
+
+		if (Buffers.Normals.Num() == NumVertices)
+		{
+			FDynamicMeshNormalOverlay* Normals = EditMesh.Attributes()->PrimaryNormals();
+			VertexIDMap.Reset();
+			for (int32 k = 0; k < NumVertices; ++k)
+			{
+				int32 NewElementID = Normals->AppendElement((FVector3f)Buffers.Normals[k]);
+				VertexIDMap[k] = NewElementID;
+			}
+			for (int32 k = 0; k < NumTriangles; ++k)
+			{
+				if (NewTriangleIndices[k] >= 0)
+				{
+					Normals->SetTriangle(NewTriangleIndices[k], MapTriangleFunc(Buffers.Triangles[k]));
+				}
+			}
+		}
+
+		const TArray<FVector2D>* AllUVSets[8] = { &Buffers.UV0, &Buffers.UV1, &Buffers.UV2, &Buffers.UV3, &Buffers.UV4, &Buffers.UV5, &Buffers.UV6, &Buffers.UV7 };
+		int32 NumUVLayers = 0;
+		for (int32 k = 0; k < 8; ++k)
+		{
+			if (AllUVSets[k]->Num() != NumVertices)
+			{
+				break;
+			}
+			NumUVLayers++;
+		}
+		EditMesh.Attributes()->SetNumUVLayers(NumUVLayers);
+		for (int32 li = 0; li < NumUVLayers; ++li)
+		{
+			FDynamicMeshUVOverlay* UVs = EditMesh.Attributes()->GetUVLayer(li);
+			VertexIDMap.Reset();
+			for (int32 k = 0; k < NumVertices; ++k)
+			{
+				int32 NewElementID = UVs->AppendElement( (FVector2f) (*AllUVSets[li])[k] );
+				VertexIDMap[k] = NewElementID;
+			}
+			for (int32 k = 0; k < NumTriangles; ++k)
+			{
+				if (NewTriangleIndices[k] >= 0)
+				{
+					UVs->SetTriangle(NewTriangleIndices[k], MapTriangleFunc(Buffers.Triangles[k]));
+				}
+			}
+		}
+
+		if (Buffers.VertexColors.Num() == NumVertices)
+		{
+			EditMesh.Attributes()->EnablePrimaryColors();
+			FDynamicMeshColorOverlay* Colors = EditMesh.Attributes()->PrimaryColors();
+			VertexIDMap.Reset();
+			for (int32 k = 0; k < NumVertices; ++k)
+			{
+				int32 NewElementID = Colors->AppendElement( ToVector4<float>(Buffers.VertexColors[k]) );
+				VertexIDMap[k] = NewElementID;
+			}
+			for (int32 k = 0; k < NumTriangles; ++k)
+			{
+				if (NewTriangleIndices[k] >= 0)
+				{
+					Colors->SetTriangle(NewTriangleIndices[k], MapTriangleFunc(Buffers.Triangles[k]));
+				}
+			}
+		}
+
+
+	}, EDynamicMeshChangeType::GeneralEdit, EDynamicMeshAttributeChangeFlags::Unknown, bDeferChangeNotifications);
+
+
+
+
+	return TargetMesh;
+}
+
+
+
+
+
 #undef LOCTEXT_NAMESPACE
diff --git a/Engine/Plugins/Experimental/GeometryScripting/Source/GeometryScriptingCore/Private/MeshComparisonFunctions.cpp b/Engine/Plugins/Experimental/GeometryScripting/Source/GeometryScriptingCore/Private/MeshComparisonFunctions.cpp
new file mode 100644
--- /dev/null
+++ b/Engine/Plugins/Experimental/GeometryScripting/Source/GeometryScriptingCore/Private/MeshComparisonFunctions.cpp
@@ -0,0 +1,171 @@
+// Copyright Epic Games, Inc. All Rights Reserved.
+
+#include "GeometryScript/MeshComparisonFunctions.h"
+
+#include "DynamicMesh/DynamicMesh3.h"
+#include "DynamicMesh/MeshNormals.h"
+#include "DynamicMesh/DynamicMeshAABBTree3.h"
+#include "MeshQueries.h"
+#include "UDynamicMesh.h"
+#include "Async/Async.h"
+
+#include "ExplicitUseGeometryMathTypes.h"		// using UE::Geometry::(math types)
+using namespace UE::Geometry;
+
+#define LOCTEXT_NAMESPACE "UGeometryScriptLibrary_MeshComparisonFunctions"
+
+
+UDynamicMesh* UGeometryScriptLibrary_MeshComparisonFunctions::IsSameMeshAs(
+	UDynamicMesh* TargetMesh,
+	UDynamicMesh* OtherMesh,
+	FGeometryScriptIsSameMeshOptions Options,
+	bool &bIsSameMesh,
+	UGeometryScriptDebug* Debug)
+{
+	bIsSameMesh = false;
+
+	if (TargetMesh == nullptr)
+	{
+		UE::Geometry::AppendError(Debug, EGeometryScriptErrorType::InvalidInputs, LOCTEXT("IsSameMeshAs_InvalidInput", "IsSameMeshAs: TargetMesh is Null"));
+		return TargetMesh;
+	}
+	if (OtherMesh == nullptr)
+	{
+		UE::Geometry::AppendError(Debug, EGeometryScriptErrorType::InvalidInputs, LOCTEXT("IsSameMeshAs_InvalidInput2", "IsSameMeshAs: OtherMesh is Null"));
+		return TargetMesh;
+	}
+
+	TargetMesh->ProcessMesh([&](const FDynamicMesh3& ReadMesh) 
+	{
+		OtherMesh->ProcessMesh([&](const FDynamicMesh3& OtherMesh)
+		{
+			FDynamicMesh3::FSameAsOptions CompareOptions;
+			CompareOptions.bCheckConnectivity = Options.bCheckConnectivity;
+			CompareOptions.bCheckEdgeIDs = Options.bCheckEdgeIDs;
+			CompareOptions.bCheckNormals = Options.bCheckNormals;
+			CompareOptions.bCheckColors = Options.bCheckColors;
+			CompareOptions.bCheckUVs = Options.bCheckUVs;
+			CompareOptions.bCheckGroups = Options.bCheckGroups;
+			CompareOptions.bCheckAttributes = Options.bCheckAttributes;
+			CompareOptions.Epsilon = Options.Epsilon;
+			bIsSameMesh = ReadMesh.IsSameAs(OtherMesh, CompareOptions);
+		});
+	});
+
+	return TargetMesh;
+}
+
+
+
+UDynamicMesh* UGeometryScriptLibrary_MeshComparisonFunctions::IsIntersectingMesh(
+	UDynamicMesh* TargetMesh,
+	FTransform TargetTransformIn,
+	UDynamicMesh* OtherMesh,
+	FTransform OtherTransformIn,
+	bool &bIsIntersecting,
+	UGeometryScriptDebug* Debug)
+{
+	bIsIntersecting = false;
+
+	if (TargetMesh == nullptr)
+	{
+		UE::Geometry::AppendError(Debug, EGeometryScriptErrorType::InvalidInputs, LOCTEXT("IsIntersectingMesh_InvalidInput", "IsIntersectingMesh: TargetMesh is Null"));
+		return TargetMesh;
+	}
+	if (OtherMesh == nullptr)
+	{
+		UE::Geometry::AppendError(Debug, EGeometryScriptErrorType::InvalidInputs, LOCTEXT("IsIntersectingMesh_InvalidInput2", "IsIntersectingMesh: OtherMesh is Null"));
+		return TargetMesh;
+	}
+
+	TargetMesh->ProcessMesh([&](const FDynamicMesh3& Mesh1) 
+	{
+		OtherMesh->ProcessMesh([&](const FDynamicMesh3& Mesh2)
+		{
+			FDynamicMeshAABBTree3 Spatials[2];
+			ParallelFor(2, [&](int32 k) { Spatials[k].SetMesh((k == 0) ? &Mesh1 : &Mesh2, true); }, EParallelForFlags::Unbalanced);
+
+			bool bIsIdentity1 = TargetTransformIn.Equals(FTransform::Identity, 0);
+			UE::Geometry::FTransform3d Transform1(TargetTransformIn);
+			bool bIsIdentity2 = OtherTransformIn.Equals(FTransform::Identity, 0);
+			UE::Geometry::FTransform3d Transform2(OtherTransformIn);
+			if (bIsIdentity1 && bIsIdentity2)
+			{
+				bIsIntersecting = Spatials[0].TestIntersection(Spatials[1]);
+			} 
+			else if (bIsIdentity1 || bIsIdentity2)
+			{
+				FIndex2i Indices = (bIsIdentity1) ? FIndex2i(0,1) : FIndex2i(1,0);
+				UE::Geometry::FTransform3d UseTransform = (bIsIdentity2) ? Transform1 : Transform2;
+				bIsIntersecting = Spatials[Indices.A].TestIntersection(Spatials[Indices.B],
+					[&](const FVector3d& Pos) { return UseTransform.TransformPosition(Pos); });
+			}
+			else
+			{
+				Transform1 = Transform1.Inverse();
+				bIsIntersecting = Spatials[0].TestIntersection(Spatials[1],
+					[&](const FVector3d& Pos) { return Transform1.TransformPosition(Transform2.TransformPosition(Pos)); });
+			}
+		});
+	});
+
+	return TargetMesh;
+}
+
+
+
+
+
+
+UDynamicMesh* UGeometryScriptLibrary_MeshComparisonFunctions::MeasureDistancesBetweenMeshes(
+	UDynamicMesh* TargetMesh,
+	UDynamicMesh* OtherMesh,
+	FGeometryScriptMeasureMeshDistanceOptions Options,
+	double& MaxDistance,
+	double& MinDistance,
+	double& AverageDistance,
+	double& RootMeanSqrDeviation,
+	UGeometryScriptDebug* Debug)
+{
+	MaxDistance = TNumericLimits<float>::Max();
+
+	if (TargetMesh == nullptr)
+	{
+		UE::Geometry::AppendError(Debug, EGeometryScriptErrorType::InvalidInputs, LOCTEXT("MeasureDistancesBetweenMeshes_InvalidInput", "MeasureDistancesBetweenMeshes: TargetMesh is Null"));
+		return TargetMesh;
+	}
+	if (OtherMesh == nullptr)
+	{
+		UE::Geometry::AppendError(Debug, EGeometryScriptErrorType::InvalidInputs, LOCTEXT("MeasureDistancesBetweenMeshes_InvalidInput2", "MeasureDistancesBetweenMeshes: OtherMesh is Null"));
+		return TargetMesh;
+	}
+
+	TargetMesh->ProcessMesh([&](const FDynamicMesh3& Mesh1) 
+	{
+		OtherMesh->ProcessMesh([&](const FDynamicMesh3& Mesh2)
+		{
+			if (Options.bSymmetric)
+			{
+				FDynamicMeshAABBTree3 Spatials[2];
+				ParallelFor(2, [&](int32 k) { Spatials[k].SetMesh((k == 0) ? &Mesh1 : &Mesh2, true); }, EParallelForFlags::Unbalanced);
+				TMeshQueries<FDynamicMesh3>::MeshDistanceStatistics(
+					Mesh1, Spatials[1], 
+					&Mesh2, &Spatials[0], true,
+					MaxDistance, MinDistance, AverageDistance, RootMeanSqrDeviation);
+			}
+			else
+			{
+				FDynamicMeshAABBTree3 Spatial2(&Mesh2, true);
+				TMeshQueries<FDynamicMesh3>::MeshDistanceStatistics(
+					Mesh1, Spatial2, 
+					(const FDynamicMesh3*)nullptr, (const FDynamicMeshAABBTree3*)nullptr, false,
+					MaxDistance, MinDistance, AverageDistance, RootMeanSqrDeviation);
+			}
+		});
+	});
+
+	return TargetMesh;
+}
+
+
+#undef LOCTEXT_NAMESPACE
\ No newline at end of file
diff --git a/Engine/Plugins/Experimental/GeometryScripting/Source/GeometryScriptingCore/Private/MeshDecompositionFunctions.cpp b/Engine/Plugins/Experimental/GeometryScripting/Source/GeometryScriptingCore/Private/MeshDecompositionFunctions.cpp
new file mode 100644
--- /dev/null
+++ b/Engine/Plugins/Experimental/GeometryScripting/Source/GeometryScriptingCore/Private/MeshDecompositionFunctions.cpp
@@ -0,0 +1,210 @@
+// Copyright Epic Games, Inc. All Rights Reserved.
+
+#include "GeometryScript/MeshDecompositionFunctions.h"
+
+#include "DynamicMesh/DynamicMesh3.h"
+#include "DynamicMesh/DynamicMeshAttributeSet.h"
+#include "Selections/MeshConnectedComponents.h"
+#include "DynamicMeshEditor.h"
+#include "UDynamicMesh.h"
+
+#include "ExplicitUseGeometryMathTypes.h"		// using UE::Geometry::(math types)
+using namespace UE::Geometry;
+
+#define LOCTEXT_NAMESPACE "UGeometryScriptLibrary_MeshDecompositionFunctions"
+
+
+void BuildNewDynamicMeshes(UDynamicMesh* TargetMesh, UDynamicMeshPool* MeshPool, TArray<FDynamicMesh3>& SplitMeshes, TArray<UDynamicMesh*>& ComponentMeshes)
+{
+	ComponentMeshes.SetNum(0);
+
+	int32 NumSplitMeshes = SplitMeshes.Num();
+	if (NumSplitMeshes == 0)
+	{
+		UDynamicMesh* ComponentMesh = (MeshPool != nullptr) ?
+			MeshPool->RequestMesh() : NewObject<UDynamicMesh>();
+		TargetMesh->ProcessMesh([&](const FDynamicMesh3& EditMesh)
+		{
+			ComponentMesh->SetMesh(EditMesh);
+		});
+		ComponentMeshes.Add(ComponentMesh);
+	}
+	else
+	{
+		for (int32 mi = 0; mi < NumSplitMeshes; ++mi)
+		{
+			UDynamicMesh* ComponentMesh = (MeshPool != nullptr) ?
+				MeshPool->RequestMesh() : NewObject<UDynamicMesh>();
+			ComponentMesh->SetMesh(MoveTemp(SplitMeshes[mi]));
+			ComponentMeshes.Add(ComponentMesh);
+		}
+	}
+}
+
+
+
+
+UDynamicMesh* UGeometryScriptLibrary_MeshDecompositionFunctions::SplitMeshByComponents(  
+	UDynamicMesh* TargetMesh, 
+	TArray<UDynamicMesh*>& ComponentMeshes,
+	UDynamicMeshPool* MeshPool,
+	UGeometryScriptDebug* Debug)
+{
+	if (TargetMesh == nullptr)
+	{
+		UE::Geometry::AppendError(Debug, EGeometryScriptErrorType::InvalidInputs, LOCTEXT("SplitMeshByComponents_InvalidInput", "SplitMeshByComponents: TargetMesh is Null"));
+		return TargetMesh;
+	}
+
+	TArray<FDynamicMesh3> SplitMeshes;
+
+	TargetMesh->ProcessMesh([&](const FDynamicMesh3& EditMesh)
+	{
+		FMeshConnectedComponents Components(&EditMesh);
+		Components.FindConnectedTriangles();
+		int32 NumComponents = Components.Num();
+		if (NumComponents <= 1)
+		{
+			return;		// for single-component case, BuildNewDynamicMeshes() will just copy TargetMesh
+		}
+		
+		TArray<int32> TriSubmeshIndices;
+		TriSubmeshIndices.SetNum(EditMesh.MaxTriangleID());
+		for (int32 ci = 0; ci < NumComponents; ++ci)
+		{
+			for (int32 tid : Components.GetComponent(ci).Indices)
+			{
+				TriSubmeshIndices[tid] = ci;
+			}
+		}
+
+		FDynamicMeshEditor::SplitMesh(&EditMesh, SplitMeshes, [&](int32 tid) { return TriSubmeshIndices[tid]; });
+	});
+
+	BuildNewDynamicMeshes(TargetMesh, MeshPool, SplitMeshes, ComponentMeshes);
+
+	return TargetMesh;
+}
+
+
+UDynamicMesh* UGeometryScriptLibrary_MeshDecompositionFunctions::SplitMeshByMaterialIDs(  
+	UDynamicMesh* TargetMesh, 
+	TArray<UDynamicMesh*>& ComponentMeshes,
+	TArray<int>& ComponentMaterialIDs,
+	UDynamicMeshPool* MeshPool,
+	UGeometryScriptDebug* Debug)
+{
+	ComponentMeshes.Reset();
+	ComponentMaterialIDs.Reset();
+
+	if (TargetMesh == nullptr)
+	{
+		UE::Geometry::AppendError(Debug, EGeometryScriptErrorType::InvalidInputs, LOCTEXT("SplitMeshByComponents_InvalidInput", "SplitMeshByComponents: TargetMesh is Null"));
+		return TargetMesh;
+	}
+
+	TArray<FDynamicMesh3> SplitMeshes;
+	TargetMesh->ProcessMesh([&](const FDynamicMesh3& EditMesh)
+	{
+		const FDynamicMeshMaterialAttribute* MaterialIDs =
+			(EditMesh.HasAttributes() && EditMesh.Attributes()->HasMaterialID()) ? EditMesh.Attributes()->GetMaterialID() : nullptr;
+		if (MaterialIDs == nullptr)
+		{
+			ComponentMaterialIDs.Add(0);
+			return;
+		}
+
+		for (int32 tid : EditMesh.TriangleIndicesItr())
+		{
+			int32 MaterialID = MaterialIDs->GetValue(tid);
+			ComponentMaterialIDs.AddUnique(MaterialID);
+		}
+		ComponentMaterialIDs.Sort();
+
+		FDynamicMeshEditor::SplitMesh(&EditMesh, SplitMeshes, [&](int32 tid) {
+			int32 MaterialID = MaterialIDs->GetValue(tid);
+			int32 Index = ComponentMaterialIDs.IndexOfByKey(MaterialID);
+			return (Index == INDEX_NONE) ? 0 : Index;
+		});
+
+	});
+
+	BuildNewDynamicMeshes(TargetMesh, MeshPool, SplitMeshes, ComponentMeshes);
+
+	return TargetMesh;
+}
+
+
+UDynamicMesh* UGeometryScriptLibrary_MeshDecompositionFunctions::GetSubMeshFromMesh(
+	UDynamicMesh* TargetMesh,
+	UDynamicMesh* StoreToSubmesh, 
+	const TArray<int>& TriangleList,
+	UDynamicMesh*& StoreToSubmeshOut, 
+	UGeometryScriptDebug* Debug)
+{
+	if (TargetMesh == nullptr)
+	{
+		UE::Geometry::AppendError(Debug, EGeometryScriptErrorType::InvalidInputs, LOCTEXT("GetSubMeshFromMesh_InvalidInput", "GetSubMeshFromMesh: TargetMesh is Null"));
+		return TargetMesh;
+	}
+	if (StoreToSubmesh == nullptr)
+	{
+		UE::Geometry::AppendError(Debug, EGeometryScriptErrorType::InvalidInputs, LOCTEXT("GetSubMeshFromMesh_InvalidInput2", "GetSubMeshFromMesh: Submesh is Null"));
+		return TargetMesh;
+	}
+
+	FDynamicMesh3 Submesh;
+	TargetMesh->ProcessMesh([&](const FDynamicMesh3& EditMesh)
+	{
+		if (EditMesh.HasAttributes())
+		{
+			Submesh.EnableAttributes();
+			Submesh.Attributes()->EnableMatchingAttributes(*EditMesh.Attributes());
+		}
+
+		FMeshIndexMappings Mappings;
+		FDynamicMeshEditResult EditResult;
+		FDynamicMeshEditor Editor(&Submesh);
+		Editor.AppendTriangles(&EditMesh, TriangleList, Mappings, EditResult);
+	});
+
+	StoreToSubmesh->SetMesh(MoveTemp(Submesh));
+	StoreToSubmeshOut = StoreToSubmesh;
+
+	return TargetMesh;
+}
+
+
+UDynamicMesh* UGeometryScriptLibrary_MeshDecompositionFunctions::CopyMeshToMesh(
+	UDynamicMesh* CopyFromMesh,
+	UDynamicMesh* CopyToMesh, 
+	UDynamicMesh*& CopyToMeshOut, 
+	UGeometryScriptDebug* Debug)
+{
+	if (CopyFromMesh == nullptr)
+	{
+		UE::Geometry::AppendError(Debug, EGeometryScriptErrorType::InvalidInputs, LOCTEXT("CopyMeshToMesh_InvalidInput", "CopyMeshToMesh: TargetMesh is Null"));
+		return CopyFromMesh;
+	}
+	if (CopyToMesh == nullptr)
+	{
+		UE::Geometry::AppendError(Debug, EGeometryScriptErrorType::InvalidInputs, LOCTEXT("CopyMeshToMesh_InvalidInput2", "CopyMeshToMesh: Submesh is Null"));
+		return CopyFromMesh;
+	}
+
+	FDynamicMesh3 MeshCopy;
+	CopyFromMesh->ProcessMesh([&](const FDynamicMesh3& EditMesh)
+	{
+		MeshCopy = EditMesh;
+	});
+
+	CopyToMesh->SetMesh(MoveTemp(MeshCopy));
+	CopyToMeshOut = CopyToMesh;
+
+	return CopyFromMesh;
+}
+
+
+
+
+#undef LOCTEXT_NAMESPACE
\ No newline at end of file
diff --git a/Engine/Plugins/Experimental/GeometryScripting/Source/GeometryScriptingCore/Private/MeshDeformFunctions.cpp b/Engine/Plugins/Experimental/GeometryScripting/Source/GeometryScriptingCore/Private/MeshDeformFunctions.cpp
--- a/Engine/Plugins/Experimental/GeometryScripting/Source/GeometryScriptingCore/Private/MeshDeformFunctions.cpp
+++ b/Engine/Plugins/Experimental/GeometryScripting/Source/GeometryScriptingCore/Private/MeshDeformFunctions.cpp
@@ -6,6 +6,7 @@
 #include "DynamicMesh/DynamicMeshAttributeSet.h"
 #include "SpaceDeformerOps/BendMeshOp.h"
 #include "SpaceDeformerOps/TwistMeshOp.h"
+#include "SpaceDeformerOps/FlareMeshOp.h"
 #include "UDynamicMesh.h"
 
 #include "ExplicitUseGeometryMathTypes.h"		// using UE::Geometry::(math types)
@@ -17,6 +18,7 @@ using namespace UE::Geometry;
 UDynamicMesh* UGeometryScriptLibrary_MeshDeformFunctions::ApplyBendWarpToMesh(
 	UDynamicMesh* TargetMesh,
 	FGeometryScriptBendWarpOptions Options,
+	FTransform BendOrientation,
 	float BendAngle,
 	float BendExtent,
 	UGeometryScriptDebug* Debug)
@@ -27,27 +29,7 @@ UDynamicMesh* UGeometryScriptLibrary_MeshDeformFunctions::ApplyBendWarpToMesh(
 		return TargetMesh;
 	}
 
-	FVector3d UseAxis = Options.BendAxis.GetSafeNormal();
-	if (FMath::Abs(1.0 - UseAxis.Length()) > 0.1)
-	{
-		UE::Geometry::AppendError(Debug, EGeometryScriptErrorType::InvalidInputs, LOCTEXT("ApplyBendWarpToMesh_InvalidAxis", "ApplyBendWarpToMesh: BendAxis is degenerate"));
-		UseAxis = FVector3d::UnitZ();
-	}
-	FFrame3d WarpFrame = FFrame3d((FVector3d)Options.BendOrigin, (FVector3d)UseAxis);
-
-	if (Options.TowardAxis.Length() > 0)
-	{
-		FVector3d UseTowardAxis = (FVector3d)Options.TowardAxis.GetSafeNormal();
-		if ( WarpFrame.Z().Dot(UseTowardAxis) < 0.95 )
-		{
-			WarpFrame.ConstrainedAlignAxis(1, UseTowardAxis, WarpFrame.Z());
-		}
-	}
-
-	if (Options.BendAxisRotation != 0)
-	{
-		WarpFrame.Rotate(FQuaterniond( WarpFrame.Z(), Options.BendAxisRotation, true));
-	}
+	FFrame3d WarpFrame = FFrame3d(BendOrientation);
 
 	TargetMesh->EditMesh([&](FDynamicMesh3& EditMesh) 
 	{
@@ -82,6 +64,7 @@ UDynamicMesh* UGeometryScriptLibrary_MeshDeformFunctions::ApplyBendWarpToMesh(
 UDynamicMesh* UGeometryScriptLibrary_MeshDeformFunctions::ApplyTwistWarpToMesh(
 	UDynamicMesh* TargetMesh,
 	FGeometryScriptTwistWarpOptions Options,
+	FTransform TwistOrientation,
 	float TwistAngle,
 	float TwistExtent,
 	UGeometryScriptDebug* Debug)
@@ -92,27 +75,7 @@ UDynamicMesh* UGeometryScriptLibrary_MeshDeformFunctions::ApplyTwistWarpToMesh(
 		return TargetMesh;
 	}
 
-	FVector3d UseAxis = Options.TwistAxis.GetSafeNormal();
-	if (FMath::Abs(1.0 - UseAxis.Length()) > 0.1)
-	{
-		UE::Geometry::AppendError(Debug, EGeometryScriptErrorType::InvalidInputs, LOCTEXT("ApplyTwistWarpToMesh_InvalidAxis", "ApplyTwistWarpToMesh: TwistAxis is degenerate"));
-		UseAxis = FVector3d::UnitZ();
-	}
-	FFrame3d WarpFrame = FFrame3d((FVector3d)Options.TwistOrigin, UseAxis);
-
-	if (Options.TowardAxis.Length() > 0)
-	{
-		FVector3d UseTowardAxis = (FVector3d)Options.TowardAxis.GetSafeNormal();
-		if ( WarpFrame.Z().Dot(UseTowardAxis) < 0.95 )
-		{
-			WarpFrame.ConstrainedAlignAxis(1, UseTowardAxis, WarpFrame.Z());
-		}
-	}
-
-	if (Options.TwistAxisRotation != 0)
-	{
-		WarpFrame.Rotate(FQuaterniond( WarpFrame.Z(), Options.TwistAxisRotation, true));
-	}
+	FFrame3d WarpFrame = FFrame3d(TwistOrientation);
 
 	TargetMesh->EditMesh([&](FDynamicMesh3& EditMesh) 
 	{
@@ -142,4 +105,51 @@ UDynamicMesh* UGeometryScriptLibrary_MeshDeformFunctions::ApplyTwistWarpToMesh(
 
 
 
+UDynamicMesh* UGeometryScriptLibrary_MeshDeformFunctions::ApplyFlareWarpToMesh(
+	UDynamicMesh* TargetMesh,
+	FGeometryScriptFlareWarpOptions Options,
+	FTransform FlareOrientation,
+	float FlarePercentX,
+	float FlarePercentY,
+	float FlareExtent,
+	UGeometryScriptDebug* Debug)
+{
+	if (TargetMesh == nullptr)
+	{
+		UE::Geometry::AppendError(Debug, EGeometryScriptErrorType::InvalidInputs, LOCTEXT("ApplyFlareWarpToMesh_InvalidInput", "ApplyFlareWarpToMesh: TargetMesh is Null"));
+		return TargetMesh;
+	}
+
+	FFrame3d WarpFrame = FFrame3d(FlareOrientation);
+
+	TargetMesh->EditMesh([&](FDynamicMesh3& EditMesh) 
+	{
+		// todo extract Flare warp into standalone math code
+		TSharedPtr<FDynamicMesh3> TmpMeshPtr = MakeShared<FDynamicMesh3>();
+		*TmpMeshPtr = MoveTemp(EditMesh);
+
+		FFlareMeshOp FlareOp;
+		FlareOp.OriginalMesh = TmpMeshPtr;
+		FlareOp.GizmoFrame = WarpFrame;
+		FlareOp.LowerBoundsInterval = (Options.bSymmetricExtents) ? -FlareExtent : -Options.LowerExtent;
+		FlareOp.UpperBoundsInterval = FlareExtent;
+		FlareOp.FlarePercentX = FlarePercentX;
+		FlareOp.FlarePercentY = FlarePercentY;
+		FlareOp.bSmoothEnds = Options.bSmoothEnds;
+
+		FlareOp.CalculateResult(nullptr);
+
+		TUniquePtr<FDynamicMesh3> NewResultMesh = FlareOp.ExtractResult();
+		FDynamicMesh3* NewResultMeshPtr = NewResultMesh.Release();
+		EditMesh = MoveTemp(*NewResultMeshPtr);
+
+	}, EDynamicMeshChangeType::GeneralEdit, EDynamicMeshAttributeChangeFlags::Unknown, false);
+
+	return TargetMesh;
+}
+
+
+
+
+
 #undef LOCTEXT_NAMESPACE
\ No newline at end of file
diff --git a/Engine/Plugins/Experimental/GeometryScripting/Source/GeometryScriptingCore/Private/MeshMaterialFunctions.cpp b/Engine/Plugins/Experimental/GeometryScripting/Source/GeometryScriptingCore/Private/MeshMaterialFunctions.cpp
new file mode 100644
--- /dev/null
+++ b/Engine/Plugins/Experimental/GeometryScripting/Source/GeometryScriptingCore/Private/MeshMaterialFunctions.cpp
@@ -0,0 +1,273 @@
+// Copyright Epic Games, Inc. All Rights Reserved.
+
+#include "GeometryScript/MeshMaterialFunctions.h"
+
+#include "DynamicMesh/DynamicMesh3.h"
+#include "DynamicMesh/DynamicMeshAttributeSet.h"
+#include "UDynamicMesh.h"
+
+#include "ExplicitUseGeometryMathTypes.h"		// using UE::Geometry::(math types)
+using namespace UE::Geometry;
+
+#define LOCTEXT_NAMESPACE "UGeometryScriptLibrary_MeshMaterialFunctions"
+
+
+template<typename ReturnType> 
+ReturnType SimpleMeshMaterialQuery(UDynamicMesh* Mesh, bool& bHasMaterials, ReturnType DefaultValue, 
+	TFunctionRef<ReturnType(const FDynamicMesh3& Mesh, const FDynamicMeshMaterialAttribute& MaterialIDs)> QueryFunc)
+{
+	bHasMaterials = false;
+	ReturnType RetVal = DefaultValue;
+	if (Mesh)
+	{
+		Mesh->ProcessMesh([&](const FDynamicMesh3& ReadMesh)
+		{
+			if (ReadMesh.HasAttributes() && ReadMesh.Attributes()->HasMaterialID() )
+			{
+				const FDynamicMeshMaterialAttribute* MaterialIDs = ReadMesh.Attributes()->GetMaterialID();
+				if (MaterialIDs != nullptr)
+				{
+					bHasMaterials = true;
+					RetVal = QueryFunc(ReadMesh, *MaterialIDs);
+				}
+			}
+		});
+	}
+	return RetVal;
+}
+
+
+void SimpleMeshMaterialEdit(UDynamicMesh* Mesh, bool bEnableIfMissing, bool& bHasMaterialIDs,
+	TFunctionRef<void(FDynamicMesh3& Mesh, FDynamicMeshMaterialAttribute& MaterialIDs)> EditFunc)
+{
+	bHasMaterialIDs = false;
+	if (Mesh)
+	{
+		Mesh->EditMesh([&](FDynamicMesh3& EditMesh)
+		{
+			if (EditMesh.HasAttributes() == false)
+			{
+				if (bEnableIfMissing)
+				{
+					EditMesh.EnableAttributes();
+				}
+				else
+				{
+					return;
+				}
+			}
+			if (EditMesh.Attributes()->HasMaterialID() == false)
+			{
+				if (bEnableIfMissing)
+				{
+					EditMesh.Attributes()->EnableMaterialID();
+				}
+				else
+				{
+					return;
+				}
+			}
+			FDynamicMeshMaterialAttribute* MaterialIDs = EditMesh.Attributes()->GetMaterialID();
+			if (ensure(MaterialIDs != nullptr))
+			{
+				bHasMaterialIDs = true;
+				EditFunc(EditMesh, *MaterialIDs);
+			}
+		}, EDynamicMeshChangeType::GeneralEdit, EDynamicMeshAttributeChangeFlags::Unknown, false);
+	}
+}
+
+
+
+
+UDynamicMesh* UGeometryScriptLibrary_MeshMaterialFunctions::EnableMaterialIDs(
+	UDynamicMesh* TargetMesh,
+	UGeometryScriptDebug* Debug)
+{
+	if (TargetMesh == nullptr)
+	{
+		UE::Geometry::AppendError(Debug, EGeometryScriptErrorType::InvalidInputs, LOCTEXT("EnableMaterialIDs_InvalidInput", "EnableMaterialIDs: TargetMesh is Null"));
+		return TargetMesh;
+	}
+
+	bool bHasMaterialIDs;
+	SimpleMeshMaterialEdit(TargetMesh, true, bHasMaterialIDs, [](FDynamicMesh3& Mesh, FDynamicMeshMaterialAttribute& MaterialIDs) {});
+
+	return TargetMesh;
+
+}
+
+
+
+UDynamicMesh* UGeometryScriptLibrary_MeshMaterialFunctions::ClearMaterialIDs(
+	UDynamicMesh* TargetMesh,
+	UGeometryScriptDebug* Debug)
+{
+	if (TargetMesh == nullptr)
+	{
+		UE::Geometry::AppendError(Debug, EGeometryScriptErrorType::InvalidInputs, LOCTEXT("ClearMaterialIDs_InvalidInput", "ClearMaterialIDs: TargetMesh is Null"));
+		return TargetMesh;
+	}
+
+	bool bHasMaterialIDs;
+	SimpleMeshMaterialEdit(TargetMesh, true, bHasMaterialIDs, 
+		[](FDynamicMesh3& Mesh, FDynamicMeshMaterialAttribute& MaterialIDs) 
+	{
+		for (int32 TriangleID : Mesh.TriangleIndicesItr())
+		{
+			MaterialIDs.SetValue(TriangleID, 0);
+		}
+	});
+
+	return TargetMesh;
+}
+
+
+
+UDynamicMesh* UGeometryScriptLibrary_MeshMaterialFunctions::RemapMaterialIDs(
+	UDynamicMesh* TargetMesh,
+	int FromMaterialID,
+	int ToMaterialID,
+	UGeometryScriptDebug* Debug)
+{
+	if (TargetMesh == nullptr)
+	{
+		UE::Geometry::AppendError(Debug, EGeometryScriptErrorType::InvalidInputs, LOCTEXT("RemapMaterialIDs_InvalidInput", "RemapMaterialIDs: TargetMesh is Null"));
+		return TargetMesh;
+	}
+
+	bool bHasMaterialIDs;
+	SimpleMeshMaterialEdit(TargetMesh, true, bHasMaterialIDs, 
+		[&](FDynamicMesh3& Mesh, FDynamicMeshMaterialAttribute& MaterialIDs) 
+	{
+		for (int32 TriangleID : Mesh.TriangleIndicesItr())
+		{
+			int32 CurID = MaterialIDs.GetValue(TriangleID);
+			if (CurID == FromMaterialID)
+			{
+				MaterialIDs.SetValue(TriangleID, ToMaterialID);
+			}
+		}
+	});
+
+	return TargetMesh;
+}
+
+
+int UGeometryScriptLibrary_MeshMaterialFunctions::GetMaxMaterialID( UDynamicMesh* TargetMesh, bool& bHasMaterialIDs )
+{
+	return SimpleMeshMaterialQuery<int32>(TargetMesh, bHasMaterialIDs, 0, [&](const FDynamicMesh3& Mesh, const FDynamicMeshMaterialAttribute& MaterialIDs) {
+		int32 MaxID = 0;
+		for (int TriangleID : Mesh.TriangleIndicesItr())
+		{
+			MaxID = FMath::Max(MaxID, MaterialIDs.GetValue(TriangleID));
+		}
+		return MaxID;
+	});
+}
+
+
+
+int32 UGeometryScriptLibrary_MeshMaterialFunctions::GetTriangleMaterialID( 
+	UDynamicMesh* TargetMesh, 
+	int TriangleID, 
+	bool& bIsValidTriangle)
+{
+	bool bHasMaterials = false;
+	return SimpleMeshMaterialQuery<int32>(TargetMesh, bHasMaterials, 0, [&](const FDynamicMesh3& Mesh, const FDynamicMeshMaterialAttribute& MaterialIDs) {
+		bIsValidTriangle = Mesh.IsTriangle(TriangleID);
+		return (bIsValidTriangle) ? MaterialIDs.GetValue(TriangleID) : 0;
+	});
+}
+
+
+UDynamicMesh* UGeometryScriptLibrary_MeshMaterialFunctions::GetAllTriangleMaterialIDs(UDynamicMesh* TargetMesh, UPARAM(ref) TArray<int>& MaterialIDs, bool& bHasMaterialIDs)
+{
+	MaterialIDs.Reset();
+	bHasMaterialIDs = false;
+	SimpleMeshMaterialQuery<int32>(TargetMesh, bHasMaterialIDs, 0, [&](const FDynamicMesh3& Mesh, const FDynamicMeshMaterialAttribute& MaterialIDAttrib) {
+		int32 MaxTriangleID = Mesh.MaxTriangleID();
+		for (int32 TriangleID = 0; TriangleID < Mesh.MaxTriangleID(); ++TriangleID)
+		{
+			int32 MaterialID = Mesh.IsTriangle(TriangleID) ? MaterialIDAttrib.GetValue(TriangleID) : -1;
+			MaterialIDs.Add(MaterialID);
+		}
+		return 0;
+	});
+	return TargetMesh;
+}
+
+
+
+UDynamicMesh* UGeometryScriptLibrary_MeshMaterialFunctions::SetTriangleMaterialID( 
+	UDynamicMesh* TargetMesh, 
+	int TriangleID, 
+	int MaterialID,
+	bool& bIsValidTriangle,
+	bool bDeferChangeNotifications)
+{
+	bIsValidTriangle = false;
+	if (TargetMesh)
+	{
+		TargetMesh->EditMesh([&](FDynamicMesh3& EditMesh)
+		{
+			if (EditMesh.IsTriangle(TriangleID) && EditMesh.HasAttributes() && EditMesh.Attributes()->HasMaterialID() )
+			{
+				FDynamicMeshMaterialAttribute* MaterialIDs = EditMesh.Attributes()->GetMaterialID();
+				if (MaterialIDs != nullptr)
+				{
+					bIsValidTriangle = true;
+					MaterialIDs->SetValue(TriangleID, MaterialID);
+				}
+			}
+		}, EDynamicMeshChangeType::GeneralEdit, EDynamicMeshAttributeChangeFlags::Unknown, bDeferChangeNotifications);
+	}
+	return TargetMesh;	
+}
+
+
+
+UDynamicMesh* UGeometryScriptLibrary_MeshMaterialFunctions::SetAllTriangleMaterialIDs(
+	UDynamicMesh* TargetMesh,
+	const TArray<int32>& TriangleMaterialIDs,
+	bool bDeferChangeNotifications,
+	UGeometryScriptDebug* Debug)
+{
+	if (TargetMesh == nullptr)
+	{
+		UE::Geometry::AppendError(Debug, EGeometryScriptErrorType::InvalidInputs, LOCTEXT("SetAllTriangleMaterialIDs_InvalidInput", "SetAllTriangleMaterialIDs: TargetMesh is Null"));
+		return TargetMesh;
+	}
+
+	TargetMesh->EditMesh([&](FDynamicMesh3& EditMesh) 
+	{
+		if (TriangleMaterialIDs.Num() < EditMesh.MaxTriangleID())
+		{
+			UE::Geometry::AppendError(Debug, EGeometryScriptErrorType::InvalidInputs, LOCTEXT("SetAllTriangleMaterialIDs_IncorrectCount", "SetAllTriangleMaterialIDs: size of provided TriangleMaterialIDs is smaller than MaxTriangleID of Mesh"));
+		}
+		else
+		{
+			if (EditMesh.HasAttributes() == false)
+			{
+				EditMesh.EnableAttributes();
+			}
+			if (EditMesh.Attributes()->HasMaterialID() == false)
+			{
+				EditMesh.Attributes()->EnableMaterialID();
+			}
+			FDynamicMeshMaterialAttribute* MaterialIDs = EditMesh.Attributes()->GetMaterialID();
+			for (int32 TriangleID : EditMesh.TriangleIndicesItr())
+			{
+				MaterialIDs->SetValue(TriangleID, TriangleMaterialIDs[TriangleID]);
+			}
+		}
+
+	}, EDynamicMeshChangeType::GeneralEdit, EDynamicMeshAttributeChangeFlags::Unknown, bDeferChangeNotifications);
+
+	return TargetMesh;
+}
+
+
+
+
+#undef LOCTEXT_NAMESPACE
\ No newline at end of file
diff --git a/Engine/Plugins/Experimental/GeometryScripting/Source/GeometryScriptingCore/Private/MeshPrimitiveFunctions.cpp b/Engine/Plugins/Experimental/GeometryScripting/Source/GeometryScriptingCore/Private/MeshPrimitiveFunctions.cpp
--- a/Engine/Plugins/Experimental/GeometryScripting/Source/GeometryScriptingCore/Private/MeshPrimitiveFunctions.cpp
+++ b/Engine/Plugins/Experimental/GeometryScripting/Source/GeometryScriptingCore/Private/MeshPrimitiveFunctions.cpp
@@ -15,6 +15,9 @@
 #include "Generators/RectangleMeshGenerator.h"
 #include "Generators/DiscMeshGenerator.h"
 #include "Generators/SweepGenerator.h"
+#include "Generators/FlatTriangulationMeshGenerator.h"
+#include "ConstrainedDelaunay2.h"
+#include "Arrangement2d.h"
 
 #include "ExplicitUseGeometryMathTypes.h"		// using UE::Geometry::(math types)
 using namespace UE::Geometry;
@@ -335,6 +338,55 @@ UDynamicMesh* UGeometryScriptLibrary_MeshPrimitiveFunctions::AppendSimpleRevolve
 
 
 
+
+UDynamicMesh* UGeometryScriptLibrary_MeshPrimitiveFunctions::AppendSimpleExtrudePolygon(
+	UDynamicMesh* TargetMesh,
+	FGeometryScriptPrimitiveOptions PrimitiveOptions,
+	FTransform Transform,
+	const TArray<FVector2D>& PolygonVertices,
+	float Height,
+	int32 HeightSteps,
+	bool bCapped,
+	UGeometryScriptDebug* Debug)
+{
+	if (TargetMesh == nullptr)
+	{
+		UE::Geometry::AppendError(Debug, EGeometryScriptErrorType::InvalidInputs, LOCTEXT("AppendSimpleExtrudePolygon_NullMesh", "AppendSimpleExtrudePolygon: TargetMesh is Null"));
+		return TargetMesh;
+	}
+	if (PolygonVertices.Num() < 3)
+	{
+		UE::Geometry::AppendError(Debug, EGeometryScriptErrorType::InvalidInputs, LOCTEXT("AppendSimpleExtrudePolygon_InvalidPolygon", "AppendSimpleExtrudePolygon: PolygonVertices array requires at least 3 positions"));
+		return TargetMesh;
+	}
+
+	FGeneralizedCylinderGenerator ExtrudeGen;
+	for (const FVector2D& Point : PolygonVertices)
+	{
+		ExtrudeGen.CrossSection.AppendVertex((FVector2d)Point);
+	}
+
+	int32 NumDivisions = FMath::Max(1, HeightSteps - 1);
+	int32 NumPathSteps = NumDivisions + 1;
+	double StepSize = (double)Height / (double)NumDivisions;
+
+	for (int32 k = 0; k <= NumPathSteps; ++k)
+	{
+		double StepHeight = (k == NumPathSteps) ? Height : ((double)k * StepSize);
+		ExtrudeGen.Path.Add(FVector3d(0, 0, StepHeight));
+	}
+
+	ExtrudeGen.InitialFrame = FFrame3d();
+	ExtrudeGen.bCapped = bCapped;
+	ExtrudeGen.bPolygroupPerQuad = (PrimitiveOptions.PolygroupMode == EGeometryScriptPrimitivePolygroupMode::PerQuad);
+	ExtrudeGen.Generate();
+
+	AppendPrimitive(TargetMesh, &ExtrudeGen, Transform, PrimitiveOptions);
+	return TargetMesh;
+}
+
+
+
 UDynamicMesh* UGeometryScriptLibrary_MeshPrimitiveFunctions::AppendSimpleSweptPolygon(
 	UDynamicMesh* TargetMesh,
 	FGeometryScriptPrimitiveOptions PrimitiveOptions,
@@ -476,6 +528,7 @@ UDynamicMesh* UGeometryScriptLibrary_MeshPrimitiveFunctions::AppendDisc(
 	int32 SpokeSteps,
 	float StartAngle,
 	float EndAngle,
+	float HoleRadius,
 	UGeometryScriptDebug* Debug)
 {
 	if (TargetMesh == nullptr)
@@ -485,21 +538,94 @@ UDynamicMesh* UGeometryScriptLibrary_MeshPrimitiveFunctions::AppendDisc(
 	}
 
 	FDiscMeshGenerator DiscGenerator;
-	DiscGenerator.Radius = FMath::Max(FMathf::ZeroTolerance, Radius);
-	DiscGenerator.Normal = FVector3f::UnitZ();
-	DiscGenerator.AngleSamples = FMath::Max(3, AngleSteps);
-	DiscGenerator.RadialSamples = FMath::Max(3, SpokeSteps);
-	DiscGenerator.StartAngle = StartAngle;
-	DiscGenerator.EndAngle = EndAngle;
-	DiscGenerator.bSinglePolygroup = (PrimitiveOptions.PolygroupMode != EGeometryScriptPrimitivePolygroupMode::PerQuad);
-	DiscGenerator.Generate();
+	FPuncturedDiscMeshGenerator PuncturedDiscGenerator;
+	FDiscMeshGenerator* UseGenerator = &DiscGenerator;
 
-	AppendPrimitive(TargetMesh, &DiscGenerator, Transform, PrimitiveOptions);
+	if (HoleRadius > 0)
+	{
+		PuncturedDiscGenerator.HoleRadius = HoleRadius;
+		UseGenerator = &PuncturedDiscGenerator;
+	}
+
+	UseGenerator->Radius = FMath::Max(FMathf::ZeroTolerance, Radius);
+	UseGenerator->Normal = FVector3f::UnitZ();
+	UseGenerator->AngleSamples = FMath::Max(3, AngleSteps);
+	UseGenerator->RadialSamples = FMath::Max(3, SpokeSteps);
+	UseGenerator->StartAngle = StartAngle;
+	UseGenerator->EndAngle = EndAngle;
+	UseGenerator->bSinglePolygroup = (PrimitiveOptions.PolygroupMode != EGeometryScriptPrimitivePolygroupMode::PerQuad);
+	UseGenerator->Generate();
+	AppendPrimitive(TargetMesh, UseGenerator, Transform, PrimitiveOptions);
 
 	return TargetMesh;
 }
 
 
+UDynamicMesh* UGeometryScriptLibrary_MeshPrimitiveFunctions::AppendTriangulatedPolygon(
+	UDynamicMesh* TargetMesh,
+	FGeometryScriptPrimitiveOptions PrimitiveOptions,
+	FTransform Transform,
+	const TArray<FVector2D>& PolygonVertices,
+	bool bAllowSelfIntersections,
+	UGeometryScriptDebug* Debug)
+{
+	if (TargetMesh == nullptr)
+	{
+		UE::Geometry::AppendError(Debug, EGeometryScriptErrorType::InvalidInputs, LOCTEXT("AppendTriangulatedPolygon_InvalidInput", "AppendTriangulatedPolygon: TargetMesh is Null"));
+		return TargetMesh;
+	}
+	if (PolygonVertices.Num() < 3)
+	{
+		UE::Geometry::AppendError(Debug, EGeometryScriptErrorType::InvalidInputs, LOCTEXT("AppendTriangulatedPolygon_InvalidPolygon", "AppendTriangulatedPolygon: PolygonVertices array requires at least 3 positions"));
+		return TargetMesh;
+	}
+
+	FPolygon2d Polygon;
+	for (FVector2D Vertex : PolygonVertices)
+	{
+		Polygon.AppendVertex(Vertex);
+	}
+	FGeneralPolygon2d GeneralPolygon(Polygon);
+
+	FConstrainedDelaunay2d Triangulator;
+	if (bAllowSelfIntersections)
+	{
+		FArrangement2d Arrangement(Polygon.Bounds());
+		// arrangement2d builds a general 2d graph that discards orientation info ...
+		Triangulator.FillRule = FConstrainedDelaunay2d::EFillRule::Odd;
+		Triangulator.bOrientedEdges = false;
+		Triangulator.bSplitBowties = true;
+		for (FSegment2d Seg : GeneralPolygon.GetOuter().Segments())
+		{
+			Arrangement.Insert(Seg);
+		}
+		Triangulator.Add(Arrangement.Graph);
+	}
+	else
+	{
+		Triangulator.Add(GeneralPolygon);
+	}
+
+	bool bTriangulationSuccess = Triangulator.Triangulate([&GeneralPolygon](const TArray<FVector2d>& Vertices, FIndex3i Tri) 
+	{
+		return GeneralPolygon.Contains((Vertices[Tri.A] + Vertices[Tri.B] + Vertices[Tri.C]) / 3.0);	// keep triangles based on the input polygon's winding
+	});
+
+	// even if bTriangulationSuccess is false, there may still be some triangles, so only fail if the mesh is empty
+	if (Triangulator.Triangles.Num() == 0)
+	{
+		UE::Geometry::AppendError(Debug, EGeometryScriptErrorType::OperationFailed, LOCTEXT("AppendTriangulatedPolygon_Failed", "AppendTriangulatedPolygon: Failed to triangulate polygon"));
+		return TargetMesh;
+	}
+
+	FFlatTriangulationMeshGenerator TriangulationMeshGen;
+	TriangulationMeshGen.Vertices2D = Triangulator.Vertices;
+	TriangulationMeshGen.Triangles2D = Triangulator.Triangles;
+	AppendPrimitive(TargetMesh, &TriangulationMeshGen.Generate(), Transform, PrimitiveOptions);
+
+	return TargetMesh;
+}
+
 
 
 #undef LOCTEXT_NAMESPACE
\ No newline at end of file
diff --git a/Engine/Plugins/Experimental/GeometryScripting/Source/GeometryScriptingCore/Private/MeshQueryFunctions.cpp b/Engine/Plugins/Experimental/GeometryScripting/Source/GeometryScriptingCore/Private/MeshQueryFunctions.cpp
--- a/Engine/Plugins/Experimental/GeometryScripting/Source/GeometryScriptingCore/Private/MeshQueryFunctions.cpp
+++ b/Engine/Plugins/Experimental/GeometryScripting/Source/GeometryScriptingCore/Private/MeshQueryFunctions.cpp
@@ -7,6 +7,7 @@
 #include "DynamicMesh/DynamicMeshAttributeSet.h"
 #include "MeshQueries.h"
 #include "MeshBoundaryLoops.h"
+#include "Selections/MeshConnectedComponents.h"
 
 
 #include "ExplicitUseGeometryMathTypes.h"		// using UE::Geometry::(math types)
@@ -115,6 +116,16 @@ int32 UGeometryScriptLibrary_MeshQueryFunctions::GetNumOpenBorderLoops( UDynamic
 }
 
 
+int32 UGeometryScriptLibrary_MeshQueryFunctions::GetNumConnectedComponents( UDynamicMesh* TargetMesh )
+{
+	return SimpleMeshQuery<int32>(TargetMesh, 0, [&](const FDynamicMesh3& Mesh) { 
+		FMeshConnectedComponents Components(&Mesh);
+		Components.FindConnectedTriangles();
+		return Components.Num();
+	});
+}
+
+
 
 int32 UGeometryScriptLibrary_MeshQueryFunctions::GetTriangleCount( UDynamicMesh* TargetMesh )
 {
@@ -199,22 +210,24 @@ FVector UGeometryScriptLibrary_MeshQueryFunctions::GetVertexPosition(UDynamicMes
 	});
 }
 
-TArray<FVector> UGeometryScriptLibrary_MeshQueryFunctions::GetAllVertexPositions(UDynamicMesh* TargetMesh, bool& bHasVertexIDGaps)
+UDynamicMesh* UGeometryScriptLibrary_MeshQueryFunctions::GetAllVertexPositions(UDynamicMesh* TargetMesh, TArray<FVector>& VertexPositions, bool& bHasVertexIDGaps)
 {
-	TArray<FVector> Positions;
+	VertexPositions.Reset();
+	bHasVertexIDGaps = false;
 	if (TargetMesh)
 	{
 		TargetMesh->ProcessMesh([&](const FDynamicMesh3& ReadMesh)
 		{
-			Positions.Init(FVector::ZeroVector, ReadMesh.MaxVertexID());
+			VertexPositions.Init(FVector::ZeroVector, ReadMesh.MaxVertexID());
 			for (int32 vid : ReadMesh.VertexIndicesItr())
 			{
-				Positions[vid] = (FVector)ReadMesh.GetVertex(vid);
+				VertexPositions[vid] = (FVector)ReadMesh.GetVertex(vid);
 			}
 			bHasVertexIDGaps = ! ReadMesh.IsCompactV();
 		});
 	}
-	return Positions;
+	
+	return TargetMesh;
 }
 
 
@@ -274,5 +287,15 @@ void UGeometryScriptLibrary_MeshQueryFunctions::GetTriangleUVs(UDynamicMesh* Tar
 
 
 
+bool UGeometryScriptLibrary_MeshQueryFunctions::GetHasMaterialIDs( UDynamicMesh* TargetMesh )
+{
+	return SimpleMeshQuery<bool>(TargetMesh, false, [&](const FDynamicMesh3& Mesh) {
+		return (Mesh.HasAttributes() && Mesh.Attributes()->HasMaterialID() && Mesh.Attributes()->GetMaterialID() != nullptr);
+	});
+}
+
+
+
+
 
 #undef LOCTEXT_NAMESPACE
diff --git a/Engine/Plugins/Experimental/GeometryScripting/Source/GeometryScriptingCore/Private/MeshSpatialFunctions.cpp b/Engine/Plugins/Experimental/GeometryScripting/Source/GeometryScriptingCore/Private/MeshSpatialFunctions.cpp
new file mode 100644
--- /dev/null
+++ b/Engine/Plugins/Experimental/GeometryScripting/Source/GeometryScriptingCore/Private/MeshSpatialFunctions.cpp
@@ -0,0 +1,278 @@
+// Copyright Epic Games, Inc. All Rights Reserved.
+
+#include "GeometryScript/MeshSpatialFunctions.h"
+#include "UDynamicMesh.h"
+
+#include "DynamicMesh/DynamicMesh3.h"
+#include "DynamicMesh/DynamicMeshAttributeSet.h"
+#include "DynamicMesh/DynamicMeshAABBTree3.h"
+#include "MeshQueries.h"
+#include "Spatial/FastWinding.h"
+
+
+#include "ExplicitUseGeometryMathTypes.h"		// using UE::Geometry::(math types)
+using namespace UE::Geometry;
+
+#define LOCTEXT_NAMESPACE "UGeometryScriptLibrary_MeshSpatial"
+
+
+void UGeometryScriptLibrary_MeshSpatial::ResetBVH(FGeometryScriptDynamicMeshBVH& ResetBVH)
+{
+	ResetBVH.Spatial.Reset();
+	ResetBVH.FWNTree.Reset();
+}
+
+UDynamicMesh* UGeometryScriptLibrary_MeshSpatial::BuildBVHForMesh(
+	UDynamicMesh* TargetMesh,
+	FGeometryScriptDynamicMeshBVH& OutputBVH,
+	UGeometryScriptDebug* Debug)
+{
+	if (TargetMesh == nullptr)
+	{
+		UE::Geometry::AppendError(Debug, EGeometryScriptErrorType::InvalidInputs, LOCTEXT("BuildBVHForMesh_InvalidInput", "BuildBVHForMesh: TargetMesh is Null"));
+		return TargetMesh;
+	}
+
+	if (OutputBVH.Spatial.IsValid() == false)
+	{
+		OutputBVH.Spatial = MakeShared<FDynamicMeshAABBTree3>();
+	}
+
+	TargetMesh->ProcessMesh([&](const FDynamicMesh3& ReadMesh)
+	{
+		OutputBVH.Spatial->SetMesh(&ReadMesh, true);
+		OutputBVH.FWNTree = MakeShared<TFastWindingTree<FDynamicMesh3>>(OutputBVH.Spatial.Get(), true);
+	});
+
+	return TargetMesh;
+}
+
+
+UDynamicMesh* UGeometryScriptLibrary_MeshSpatial::IsBVHValidForMesh(
+	UDynamicMesh* TargetMesh,
+	const FGeometryScriptDynamicMeshBVH& TestBVH,
+	bool& bIsValid,
+	UGeometryScriptDebug* Debug)
+{
+	bIsValid = false;
+	if (TargetMesh == nullptr)
+	{
+		UE::Geometry::AppendError(Debug, EGeometryScriptErrorType::InvalidInputs, LOCTEXT("IsBVHValidForMesh_InvalidInput", "IsBVHValidForMesh: TargetMesh is Null"));
+		return TargetMesh;
+	}
+
+	if (TestBVH.Spatial.IsValid())
+	{
+		TargetMesh->ProcessMesh([&](const FDynamicMesh3& ReadMesh)
+		{
+			if (TestBVH.Spatial->GetMesh() == &ReadMesh)
+			{
+				bIsValid = TestBVH.Spatial->IsValid(false);
+			}
+		});
+	}
+
+	return TargetMesh;
+}
+
+
+UDynamicMesh* UGeometryScriptLibrary_MeshSpatial::RebuildBVHForMesh(
+	UDynamicMesh* TargetMesh,
+	FGeometryScriptDynamicMeshBVH& UpdateBVH,
+	bool bOnlyIfInvalid,
+	UGeometryScriptDebug* Debug)
+{
+	if (TargetMesh == nullptr)
+	{
+		UE::Geometry::AppendError(Debug, EGeometryScriptErrorType::InvalidInputs, LOCTEXT("RebuildBVHForMesh_InvalidInput", "RebuildBVHForMesh: TargetMesh is Null"));
+		return TargetMesh;
+	}
+
+	if (UpdateBVH.Spatial.IsValid() == false)
+	{
+		UpdateBVH.Spatial = MakeShared<FDynamicMeshAABBTree3>();
+	}
+
+	TargetMesh->ProcessMesh([&](const FDynamicMesh3& ReadMesh)
+	{
+		if (bOnlyIfInvalid == false || UpdateBVH.Spatial->GetMesh() != &ReadMesh || UpdateBVH.Spatial->IsValid(false) == false)
+		{
+			UpdateBVH.Spatial->SetMesh(&ReadMesh, false);
+			UpdateBVH.Spatial->Build();
+			// FWNTree does not support rebuilding like this?
+			UpdateBVH.FWNTree = MakeShared<TFastWindingTree<FDynamicMesh3>>(UpdateBVH.Spatial.Get(), true);
+		}
+	});
+
+	return TargetMesh;
+}
+
+
+
+UDynamicMesh* UGeometryScriptLibrary_MeshSpatial::FindNearestPointOnMesh(
+	UDynamicMesh* TargetMesh,
+	const FGeometryScriptDynamicMeshBVH& QueryBVH,
+	FVector QueryPoint,
+	FGeometryScriptSpatialQueryOptions Options,
+	FGeometryScriptTrianglePoint& NearestResult,
+	TEnumAsByte<EGeometryScriptSearchOutcomePins>& Outcome,
+	UGeometryScriptDebug* Debug)
+{
+	Outcome = EGeometryScriptSearchOutcomePins::NotFound;
+	NearestResult.bValid = false;
+	if (TargetMesh == nullptr)
+	{
+		UE::Geometry::AppendError(Debug, EGeometryScriptErrorType::InvalidInputs, LOCTEXT("FindNearestPointOnMesh_InvalidInput", "FindNearestPointOnMesh: TargetMesh is Null"));
+		return TargetMesh;
+	}
+
+	bool bHaveValidBVH = false;
+	TargetMesh->ProcessMesh([&](const FDynamicMesh3& ReadMesh)
+	{
+		if (QueryBVH.Spatial.IsValid())
+		{
+			TargetMesh->ProcessMesh([&](const FDynamicMesh3& ReadMesh)
+			{
+				if (QueryBVH.Spatial->GetMesh() == &ReadMesh && QueryBVH.Spatial->IsValid(Options.bAllowUnsafeModifiedQueries))
+				{
+					bHaveValidBVH = true;
+					IMeshSpatial::FQueryOptions QueryOptions;
+					QueryOptions.MaxDistance = (Options.MaxDistance == 0) ? TNumericLimits<float>::Max() : Options.MaxDistance;
+					double NearDistSqr = 0;
+					int NearestTID = QueryBVH.Spatial->FindNearestTriangle( (FVector3d)QueryPoint, NearDistSqr, QueryOptions);
+					if (NearestTID >= 0)
+					{
+						Outcome = EGeometryScriptSearchOutcomePins::Found;
+						FDistPoint3Triangle3d Distance = TMeshQueries<FDynamicMesh3>::TriangleDistance(ReadMesh, NearestTID, (FVector3d)QueryPoint);
+						NearestResult.bValid = true;
+						NearestResult.TriangleID = NearestTID;
+						NearestResult.Position = (FVector)Distance.ClosestTrianglePoint;
+						NearestResult.BaryCoords = (FVector)Distance.TriangleBaryCoords;
+					}
+				}
+			});
+		}
+	});
+
+	if (!bHaveValidBVH)
+	{
+		UE::Geometry::AppendError(Debug, EGeometryScriptErrorType::InvalidInputs, LOCTEXT("FindNearestPointOnMesh_InvalidBVH", "FindNearestPointOnMesh: QueryBVH is Invalid for this TargetMesh"));
+	}
+
+	return TargetMesh;
+}
+
+
+
+
+UDynamicMesh* UGeometryScriptLibrary_MeshSpatial::FindNearestRayIntersectionWithMesh(
+	UDynamicMesh* TargetMesh,
+	UPARAM(ref) const FGeometryScriptDynamicMeshBVH& QueryBVH,
+	FVector RayOrigin,
+	FVector RayDirection,
+	FGeometryScriptSpatialQueryOptions Options,
+	FGeometryScriptRayHitResult& HitResult,
+	TEnumAsByte<EGeometryScriptSearchOutcomePins>& Outcome,
+	UGeometryScriptDebug* Debug)
+{
+	Outcome = EGeometryScriptSearchOutcomePins::NotFound;
+	HitResult.bHit = false;
+
+	if (TargetMesh == nullptr)
+	{
+		UE::Geometry::AppendError(Debug, EGeometryScriptErrorType::InvalidInputs, LOCTEXT("FindNearestRayIntersectionWithMesh_InvalidInput", "FindNearestRayIntersectionWithMesh: TargetMesh is Null"));
+		return TargetMesh;
+	}
+
+	bool bHaveValidBVH = false;
+	TargetMesh->ProcessMesh([&](const FDynamicMesh3& ReadMesh)
+	{
+		if (QueryBVH.Spatial.IsValid())
+		{
+			TargetMesh->ProcessMesh([&](const FDynamicMesh3& ReadMesh)
+			{
+				if (QueryBVH.Spatial->GetMesh() == &ReadMesh && QueryBVH.Spatial->IsValid(Options.bAllowUnsafeModifiedQueries))
+				{
+					bHaveValidBVH = true;
+					IMeshSpatial::FQueryOptions QueryOptions;
+					QueryOptions.MaxDistance = (Options.MaxDistance == 0) ? TNumericLimits<float>::Max() : Options.MaxDistance;
+					UE::Geometry::FRay3d Ray((FVector3d)RayOrigin, Normalized((FVector3d)RayDirection));
+					int HitTID = QueryBVH.Spatial->FindNearestHitTriangle(Ray, QueryOptions);
+					if (HitTID >= 0)
+					{
+						Outcome = EGeometryScriptSearchOutcomePins::Found;
+						FIntrRay3Triangle3d Intersection = TMeshQueries<FDynamicMesh3>::TriangleIntersection(ReadMesh, HitTID, Ray);
+						HitResult.RayParameter = Intersection.RayParameter;
+						HitResult.bHit = true;
+						HitResult.HitTriangleID = HitTID;
+						HitResult.HitPosition = Ray.PointAt(Intersection.RayParameter);
+						HitResult.HitBaryCoords = (FVector)Intersection.TriangleBaryCoords;
+					}
+				}
+			});
+		}
+	});
+
+	if (!bHaveValidBVH)
+	{
+		UE::Geometry::AppendError(Debug, EGeometryScriptErrorType::InvalidInputs, LOCTEXT("FindNearestRayIntersectionWithMesh_InvalidBVH", "FindNearestRayIntersectionWithMesh: QueryBVH is Invalid for this TargetMesh"));
+	}
+
+	return TargetMesh;
+}
+
+
+
+
+
+
+UDynamicMesh* UGeometryScriptLibrary_MeshSpatial::IsPointInsideMesh(
+	UDynamicMesh* TargetMesh,
+	UPARAM(ref) const FGeometryScriptDynamicMeshBVH& QueryBVH,
+	FVector QueryPoint,
+	FGeometryScriptSpatialQueryOptions Options,
+	bool& bIsInside,
+	TEnumAsByte<EGeometryScriptContainmentOutcomePins>& Outcome,
+	UGeometryScriptDebug* Debug)
+{
+	Outcome = EGeometryScriptContainmentOutcomePins::Outside;
+	bIsInside = false;
+
+	if (TargetMesh == nullptr)
+	{
+		UE::Geometry::AppendError(Debug, EGeometryScriptErrorType::InvalidInputs, LOCTEXT("IsPointInsideMesh_InvalidInput", "IsPointInsideMesh: TargetMesh is Null"));
+		return TargetMesh;
+	}
+
+	bool bHaveValidBVH = false;
+	TargetMesh->ProcessMesh([&](const FDynamicMesh3& ReadMesh)
+	{
+		if (QueryBVH.Spatial.IsValid() && QueryBVH.FWNTree.IsValid())
+		{
+			TargetMesh->ProcessMesh([&](const FDynamicMesh3& ReadMesh)
+			{
+				if (QueryBVH.Spatial->GetMesh() == &ReadMesh && QueryBVH.Spatial->IsValid(Options.bAllowUnsafeModifiedQueries))
+				{
+					bHaveValidBVH = true;
+					bIsInside = QueryBVH.FWNTree->IsInside((FVector3d)QueryPoint, Options.WindingIsoThreshold);
+					if (bIsInside)
+					{
+						Outcome = EGeometryScriptContainmentOutcomePins::Inside;
+					}
+				}
+			});
+		}
+	});
+
+	if (!bHaveValidBVH)
+	{
+		UE::Geometry::AppendError(Debug, EGeometryScriptErrorType::InvalidInputs, LOCTEXT("FindNearestRayIntersectionWithMesh_InvalidBVH", "FindNearestRayIntersectionWithMesh: QueryBVH is Invalid for this TargetMesh."));
+	}
+
+	return TargetMesh;
+}
+
+
+
+
+#undef LOCTEXT_NAMESPACE
diff --git a/Engine/Plugins/Experimental/GeometryScripting/Source/GeometryScriptingCore/Private/MeshUVFunctions.cpp b/Engine/Plugins/Experimental/GeometryScripting/Source/GeometryScriptingCore/Private/MeshUVFunctions.cpp
--- a/Engine/Plugins/Experimental/GeometryScripting/Source/GeometryScriptingCore/Private/MeshUVFunctions.cpp
+++ b/Engine/Plugins/Experimental/GeometryScripting/Source/GeometryScriptingCore/Private/MeshUVFunctions.cpp
@@ -5,6 +5,9 @@
 #include "DynamicMesh/DynamicMesh3.h"
 #include "DynamicMesh/DynamicMeshAttributeSet.h"
 #include "Polygroups/PolygroupSet.h"
+#include "Parameterization/DynamicMeshUVEditor.h"
+#include "Selections/MeshConnectedComponents.h"
+#include "Async/ParallelFor.h"
 #include "UDynamicMesh.h"
 
 #include "ExplicitUseGeometryMathTypes.h"		// using UE::Geometry::(math types)
@@ -13,6 +16,39 @@ using namespace UE::Geometry;
 #define LOCTEXT_NAMESPACE "UGeometryScriptLibrary_MeshUVFunctions"
 
 
+
+UDynamicMesh* UGeometryScriptLibrary_MeshUVFunctions::SetNumUVSets(
+	UDynamicMesh* TargetMesh,
+	int NumUVSets,
+	UGeometryScriptDebug* Debug)
+{
+	if (TargetMesh == nullptr)
+	{
+		UE::Geometry::AppendError(Debug, EGeometryScriptErrorType::InvalidInputs, LOCTEXT("SetNumUVSets_InvalidInput", "SetNumUVSets: TargetMesh is Null"));
+		return TargetMesh;
+	}
+	if (NumUVSets > 8)
+	{
+		UE::Geometry::AppendError(Debug, EGeometryScriptErrorType::InvalidInputs, LOCTEXT("SetNumUVSets_InvalidNumUVSets", "SetNumUVSets: Maximum of 8 UV Sets are supported"));
+		return TargetMesh;
+	}
+	TargetMesh->EditMesh([&](FDynamicMesh3& EditMesh)
+	{
+		if (EditMesh.HasAttributes() == false)
+		{
+			EditMesh.EnableAttributes();
+		}
+		if (NumUVSets != EditMesh.Attributes()->NumUVLayers())
+		{
+			EditMesh.Attributes()->SetNumUVLayers(NumUVSets);
+		}
+	}, EDynamicMeshChangeType::GeneralEdit, EDynamicMeshAttributeChangeFlags::Unknown, false);
+
+	return TargetMesh;
+}
+
+
+
 UDynamicMesh* UGeometryScriptLibrary_MeshUVFunctions::SetMeshTriangleUVs(
 	UDynamicMesh* TargetMesh,
 	int UVSetIndex,
@@ -47,4 +83,258 @@ UDynamicMesh* UGeometryScriptLibrary_MeshUVFunctions::SetMeshTriangleUVs(
 
 
 
+void ApplyMeshUVEditorOperation(UDynamicMesh* TargetMesh, int32 UVSetIndex, bool& bHasUVSet, UGeometryScriptDebug* Debug,
+	TFunctionRef<void(FDynamicMesh3& Mesh, FDynamicMeshUVOverlay* UVOverlay, FDynamicMeshUVEditor& UVEditor)> EditFunc)
+{
+	bHasUVSet = false;
+	TargetMesh->EditMesh([&](FDynamicMesh3& EditMesh)
+	{
+		if (EditMesh.HasAttributes() == false
+			|| UVSetIndex >= EditMesh.Attributes()->NumUVLayers()
+			|| EditMesh.Attributes()->GetUVLayer(UVSetIndex) == nullptr)
+		{
+			return;
+		}
+
+		bHasUVSet = true;
+		FDynamicMeshUVOverlay* UVOverlay = EditMesh.Attributes()->GetUVLayer(UVSetIndex);
+		FDynamicMeshUVEditor Editor(&EditMesh, UVOverlay);
+		EditFunc(EditMesh, UVOverlay, Editor);
+
+	}, EDynamicMeshChangeType::GeneralEdit, EDynamicMeshAttributeChangeFlags::Unknown, false);
+}
+
+
+
+
+UDynamicMesh* UGeometryScriptLibrary_MeshUVFunctions::TranslateMeshUVs(
+	UDynamicMesh* TargetMesh,
+	int UVSetIndex,
+	FVector2D Translation,
+	UGeometryScriptDebug* Debug)
+{
+	if (TargetMesh == nullptr)
+	{
+		UE::Geometry::AppendError(Debug, EGeometryScriptErrorType::InvalidInputs, LOCTEXT("TranslateMeshUVs_InvalidInput", "TranslateMeshUVs: TargetMesh is Null"));
+		return TargetMesh;
+	}
+
+	bool bHasUVSet = false;
+	ApplyMeshUVEditorOperation(TargetMesh, UVSetIndex, bHasUVSet, Debug,
+		[&](FDynamicMesh3& EditMesh, FDynamicMeshUVOverlay* UVOverlay, FDynamicMeshUVEditor& UVEditor)
+	{
+		for (int32 elemid : UVOverlay->ElementIndicesItr())
+		{
+			FVector2f UV = UVOverlay->GetElement(elemid);
+			UVOverlay->SetElement(elemid, UV + (FVector2f)Translation);
+		}
+	});
+	if (bHasUVSet == false)
+	{
+		UE::Geometry::AppendError(Debug, EGeometryScriptErrorType::InvalidInputs, LOCTEXT("TranslateMeshUVs_InvalidUVSet", "TranslateMeshUVs: UVSetIndex does not exist on TargetMesh"));
+	}
+
+	return TargetMesh;
+}
+
+
+UDynamicMesh* UGeometryScriptLibrary_MeshUVFunctions::ScaleMeshUVs(
+	UDynamicMesh* TargetMesh,
+	int UVSetIndex,
+	FVector2D Scale,
+	FVector2D ScaleOrigin,
+	UGeometryScriptDebug* Debug)
+{
+	if (TargetMesh == nullptr)
+	{
+		UE::Geometry::AppendError(Debug, EGeometryScriptErrorType::InvalidInputs, LOCTEXT("ScaleMeshUVs_InvalidInput", "ScaleMeshUVs: TargetMesh is Null"));
+		return TargetMesh;
+	}
+
+	FVector2f UseScale = Scale;
+	if (UseScale.Length() < 0.0001)
+	{
+		UseScale = FVector2f::One();
+	}
+	FVector2f UseOrigin = (FVector2f)ScaleOrigin;
+
+	bool bHasUVSet = false;
+	ApplyMeshUVEditorOperation(TargetMesh, UVSetIndex, bHasUVSet, Debug,
+		[&](FDynamicMesh3& EditMesh, FDynamicMeshUVOverlay* UVOverlay, FDynamicMeshUVEditor& UVEditor)
+	{
+		for (int32 elemid : UVOverlay->ElementIndicesItr())
+		{
+			FVector2f UV = UVOverlay->GetElement(elemid);
+			UV = (UV - UseOrigin) * UseScale + UseOrigin;
+			UVOverlay->SetElement(elemid, UV);
+		}
+	});
+	if (bHasUVSet == false)
+	{
+		UE::Geometry::AppendError(Debug, EGeometryScriptErrorType::InvalidInputs, LOCTEXT("ScaleMeshUVs_InvalidUVSet", "ScaleMeshUVs: UVSetIndex does not exist on TargetMesh"));
+	}
+
+	return TargetMesh;
+}
+
+
+
+UDynamicMesh* UGeometryScriptLibrary_MeshUVFunctions::RotateMeshUVs(
+	UDynamicMesh* TargetMesh,
+	int UVSetIndex,
+	float RotationAngle,
+	FVector2D RotationOrigin,
+	UGeometryScriptDebug* Debug)
+{
+	if (TargetMesh == nullptr)
+	{
+		UE::Geometry::AppendError(Debug, EGeometryScriptErrorType::InvalidInputs, LOCTEXT("RotateMeshUVs_InvalidInput", "RotateMeshUVs: TargetMesh is Null"));
+		return TargetMesh;
+	}
+
+	FMatrix2f RotationMat = FMatrix2f::RotationDeg(RotationAngle);
+	FVector2f UseOrigin = (FVector2f)RotationOrigin;
+
+	bool bHasUVSet = false;
+	ApplyMeshUVEditorOperation(TargetMesh, UVSetIndex, bHasUVSet, Debug,
+		[&](FDynamicMesh3& EditMesh, FDynamicMeshUVOverlay* UVOverlay, FDynamicMeshUVEditor& UVEditor)
+	{
+		for (int32 elemid : UVOverlay->ElementIndicesItr())
+		{
+			FVector2f UV = UVOverlay->GetElement(elemid);
+			UV = RotationMat* (UV - UseOrigin) + UseOrigin;
+			UVOverlay->SetElement(elemid, UV);
+		}
+	});
+	if (bHasUVSet == false)
+	{
+		UE::Geometry::AppendError(Debug, EGeometryScriptErrorType::InvalidInputs, LOCTEXT("RotateMeshUVs_InvalidUVSet", "RotateMeshUVs: UVSetIndex does not exist on TargetMesh"));
+	}
+
+	return TargetMesh;
+}
+
+
+
+
+UDynamicMesh* UGeometryScriptLibrary_MeshUVFunctions::SetMeshUVsFromPlanarProjection(
+	UDynamicMesh* TargetMesh,
+	int UVSetIndex,
+	FTransform PlaneTransform,
+	UGeometryScriptDebug* Debug)
+{
+	if (TargetMesh == nullptr)
+	{
+		UE::Geometry::AppendError(Debug, EGeometryScriptErrorType::InvalidInputs, LOCTEXT("SetMeshUVsFromPlanarProjection_InvalidInput", "SetMeshUVsFromPlanarProjection: TargetMesh is Null"));
+		return TargetMesh;
+	}
+
+	bool bHasUVSet = false;
+	ApplyMeshUVEditorOperation(TargetMesh, UVSetIndex, bHasUVSet, Debug,
+		[&](FDynamicMesh3& EditMesh, FDynamicMeshUVOverlay* UVOverlay, FDynamicMeshUVEditor& UVEditor)
+	{
+		TArray<int32> AllTriangles;
+		for (int32 tid : EditMesh.TriangleIndicesItr())
+		{
+			AllTriangles.Add(tid);
+		}
+
+		FFrame3d ProjectionFrame(PlaneTransform);
+		FVector Scale = PlaneTransform.GetScale3D();
+		FVector2d Dimensions(Scale.X, Scale.Y);
+
+		UVEditor.SetTriangleUVsFromPlanarProjection(AllTriangles, [&](const FVector3d& Pos) { return Pos; },
+			ProjectionFrame, Dimensions);
+	});
+	if (bHasUVSet == false)
+	{
+		UE::Geometry::AppendError(Debug, EGeometryScriptErrorType::InvalidInputs, LOCTEXT("SetMeshUVsFromPlanarProjection_InvalidUVSet", "SetMeshUVsFromPlanarProjection: UVSetIndex does not exist on TargetMesh"));
+	}
+
+	return TargetMesh;
+}
+
+
+
+
+UDynamicMesh* UGeometryScriptLibrary_MeshUVFunctions::SetMeshUVsFromBoxProjection(
+	UDynamicMesh* TargetMesh,
+	int UVSetIndex,
+	FTransform PlaneTransform,
+	int MinIslandTriCount,
+	UGeometryScriptDebug* Debug)
+{
+	if (TargetMesh == nullptr)
+	{
+		UE::Geometry::AppendError(Debug, EGeometryScriptErrorType::InvalidInputs, LOCTEXT("SetMeshUVsFromBoxProjection_InvalidInput", "SetMeshUVsFromBoxProjection: TargetMesh is Null"));
+		return TargetMesh;
+	}
+
+	bool bHasUVSet = false;
+	ApplyMeshUVEditorOperation(TargetMesh, UVSetIndex, bHasUVSet, Debug,
+		[&](FDynamicMesh3& EditMesh, FDynamicMeshUVOverlay* UVOverlay, FDynamicMeshUVEditor& UVEditor)
+	{
+		TArray<int32> AllTriangles;
+		for (int32 tid : EditMesh.TriangleIndicesItr())
+		{
+			AllTriangles.Add(tid);
+		}
+
+		FFrame3d ProjectionFrame(PlaneTransform);
+		FVector Scale = PlaneTransform.GetScale3D();
+		FVector3d Dimensions = (FVector)Scale;
+		UVEditor.SetTriangleUVsFromBoxProjection(AllTriangles, [&](const FVector3d& Pos) { return Pos; },
+			ProjectionFrame, Dimensions, MinIslandTriCount);
+	});
+	if (bHasUVSet == false)
+	{
+		UE::Geometry::AppendError(Debug, EGeometryScriptErrorType::InvalidInputs, LOCTEXT("SetMeshUVsFromBoxProjection_InvalidUVSet", "SetMeshUVsFromBoxProjection: UVSetIndex does not exist on TargetMesh"));
+	}
+
+	return TargetMesh;
+}
+
+
+
+UDynamicMesh* UGeometryScriptLibrary_MeshUVFunctions::RepackMeshUVs( 
+	UDynamicMesh* TargetMesh, 
+	int UVSetIndex,
+	FGeometryScriptRepackUVsOptions RepackOptions,
+	UGeometryScriptDebug* Debug)
+{
+	if (TargetMesh == nullptr)
+	{
+		UE::Geometry::AppendError(Debug, EGeometryScriptErrorType::InvalidInputs, LOCTEXT("RepackMeshUVs_InvalidInput", "RepackMeshUVs: TargetMesh is Null"));
+		return TargetMesh;
+	}
+
+	bool bHasUVSet = false;
+	ApplyMeshUVEditorOperation(TargetMesh, UVSetIndex, bHasUVSet, Debug,
+		[&](FDynamicMesh3& EditMesh, FDynamicMeshUVOverlay* UVOverlay, FDynamicMeshUVEditor& UVEditor)
+	{
+		if (RepackOptions.bOptimizeIslandRotation)
+		{
+			FMeshConnectedComponents UVComponents(&EditMesh);
+			UVComponents.FindConnectedTriangles([&](int32 Triangle0, int32 Triangle1) {
+				return UVOverlay->AreTrianglesConnected(Triangle0, Triangle1);
+			});
+
+			ParallelFor(UVComponents.Num(), [&](int32 k)
+			{
+				UVEditor.AutoOrientUVArea(UVComponents[k].Indices);
+			});
+		}
+
+		UVEditor.QuickPack(FMath::Max(16, RepackOptions.TargetImageWidth));
+	});
+	if (bHasUVSet == false)
+	{
+		UE::Geometry::AppendError(Debug, EGeometryScriptErrorType::InvalidInputs, LOCTEXT("RepackMeshUVs_InvalidUVSet", "RepackMeshUVs: UVSetIndex does not exist on TargetMesh"));
+	}
+
+	return TargetMesh;
+}
+
+
+
 #undef LOCTEXT_NAMESPACE
\ No newline at end of file
diff --git a/Engine/Plugins/Experimental/GeometryScripting/Source/GeometryScriptingCore/Public/GeometryScript/ContainmentFunctions.h b/Engine/Plugins/Experimental/GeometryScripting/Source/GeometryScriptingCore/Public/GeometryScript/ContainmentFunctions.h
new file mode 100644
--- /dev/null
+++ b/Engine/Plugins/Experimental/GeometryScripting/Source/GeometryScriptingCore/Public/GeometryScript/ContainmentFunctions.h
@@ -0,0 +1,84 @@
+// Copyright Epic Games, Inc. All Rights Reserved.
+
+#pragma once
+
+#include "CoreMinimal.h"
+#include "Kismet/BlueprintFunctionLibrary.h"
+#include "GeometryScript/GeometryScriptTypes.h"
+#include "ContainmentFunctions.generated.h"
+
+class UDynamicMesh;
+
+USTRUCT(BlueprintType)
+struct GEOMETRYSCRIPTINGCORE_API FGeometryScriptConvexHullOptions
+{
+	GENERATED_BODY()
+public:
+	UPROPERTY(BlueprintReadWrite, Category = Options)
+	bool bPrefilterVertices = true;
+
+	UPROPERTY(BlueprintReadWrite, Category = Options)
+	int PrefilterGridResolution = 128;
+
+	/** */
+	UPROPERTY(BlueprintReadWrite, Category = Options)
+	int SimplifyToFaceCount = 0;
+};
+
+
+
+USTRUCT(BlueprintType)
+struct GEOMETRYSCRIPTINGCORE_API FGeometryScriptSweptHullOptions
+{
+	GENERATED_BODY()
+public:
+	UPROPERTY(BlueprintReadWrite, Category = Options)
+	bool bPrefilterVertices = true;
+
+	UPROPERTY(BlueprintReadWrite, Category = Options)
+	int PrefilterGridResolution = 128;
+
+	UPROPERTY(BlueprintReadWrite, Category = Options)
+	float MinThickness = 0.01;
+
+	/** */
+	UPROPERTY(BlueprintReadWrite, Category = Options)
+	bool bSimplify = true;
+
+	UPROPERTY(BlueprintReadWrite, Category = Options)
+	float MinEdgeLength = 0.1;
+
+	UPROPERTY(BlueprintReadWrite, Category = Options)
+	float SimplifyTolerance = 0.1;
+};
+
+
+
+
+
+UCLASS(meta = (ScriptName = "GeometryScript_Containment"))
+class GEOMETRYSCRIPTINGCORE_API UGeometryScriptLibrary_ContainmentFunctions : public UBlueprintFunctionLibrary
+{
+	GENERATED_BODY()
+public:
+
+	UFUNCTION(BlueprintCallable, Category = "GeometryScript|Containment", meta=(ScriptMethod))
+	static UPARAM(DisplayName = "Target Mesh") UDynamicMesh* 
+	ComputeMeshConvexHull(  
+		UDynamicMesh* TargetMesh, 
+		UPARAM(DisplayName = "Hull Mesh", ref) UDynamicMesh* CopyToMesh, 
+		UPARAM(DisplayName = "Hull Mesh") UDynamicMesh*& CopyToMeshOut, 
+		FGeometryScriptConvexHullOptions Options,
+		UGeometryScriptDebug* Debug = nullptr);
+
+	UFUNCTION(BlueprintCallable, Category = "GeometryScript|Containment", meta=(ScriptMethod))
+	static UPARAM(DisplayName = "Target Mesh") UDynamicMesh* 
+	ComputeMeshSweptHull(  
+		UDynamicMesh* TargetMesh, 
+		UPARAM(DisplayName = "Hull Mesh", ref) UDynamicMesh* CopyToMesh, 
+		UPARAM(DisplayName = "Hull Mesh") UDynamicMesh*& CopyToMeshOut, 
+		FTransform ProjectionFrame,
+		FGeometryScriptSweptHullOptions Options,
+		UGeometryScriptDebug* Debug = nullptr);
+
+};
\ No newline at end of file
diff --git a/Engine/Plugins/Experimental/GeometryScripting/Source/GeometryScriptingCore/Public/GeometryScript/GeometryScriptTypes.h b/Engine/Plugins/Experimental/GeometryScripting/Source/GeometryScriptingCore/Public/GeometryScript/GeometryScriptTypes.h
--- a/Engine/Plugins/Experimental/GeometryScripting/Source/GeometryScriptingCore/Public/GeometryScript/GeometryScriptTypes.h
+++ b/Engine/Plugins/Experimental/GeometryScripting/Source/GeometryScriptingCore/Public/GeometryScript/GeometryScriptTypes.h
@@ -3,8 +3,14 @@
 #pragma once
 
 #include "CoreMinimal.h"
+#include "GeometryBase.h"
 #include "GeometryScriptTypes.generated.h"
 
+PREDECLARE_GEOMETRY(class FDynamicMesh3);
+PREDECLARE_GEOMETRY(template<typename MeshType> class TMeshAABBTree3);
+PREDECLARE_GEOMETRY(template<typename MeshType> class TFastWindingTree);
+PREDECLARE_GEOMETRY(typedef TMeshAABBTree3<FDynamicMesh3> FDynamicMeshAABBTree3);
+
 
 UENUM(BlueprintType)
 enum EGeometryScriptOutcomePins
@@ -13,6 +19,19 @@ enum EGeometryScriptOutcomePins
 	Success
 };
 
+UENUM(BlueprintType)
+enum EGeometryScriptSearchOutcomePins
+{
+	Found,
+	NotFound
+};
+
+UENUM(BlueprintType)
+enum EGeometryScriptContainmentOutcomePins
+{
+	Inside, 
+	Outside
+};
 
 
 UENUM(BlueprintType)
@@ -74,6 +93,25 @@ public:
 };
 
 
+USTRUCT(BlueprintType)
+struct GEOMETRYSCRIPTINGCORE_API FGeometryScriptTrianglePoint
+{
+	GENERATED_BODY()
+public:
+	UPROPERTY(BlueprintReadWrite, Category = Options)
+	bool bValid = false;
+
+	UPROPERTY(BlueprintReadWrite, Category = Options)
+	int TriangleID = -1;
+
+	UPROPERTY(BlueprintReadWrite, Category = Options)
+	FVector Position;
+
+	UPROPERTY(BlueprintReadWrite, Category = Options)
+	FVector BaryCoords;
+};
+
+
 USTRUCT(BlueprintType)
 struct GEOMETRYSCRIPTINGCORE_API FGeometryScriptUVTriangle
 {
@@ -90,6 +128,24 @@ public:
 };
 
 
+
+//
+// Spatial data structures
+//
+
+
+USTRUCT(BlueprintType, meta = (DisplayName = "DynamicMesh BVH Cache"))
+struct GEOMETRYSCRIPTINGCORE_API FGeometryScriptDynamicMeshBVH
+{
+	GENERATED_BODY()
+public:
+
+	TSharedPtr<UE::Geometry::FDynamicMeshAABBTree3> Spatial;
+	TSharedPtr<UE::Geometry::TFastWindingTree<UE::Geometry::FDynamicMesh3>> FWNTree;
+};
+
+
+
 //
 // Errors/Debugging
 //
@@ -131,7 +187,7 @@ public:
 };
 
 
-UCLASS(BlueprintType)
+UCLASS(BlueprintType, meta = (TestMetadata))
 class GEOMETRYSCRIPTINGCORE_API UGeometryScriptDebug : public UObject
 {
 	GENERATED_BODY()
@@ -152,8 +208,8 @@ namespace UE
 {
 namespace Geometry
 {
-	FGeometryScriptDebugMessage MakeScriptError(EGeometryScriptErrorType ErrorTypeIn, const FText& MessageIn);
+	GEOMETRYSCRIPTINGCORE_API FGeometryScriptDebugMessage MakeScriptError(EGeometryScriptErrorType ErrorTypeIn, const FText& MessageIn);
 
-	void AppendError(UGeometryScriptDebug* Debug, EGeometryScriptErrorType ErrorTypeIn, const FText& MessageIn);
+	GEOMETRYSCRIPTINGCORE_API void AppendError(UGeometryScriptDebug* Debug, EGeometryScriptErrorType ErrorTypeIn, const FText& MessageIn);
 }
 }
\ No newline at end of file
diff --git a/Engine/Plugins/Experimental/GeometryScripting/Source/GeometryScriptingCore/Public/GeometryScript/MeshBasicEditFunctions.h b/Engine/Plugins/Experimental/GeometryScripting/Source/GeometryScriptingCore/Public/GeometryScript/MeshBasicEditFunctions.h
--- a/Engine/Plugins/Experimental/GeometryScripting/Source/GeometryScriptingCore/Public/GeometryScript/MeshBasicEditFunctions.h
+++ b/Engine/Plugins/Experimental/GeometryScripting/Source/GeometryScriptingCore/Public/GeometryScript/MeshBasicEditFunctions.h
@@ -10,6 +10,45 @@
 class UDynamicMesh;
 
 
+USTRUCT(BlueprintType)
+struct GEOMETRYSCRIPTINGCORE_API FGeometryScriptSimpleMeshBuffers
+{
+	GENERATED_BODY()
+public:
+	UPROPERTY(BlueprintReadWrite, Category = Options)
+	TArray<FVector> Vertices;
+
+	UPROPERTY(BlueprintReadWrite, Category = Options)
+	TArray<FVector> Normals;
+
+	UPROPERTY(BlueprintReadWrite, Category = Options)
+	TArray<FVector2D> UV0;
+	UPROPERTY(BlueprintReadWrite, Category = Options)
+	TArray<FVector2D> UV1;
+	UPROPERTY(BlueprintReadWrite, Category = Options)
+	TArray<FVector2D> UV2;
+	UPROPERTY(BlueprintReadWrite, Category = Options)
+	TArray<FVector2D> UV3;
+	UPROPERTY(BlueprintReadWrite, Category = Options)
+	TArray<FVector2D> UV4;
+	UPROPERTY(BlueprintReadWrite, Category = Options)
+	TArray<FVector2D> UV5;
+	UPROPERTY(BlueprintReadWrite, Category = Options)
+	TArray<FVector2D> UV6;
+	UPROPERTY(BlueprintReadWrite, Category = Options)
+	TArray<FVector2D> UV7;
+
+	UPROPERTY(BlueprintReadWrite, Category = Options)
+	TArray<FLinearColor> VertexColors;
+
+
+	UPROPERTY(BlueprintReadWrite, Category = Options)
+	TArray<FIntVector> Triangles;
+
+	UPROPERTY(BlueprintReadWrite, Category = Options)
+	TArray<int> TriGroupIDs;
+};
+
 
 UCLASS(meta = (ScriptName = "GeometryScript_MeshEdits"))
 class GEOMETRYSCRIPTINGCORE_API UGeometryScriptLibrary_MeshBasicEditFunctions : public UBlueprintFunctionLibrary
@@ -59,6 +98,13 @@ public:
 		bool& bWasVertexDeleted,
 		bool bDeferChangeNotifications = false );
 
+	UFUNCTION(BlueprintCallable, Category = "GeometryScript|MeshEdits", meta=(ScriptMethod))
+	static UPARAM(DisplayName = "Target Mesh") UDynamicMesh* 
+	DeleteVerticesFromMesh( 
+		UDynamicMesh* TargetMesh, 
+		const TArray<int32>& VertexList,
+		int& NumDeleted,
+		bool bDeferChangeNotifications = false );
 
 
 	UFUNCTION(BlueprintCallable, Category = "GeometryScript|MeshEdits", meta=(ScriptMethod))
@@ -90,7 +136,13 @@ public:
 		bool& bWasTriangleDeleted,
 		bool bDeferChangeNotifications = false );
 
-
+	UFUNCTION(BlueprintCallable, Category = "GeometryScript|MeshEdits", meta=(ScriptMethod))
+	static UPARAM(DisplayName = "Target Mesh") UDynamicMesh* 
+	DeleteTrianglesFromMesh( 
+		UDynamicMesh* TargetMesh, 
+		const TArray<int32>& TriangleList,
+		int& NumDeleted,
+		bool bDeferChangeNotifications = false );
 
 
 	UFUNCTION(BlueprintCallable, Category = "GeometryScript|MeshEdits", meta=(ScriptMethod))
@@ -114,4 +166,15 @@ public:
 		UGeometryScriptDebug* Debug = nullptr);
 
 
+
+	UFUNCTION(BlueprintCallable, Category = "GeometryScript|MeshEdits", meta=(ScriptMethod))
+	static UPARAM(DisplayName = "Target Mesh") UDynamicMesh* 
+	AppendBuffersToMesh( 
+		UDynamicMesh* TargetMesh, 
+		const FGeometryScriptSimpleMeshBuffers& Buffers,
+		TArray<int>& NewTriangleIndices,
+		int MaterialID = 0,
+		bool bDeferChangeNotifications = false,
+		UGeometryScriptDebug* Debug = nullptr );
+
 };
\ No newline at end of file
diff --git a/Engine/Plugins/Experimental/GeometryScripting/Source/GeometryScriptingCore/Public/GeometryScript/MeshComparisonFunctions.h b/Engine/Plugins/Experimental/GeometryScripting/Source/GeometryScriptingCore/Public/GeometryScript/MeshComparisonFunctions.h
new file mode 100644
--- /dev/null
+++ b/Engine/Plugins/Experimental/GeometryScripting/Source/GeometryScriptingCore/Public/GeometryScript/MeshComparisonFunctions.h
@@ -0,0 +1,93 @@
+// Copyright Epic Games, Inc. All Rights Reserved.
+
+#pragma once
+
+#include "CoreMinimal.h"
+#include "Kismet/BlueprintFunctionLibrary.h"
+#include "GeometryScript/GeometryScriptTypes.h"
+#include "MeshComparisonFunctions.generated.h"
+
+class UDynamicMesh;
+
+
+USTRUCT(BlueprintType)
+struct GEOMETRYSCRIPTINGCORE_API FGeometryScriptIsSameMeshOptions
+{
+	GENERATED_BODY()
+public:
+	UPROPERTY(BlueprintReadWrite, Category = Options)
+	bool bCheckConnectivity = true;
+
+	UPROPERTY(BlueprintReadWrite, Category = Options)
+	bool bCheckEdgeIDs = false;
+
+	UPROPERTY(BlueprintReadWrite, Category = Options)
+	bool bCheckNormals = false;
+
+	UPROPERTY(BlueprintReadWrite, Category = Options)
+	bool bCheckColors = false;
+
+	UPROPERTY(BlueprintReadWrite, Category = Options)
+	bool bCheckUVs = false;
+
+	UPROPERTY(BlueprintReadWrite, Category = Options)
+	bool bCheckGroups = false;
+
+	UPROPERTY(BlueprintReadWrite, Category = Options)
+	bool bCheckAttributes = false;
+
+	UPROPERTY(BlueprintReadWrite, Category = Options)
+	float Epsilon = 1e-06;
+};
+
+
+
+USTRUCT(BlueprintType)
+struct GEOMETRYSCRIPTINGCORE_API FGeometryScriptMeasureMeshDistanceOptions
+{
+	GENERATED_BODY()
+public:
+	UPROPERTY(BlueprintReadWrite, Category = Options)
+	bool bSymmetric = true;
+};
+
+
+
+UCLASS(meta = (ScriptName = "GeometryScript_MeshComparison"))
+class GEOMETRYSCRIPTINGCORE_API UGeometryScriptLibrary_MeshComparisonFunctions : public UBlueprintFunctionLibrary
+{
+	GENERATED_BODY()
+public:
+
+	UFUNCTION(BlueprintCallable, Category = "GeometryScript|Comparison", meta = (ScriptMethod))
+	static UPARAM(DisplayName = "Target Mesh") UDynamicMesh*
+	IsSameMeshAs(
+		UDynamicMesh* TargetMesh,
+		UDynamicMesh* OtherMesh,
+		FGeometryScriptIsSameMeshOptions Options,
+		bool &bIsSameMesh,
+		UGeometryScriptDebug* Debug = nullptr);
+
+	UFUNCTION(BlueprintCallable, Category = "GeometryScript|Comparison", meta = (ScriptMethod))
+	static UPARAM(DisplayName = "Target Mesh") UDynamicMesh*
+	MeasureDistancesBetweenMeshes(
+		UDynamicMesh* TargetMesh,
+		UDynamicMesh* OtherMesh,
+		FGeometryScriptMeasureMeshDistanceOptions Options,
+		double& MaxDistance,
+		double& MinDistance,
+		double& AverageDistance,
+		double& RootMeanSqrDeviation,
+		UGeometryScriptDebug* Debug = nullptr);
+
+	UFUNCTION(BlueprintCallable, Category = "GeometryScript|Comparison", meta = (ScriptMethod))
+	static UPARAM(DisplayName = "Target Mesh") UDynamicMesh*
+	IsIntersectingMesh(
+		UDynamicMesh* TargetMesh,
+		FTransform TargetTransform,
+		UDynamicMesh* OtherMesh,
+		FTransform OtherTransform,
+		bool &bIsIntersecting,
+		UGeometryScriptDebug* Debug = nullptr);
+
+};
\ No newline at end of file
diff --git a/Engine/Plugins/Experimental/GeometryScripting/Source/GeometryScriptingCore/Public/GeometryScript/MeshDecompositionFunctions.h b/Engine/Plugins/Experimental/GeometryScripting/Source/GeometryScriptingCore/Public/GeometryScript/MeshDecompositionFunctions.h
new file mode 100644
--- /dev/null
+++ b/Engine/Plugins/Experimental/GeometryScripting/Source/GeometryScriptingCore/Public/GeometryScript/MeshDecompositionFunctions.h
@@ -0,0 +1,53 @@
+// Copyright Epic Games, Inc. All Rights Reserved.
+
+#pragma once
+
+#include "CoreMinimal.h"
+#include "Kismet/BlueprintFunctionLibrary.h"
+#include "GeometryScript/GeometryScriptTypes.h"
+#include "MeshDecompositionFunctions.generated.h"
+
+class UDynamicMesh;
+
+
+UCLASS(meta = (ScriptName = "GeometryScript_MeshDecomposition"))
+class GEOMETRYSCRIPTINGCORE_API UGeometryScriptLibrary_MeshDecompositionFunctions : public UBlueprintFunctionLibrary
+{
+	GENERATED_BODY()
+public:
+
+	UFUNCTION(BlueprintCallable, Category = "GeometryScript|Decomposition", meta=(ScriptMethod))
+	static UPARAM(DisplayName = "Target Mesh") UDynamicMesh* 
+	SplitMeshByComponents(  
+		UDynamicMesh* TargetMesh, 
+		TArray<UDynamicMesh*>& ComponentMeshes,
+		UDynamicMeshPool* MeshPool,
+		UGeometryScriptDebug* Debug = nullptr);
+
+	UFUNCTION(BlueprintCallable, Category = "GeometryScript|Decomposition", meta=(ScriptMethod))
+	static UPARAM(DisplayName = "Target Mesh") UDynamicMesh* 
+	SplitMeshByMaterialIDs(  
+		UDynamicMesh* TargetMesh, 
+		TArray<UDynamicMesh*>& ComponentMeshes,
+		TArray<int>& ComponentMaterialIDs,
+		UDynamicMeshPool* MeshPool,
+		UGeometryScriptDebug* Debug = nullptr);
+
+	UFUNCTION(BlueprintCallable, Category = "GeometryScript|Decomposition", meta=(ScriptMethod))
+	static UPARAM(DisplayName = "Target Mesh") UDynamicMesh* 
+	GetSubMeshFromMesh(  
+		UDynamicMesh* TargetMesh, 
+		UPARAM(DisplayName = "Copy To Submesh", ref) UDynamicMesh* StoreToSubmesh, 
+		const TArray<int>& TriangleList,
+		UPARAM(DisplayName = "Copy To Submesh") UDynamicMesh*& StoreToSubmeshOut, 
+		UGeometryScriptDebug* Debug = nullptr);
+
+	UFUNCTION(BlueprintCallable, Category = "GeometryScript|Decomposition", meta=(ScriptMethod))
+	static UPARAM(DisplayName = "Copy From Mesh") UDynamicMesh* 
+	CopyMeshToMesh(  
+		UDynamicMesh* CopyFromMesh, 
+		UPARAM(DisplayName = "Copy To Mesh", ref) UDynamicMesh* CopyToMesh, 
+		UPARAM(DisplayName = "Copy To Mesh") UDynamicMesh*& CopyToMeshOut, 
+		UGeometryScriptDebug* Debug = nullptr);
+
+};
\ No newline at end of file
diff --git a/Engine/Plugins/Experimental/GeometryScripting/Source/GeometryScriptingCore/Public/GeometryScript/MeshDeformFunctions.h b/Engine/Plugins/Experimental/GeometryScripting/Source/GeometryScriptingCore/Public/GeometryScript/MeshDeformFunctions.h
--- a/Engine/Plugins/Experimental/GeometryScripting/Source/GeometryScriptingCore/Public/GeometryScript/MeshDeformFunctions.h
+++ b/Engine/Plugins/Experimental/GeometryScripting/Source/GeometryScriptingCore/Public/GeometryScript/MeshDeformFunctions.h
@@ -15,22 +15,6 @@ struct GEOMETRYSCRIPTINGCORE_API FGeometryScriptBendWarpOptions
 {
 	GENERATED_BODY()
 public:
-	/** Origin of the Bend Warp */
-	UPROPERTY(BlueprintReadWrite, Category = Options)
-	FVector BendOrigin = FVector(0,0,0);
-
-	/** Bend will be along this axis, ie if bending a long tube, this axis would be aligned with the tube axis */
-	UPROPERTY(BlueprintReadWrite, Category = Options)
-	FVector BendAxis = FVector(0, 0, 1);
-
-	/** Bend will bend towards this direction, eg if bending a vertical (+Z) bar with BendAxis = (0,0,1), A positive bend angle would bend the bar towards this direction (in the XY plane). Automatically calculated if this vector is Zero. */
-	UPROPERTY(BlueprintReadWrite, Category = Options)
-	FVector TowardAxis = FVector(0, 0, 0);
-
-	/** Rotation in Degrees around the Bend Axis, this is applied after the TowardAxis rotation */
-	UPROPERTY(BlueprintReadWrite, Category = Options)
-	float BendAxisRotation = 0;
-
 	/** Symmetric extents are [-BendExtent,BendExtent], if disabled, then [-LowerExtent,BendExtent] is used  */
 	UPROPERTY(BlueprintReadWrite, Category = Options)
 	bool bSymmetricExtents = true;
@@ -51,22 +35,26 @@ struct GEOMETRYSCRIPTINGCORE_API FGeometryScriptTwistWarpOptions
 {
 	GENERATED_BODY()
 public:
-	/** Origin of the Bend Warp */
+	/** Symmetric extents are [-BendExtent,BendExtent], if disabled, then [-LowerExtent,BendExtent] is used  */
 	UPROPERTY(BlueprintReadWrite, Category = Options)
-	FVector TwistOrigin = FVector(0,0,0);
+	bool bSymmetricExtents = true;
 
-	/** Twist will be along this axis, ie if twisting a long tube, this the tube would twist "around" this axis */
-	UPROPERTY(BlueprintReadWrite, Category = Options)
-	FVector TwistAxis = FVector(0, 0, 1);
+	/** Lower extent used when bSymmetricExtents = false */
+	UPROPERTY(BlueprintReadWrite, Category = Options, meta = (EditCondition = "bSymmetricExtents == false"))
+	float LowerExtent = 10;
 
-	/** Twist will begin pointed towards this direction (perpendicular to TwistAxis). Automatically calculated if this vector is Zero. */
+	/** If true, the Twist is "centered" at the Origin, ie the regions on either side of the extents are rigidly transformed. If false, the Twist begins at the start of the Lower Extents, and the "lower" region is not affected. */
 	UPROPERTY(BlueprintReadWrite, Category = Options)
-	FVector TowardAxis = FVector(0, 0, 0);
+	bool bBidirectional = true;
+};
+
 
-	/** Rotation in Degrees around the Twist Axis, this is applied after the TowardAxis rotation */
-	UPROPERTY(BlueprintReadWrite, Category = Options)
-	float TwistAxisRotation = 0;
 
+USTRUCT(BlueprintType)
+struct GEOMETRYSCRIPTINGCORE_API FGeometryScriptFlareWarpOptions
+{
+	GENERATED_BODY()
+public:
 	/** Symmetric extents are [-BendExtent,BendExtent], if disabled, then [-LowerExtent,BendExtent] is used  */
 	UPROPERTY(BlueprintReadWrite, Category = Options)
 	bool bSymmetricExtents = true;
@@ -75,9 +63,9 @@ public:
 	UPROPERTY(BlueprintReadWrite, Category = Options, meta = (EditCondition = "bSymmetricExtents == false"))
 	float LowerExtent = 10;
 
-	/** If true, the Twist is "centered" at the Origin, ie the regions on either side of the extents are rigidly transformed. If false, the Twist begins at the start of the Lower Extents, and the "lower" region is not affected. */
+	/** */
 	UPROPERTY(BlueprintReadWrite, Category = Options)
-	bool bBidirectional = true;
+	bool bSmoothEnds = true;
 };
 
 
@@ -94,6 +82,7 @@ public:
 	ApplyBendWarpToMesh(  
 		UDynamicMesh* TargetMesh, 
 		FGeometryScriptBendWarpOptions Options,
+		FTransform BendOrientation,
 		float BendAngle = 45,
 		float BendExtent = 50,
 		UGeometryScriptDebug* Debug = nullptr);
@@ -103,9 +92,19 @@ public:
 	ApplyTwistWarpToMesh(  
 		UDynamicMesh* TargetMesh, 
 		FGeometryScriptTwistWarpOptions Options,
+		FTransform TwistOrientation,
 		float TwistAngle = 45,
 		float TwistExtent = 50,
 		UGeometryScriptDebug* Debug = nullptr);
 
-
+	UFUNCTION(BlueprintCallable, Category = "GeometryScript|Deformations", meta=(ScriptMethod))
+	static UPARAM(DisplayName = "Target Mesh") UDynamicMesh* 
+	ApplyFlareWarpToMesh(  
+		UDynamicMesh* TargetMesh, 
+		FGeometryScriptFlareWarpOptions Options,
+		FTransform FlareOrientation,
+		float FlarePercentX = 0,
+		float FlarePercentY = 0,
+		float FlareExtent = 50,
+		UGeometryScriptDebug* Debug = nullptr);
 };
\ No newline at end of file
diff --git a/Engine/Plugins/Experimental/GeometryScripting/Source/GeometryScriptingCore/Public/GeometryScript/MeshMaterialFunctions.h b/Engine/Plugins/Experimental/GeometryScripting/Source/GeometryScriptingCore/Public/GeometryScript/MeshMaterialFunctions.h
new file mode 100644
--- /dev/null
+++ b/Engine/Plugins/Experimental/GeometryScripting/Source/GeometryScriptingCore/Public/GeometryScript/MeshMaterialFunctions.h
@@ -0,0 +1,65 @@
+// Copyright Epic Games, Inc. All Rights Reserved.
+
+#pragma once
+
+#include "CoreMinimal.h"
+#include "Kismet/BlueprintFunctionLibrary.h"
+#include "GeometryScript/GeometryScriptTypes.h"
+#include "MeshMaterialFunctions.generated.h"
+
+class UDynamicMesh;
+
+
+UCLASS(meta = (ScriptName = "GeometryScript_Materials"))
+class GEOMETRYSCRIPTINGCORE_API UGeometryScriptLibrary_MeshMaterialFunctions : public UBlueprintFunctionLibrary
+{
+	GENERATED_BODY()
+public:
+
+	UFUNCTION(BlueprintPure, Category = "GeometryScript|MeshQueries", meta=(ScriptMethod))
+	static UPARAM(DisplayName = "Enabled") int GetMaxMaterialID( UDynamicMesh* TargetMesh, bool& bHasMaterialIDs );
+
+	UFUNCTION(BlueprintCallable, Category = "GeometryScript|Materials", meta=(ScriptMethod))
+	static UPARAM(DisplayName = "Target Mesh") UDynamicMesh* 
+	EnableMaterialIDs( UDynamicMesh* TargetMesh, UGeometryScriptDebug* Debug = nullptr );
+
+	UFUNCTION(BlueprintCallable, Category = "GeometryScript|Materials", meta=(ScriptMethod))
+	static UPARAM(DisplayName = "Target Mesh") UDynamicMesh* 
+	ClearMaterialIDs( UDynamicMesh* TargetMesh, UGeometryScriptDebug* Debug = nullptr );
+
+	UFUNCTION(BlueprintCallable, Category = "GeometryScript|Materials", meta=(ScriptMethod))
+	static UPARAM(DisplayName = "Target Mesh") UDynamicMesh* 
+	RemapMaterialIDs( 
+		UDynamicMesh* TargetMesh, 
+		int FromMaterialID,
+		int ToMaterialID,
+		UGeometryScriptDebug* Debug = nullptr );
+
+
+	UFUNCTION(BlueprintPure, Category = "GeometryScript|Materials", meta=(ScriptMethod))
+	static UPARAM(DisplayName = "Material ID") int32
+	GetTriangleMaterialID( UDynamicMesh* TargetMesh, int TriangleID, bool& bIsValidTriangle );
+
+	UFUNCTION(BlueprintPure, Category = "GeometryScript|MeshQueries", meta=(ScriptMethod))
+	static UPARAM(DisplayName = "Target Mesh") UDynamicMesh* 
+	GetAllTriangleMaterialIDs( UDynamicMesh* TargetMesh, UPARAM(ref) TArray<int>& MaterialIDs, bool& bHasMaterialIDs );
+
+
+	UFUNCTION(BlueprintCallable, Category = "GeometryScript|Materials", meta=(ScriptMethod))
+	static UPARAM(DisplayName = "Target Mesh") UDynamicMesh* 
+	SetTriangleMaterialID( 
+		UDynamicMesh* TargetMesh, 
+		int TriangleID, 
+		int MaterialID,
+		bool& bIsValidTriangle, 
+		bool bDeferChangeNotifications = false );
+
+	UFUNCTION(BlueprintCallable, Category = "GeometryScript|Materials", meta=(ScriptMethod))
+	static UPARAM(DisplayName = "Target Mesh") UDynamicMesh*
+	SetAllTriangleMaterialIDs(
+		UDynamicMesh* TargetMesh,
+		UPARAM(ref) const TArray<int32>& TriangleMaterialIDs,
+		bool bDeferChangeNotifications = false,
+		UGeometryScriptDebug* Debug = nullptr);
+
+};
\ No newline at end of file
diff --git a/Engine/Plugins/Experimental/GeometryScripting/Source/GeometryScriptingCore/Public/GeometryScript/MeshPrimitiveFunctions.h b/Engine/Plugins/Experimental/GeometryScripting/Source/GeometryScriptingCore/Public/GeometryScript/MeshPrimitiveFunctions.h
--- a/Engine/Plugins/Experimental/GeometryScripting/Source/GeometryScriptingCore/Public/GeometryScript/MeshPrimitiveFunctions.h
+++ b/Engine/Plugins/Experimental/GeometryScripting/Source/GeometryScriptingCore/Public/GeometryScript/MeshPrimitiveFunctions.h
@@ -70,9 +70,9 @@ public:
 		FGeometryScriptPrimitiveOptions PrimitiveOptions,
 		FTransform Transform,
 		float Radius = 50,
-		int32 StepsX = 0,
-		int32 StepsY = 0,
-		int32 StepsZ = 0,
+		int32 StepsX = 6,
+		int32 StepsY = 6,
+		int32 StepsZ = 6,
 		UGeometryScriptDebug* Debug = nullptr);
 
 
@@ -146,6 +146,22 @@ public:
 		int32 Steps = 8,
 		UGeometryScriptDebug* Debug = nullptr);
 
+	/**
+	 * Polygon should be oriented counter-clockwise to produce a correctly-oriented shape, otherwise it will be inside-out
+	 * Polygon endpoint is not repeated.
+	 */
+	UFUNCTION(BlueprintCallable, Category = "GeometryScript|Primitives", meta=(ScriptMethod))
+	static UPARAM(DisplayName = "Target Mesh") UDynamicMesh* 
+	AppendSimpleExtrudePolygon( 
+		UDynamicMesh* TargetMesh, 
+		FGeometryScriptPrimitiveOptions PrimitiveOptions,
+		FTransform Transform,
+		const TArray<FVector2D>& PolygonVertices,
+		float Height = 100,
+		int32 HeightSteps = 0,
+		bool bCapped = true,
+		UGeometryScriptDebug* Debug = nullptr);
+
 
 	UFUNCTION(BlueprintCallable, Category = "GeometryScript|Primitives", meta=(ScriptMethod))
 	static UPARAM(DisplayName = "Target Mesh") UDynamicMesh* 
@@ -198,6 +214,21 @@ public:
 		int32 SpokeSteps = 0,
 		float StartAngle = 0,
 		float EndAngle = 360,
+		float HoleRadius = 0,
+		UGeometryScriptDebug* Debug = nullptr);
+
+	/**
+	* Polygon should be oriented counter-clockwise to produce a correctly-oriented shape, otherwise it will be inside-out
+	* Polygon endpoint is not repeated.
+	*/
+	UFUNCTION(BlueprintCallable, Category = "GeometryScript|Primitives", meta=(ScriptMethod))
+	static UPARAM(DisplayName = "Target Mesh") UDynamicMesh* 
+	AppendTriangulatedPolygon( 
+		UDynamicMesh* TargetMesh, 
+		FGeometryScriptPrimitiveOptions PrimitiveOptions,
+		FTransform Transform,
+		const TArray<FVector2D>& PolygonVertices,
+		bool bAllowSelfIntersections = true,
 		UGeometryScriptDebug* Debug = nullptr);
 
 };
\ No newline at end of file
diff --git a/Engine/Plugins/Experimental/GeometryScripting/Source/GeometryScriptingCore/Public/GeometryScript/MeshQueryFunctions.h b/Engine/Plugins/Experimental/GeometryScripting/Source/GeometryScriptingCore/Public/GeometryScript/MeshQueryFunctions.h
--- a/Engine/Plugins/Experimental/GeometryScripting/Source/GeometryScriptingCore/Public/GeometryScript/MeshQueryFunctions.h
+++ b/Engine/Plugins/Experimental/GeometryScripting/Source/GeometryScriptingCore/Public/GeometryScript/MeshQueryFunctions.h
@@ -37,6 +37,8 @@ public:
 	UFUNCTION(BlueprintPure, Category = "GeometryScript|MeshQueries", meta=(ScriptMethod))
 	static UPARAM(DisplayName = "Num Edges") int32 GetNumOpenBorderEdges( UDynamicMesh* TargetMesh );
 
+	UFUNCTION(BlueprintPure, Category = "GeometryScript|MeshQueries", meta=(ScriptMethod))
+	static UPARAM(DisplayName = "Num Edges") int32 GetNumConnectedComponents( UDynamicMesh* TargetMesh );
 
 
 	UFUNCTION(BlueprintPure, Category = "GeometryScript|MeshQueries", meta=(ScriptMethod))
@@ -79,7 +81,8 @@ public:
 	static UPARAM(DisplayName = "Position") FVector GetVertexPosition( UDynamicMesh* TargetMesh, int32 VertexID, bool& bIsValidVertex );
 
 	UFUNCTION(BlueprintPure, Category = "GeometryScript|MeshQueries", meta=(ScriptMethod))
-	static UPARAM(DisplayName = "Positions") TArray<FVector> GetAllVertexPositions( UDynamicMesh* TargetMesh, bool& bHasVertexIDGaps );
+	static UPARAM(DisplayName = "Target Mesh") UDynamicMesh* 
+	GetAllVertexPositions( UDynamicMesh* TargetMesh,  UPARAM(ref) TArray<FVector>& VertexPositions, bool& bHasVertexIDGaps );
 
 
 
@@ -93,4 +96,8 @@ public:
 	static void GetTriangleUVs( UDynamicMesh* TargetMesh, int32 UVSetIndex, int32 TriangleID, FVector2D& UV1, FVector2D& UV2, FVector2D& UV3, bool& bHaveValidUVs );
 
 
+
+	UFUNCTION(BlueprintPure, Category = "GeometryScript|MeshQueries", meta=(ScriptMethod))
+	static UPARAM(DisplayName = "Enabled") bool GetHasMaterialIDs( UDynamicMesh* TargetMesh );
+
 };
\ No newline at end of file
diff --git a/Engine/Plugins/Experimental/GeometryScripting/Source/GeometryScriptingCore/Public/GeometryScript/MeshSpatialFunctions.h b/Engine/Plugins/Experimental/GeometryScripting/Source/GeometryScriptingCore/Public/GeometryScript/MeshSpatialFunctions.h
new file mode 100644
--- /dev/null
+++ b/Engine/Plugins/Experimental/GeometryScripting/Source/GeometryScriptingCore/Public/GeometryScript/MeshSpatialFunctions.h
@@ -0,0 +1,122 @@
+// Copyright Epic Games, Inc. All Rights Reserved.
+
+#pragma once
+
+#include "CoreMinimal.h"
+#include "Kismet/BlueprintFunctionLibrary.h"
+#include "GeometryScript/GeometryScriptTypes.h"
+#include "MeshSpatialFunctions.generated.h"
+
+class UDynamicMesh;
+
+
+USTRUCT(BlueprintType, meta = (DisplayName = "BVH Query Options"))
+struct GEOMETRYSCRIPTINGCORE_API FGeometryScriptSpatialQueryOptions
+{
+	GENERATED_BODY()
+public:
+	UPROPERTY(BlueprintReadWrite, Category = Options)
+	float MaxDistance = 0;
+
+	UPROPERTY(BlueprintReadWrite, Category = Options)
+	bool bAllowUnsafeModifiedQueries = false;
+
+	UPROPERTY(BlueprintReadWrite, Category = Options)
+	float WindingIsoThreshold = 0.5;
+};
+
+
+
+
+
+USTRUCT(BlueprintType, meta = (DisplayName = "Ray Hit Result"))
+struct GEOMETRYSCRIPTINGCORE_API FGeometryScriptRayHitResult
+{
+	GENERATED_BODY()
+public:
+	UPROPERTY(BlueprintReadWrite, Category = Options)
+	bool bHit = false;
+
+	UPROPERTY(BlueprintReadWrite, Category = Options)
+	float RayParameter = 0;
+
+	UPROPERTY(BlueprintReadWrite, Category = Options)
+	int HitTriangleID = -1;
+
+	UPROPERTY(BlueprintReadWrite, Category = Options)
+	FVector HitPosition;
+
+	UPROPERTY(BlueprintReadWrite, Category = Options)
+	FVector HitBaryCoords;
+};
+
+
+
+UCLASS(meta = (ScriptName = "GeometryScript_MeshSpatial"))
+class GEOMETRYSCRIPTINGCORE_API UGeometryScriptLibrary_MeshSpatial : public UBlueprintFunctionLibrary
+{
+	GENERATED_BODY()
+public:
+	UFUNCTION(BlueprintCallable, Category = "GeometryScript|Spatial", meta = (ScriptMethod))
+	static void ResetBVH(UPARAM(ref) FGeometryScriptDynamicMeshBVH& ResetBVH);
+
+	UFUNCTION(BlueprintCallable, Category = "GeometryScript|Spatial", meta=(ScriptMethod))
+	static UPARAM(DisplayName = "Target Mesh") UDynamicMesh* 
+	BuildBVHForMesh(
+		UDynamicMesh* TargetMesh,
+		FGeometryScriptDynamicMeshBVH& OutputBVH,
+		UGeometryScriptDebug* Debug = nullptr );
+
+	UFUNCTION(BlueprintCallable, Category = "GeometryScript|Spatial", meta=(ScriptMethod))
+	static UPARAM(DisplayName = "Target Mesh") UDynamicMesh* 
+	IsBVHValidForMesh(
+		UDynamicMesh* TargetMesh,
+		UPARAM(ref) const FGeometryScriptDynamicMeshBVH& TestBVH,
+		bool& bIsValid,
+		UGeometryScriptDebug* Debug = nullptr );
+
+	UFUNCTION(BlueprintCallable, Category = "GeometryScript|Spatial", meta=(ScriptMethod))
+	static UPARAM(DisplayName = "Target Mesh") UDynamicMesh* 
+	RebuildBVHForMesh(
+		UDynamicMesh* TargetMesh,
+		UPARAM(ref) FGeometryScriptDynamicMeshBVH& UpdateBVH,
+		bool bOnlyIfInvalid = true,
+		UGeometryScriptDebug* Debug = nullptr );
+
+	UFUNCTION(BlueprintCallable, Category = "GeometryScript|Spatial", meta=(ScriptMethod, ExpandEnumAsExecs = "Outcome"))
+	static UPARAM(DisplayName = "Target Mesh") UDynamicMesh* 
+	FindNearestPointOnMesh(
+		UDynamicMesh* TargetMesh,
+		UPARAM(ref) const FGeometryScriptDynamicMeshBVH& QueryBVH,
+		FVector QueryPoint,
+		FGeometryScriptSpatialQueryOptions Options,
+		FGeometryScriptTrianglePoint& NearestResult,
+		TEnumAsByte<EGeometryScriptSearchOutcomePins>& Outcome,
+		UGeometryScriptDebug* Debug = nullptr );
+
+
+	UFUNCTION(BlueprintCallable, Category = "GeometryScript|Spatial", meta=(ScriptMethod, ExpandEnumAsExecs = "Outcome"))
+	static UPARAM(DisplayName = "Target Mesh") UDynamicMesh* 
+	FindNearestRayIntersectionWithMesh(
+		UDynamicMesh* TargetMesh,
+		UPARAM(ref) const FGeometryScriptDynamicMeshBVH& QueryBVH,
+		FVector RayOrigin,
+		FVector RayDirection,
+		FGeometryScriptSpatialQueryOptions Options,
+		FGeometryScriptRayHitResult& HitResult,
+		TEnumAsByte<EGeometryScriptSearchOutcomePins>& Outcome,
+		UGeometryScriptDebug* Debug = nullptr );
+
+
+	UFUNCTION(BlueprintCallable, Category = "GeometryScript|Spatial", meta=(ScriptMethod, ExpandEnumAsExecs = "Outcome"))
+	static UPARAM(DisplayName = "Target Mesh") UDynamicMesh* 
+	IsPointInsideMesh(
+		UDynamicMesh* TargetMesh,
+		UPARAM(ref) const FGeometryScriptDynamicMeshBVH& QueryBVH,
+		FVector QueryPoint,
+		FGeometryScriptSpatialQueryOptions Options,
+		bool& bIsInside,
+		TEnumAsByte<EGeometryScriptContainmentOutcomePins>& Outcome,
+		UGeometryScriptDebug* Debug = nullptr );
+
+};
\ No newline at end of file
diff --git a/Engine/Plugins/Experimental/GeometryScripting/Source/GeometryScriptingCore/Public/GeometryScript/MeshUVFunctions.h b/Engine/Plugins/Experimental/GeometryScripting/Source/GeometryScriptingCore/Public/GeometryScript/MeshUVFunctions.h
--- a/Engine/Plugins/Experimental/GeometryScripting/Source/GeometryScriptingCore/Public/GeometryScript/MeshUVFunctions.h
+++ b/Engine/Plugins/Experimental/GeometryScripting/Source/GeometryScriptingCore/Public/GeometryScript/MeshUVFunctions.h
@@ -10,12 +10,31 @@
 class UDynamicMesh;
 
 
+USTRUCT(BlueprintType)
+struct GEOMETRYSCRIPTINGCORE_API FGeometryScriptRepackUVsOptions
+{
+	GENERATED_BODY()
+public:
+	UPROPERTY(BlueprintReadWrite, Category = Options)
+	int TargetImageWidth = 512;
+
+	UPROPERTY(BlueprintReadWrite, Category = Options)
+	bool bOptimizeIslandRotation = true;
+};
+
+
 UCLASS(meta = (ScriptName = "GeometryScript_UVs"))
 class GEOMETRYSCRIPTINGCORE_API UGeometryScriptLibrary_MeshUVFunctions : public UBlueprintFunctionLibrary
 {
 	GENERATED_BODY()
 public:
 
+	UFUNCTION(BlueprintCallable, Category = "GeometryScript|UVs", meta=(ScriptMethod))
+	static UPARAM(DisplayName = "Target Mesh") UDynamicMesh* 
+	SetNumUVSets( 
+		UDynamicMesh* TargetMesh, 
+		int NumUVSets,
+		UGeometryScriptDebug* Debug = nullptr );
 
 
 	UFUNCTION(BlueprintCallable, Category = "GeometryScript|UVs", meta=(ScriptMethod))
@@ -28,4 +47,62 @@ public:
 		bool& bIsValidTriangle, 
 		bool bDeferChangeNotifications = false );
 
+
+	UFUNCTION(BlueprintCallable, Category = "GeometryScript|UVs", meta=(ScriptMethod))
+	static UPARAM(DisplayName = "Target Mesh") UDynamicMesh* 
+	TranslateMeshUVs( 
+		UDynamicMesh* TargetMesh, 
+		int UVSetIndex,
+		FVector2D Translation,
+		UGeometryScriptDebug* Debug = nullptr );
+
+	UFUNCTION(BlueprintCallable, Category = "GeometryScript|UVs", meta=(ScriptMethod))
+	static UPARAM(DisplayName = "Target Mesh") UDynamicMesh* 
+	ScaleMeshUVs( 
+		UDynamicMesh* TargetMesh, 
+		int UVSetIndex,
+		FVector2D Scale,
+		FVector2D ScaleOrigin,
+		UGeometryScriptDebug* Debug = nullptr );
+
+	UFUNCTION(BlueprintCallable, Category = "GeometryScript|UVs", meta=(ScriptMethod))
+	static UPARAM(DisplayName = "Target Mesh") UDynamicMesh* 
+	RotateMeshUVs( 
+		UDynamicMesh* TargetMesh, 
+		int UVSetIndex,
+		float RotationAngle,
+		FVector2D RotationOrigin,
+		UGeometryScriptDebug* Debug = nullptr );
+
+
+	/**
+	 * Scale of PlaneTransform defines world-space dimension that maps to 1 UV dimension
+	 */
+	UFUNCTION(BlueprintCallable, Category = "GeometryScript|UVs", meta=(ScriptMethod))
+	static UPARAM(DisplayName = "Target Mesh") UDynamicMesh* 
+	SetMeshUVsFromPlanarProjection( 
+		UDynamicMesh* TargetMesh, 
+		int UVSetIndex,
+		FTransform PlaneTransform,
+		UGeometryScriptDebug* Debug = nullptr );
+
+
+	UFUNCTION(BlueprintCallable, Category = "GeometryScript|UVs", meta=(ScriptMethod))
+	static UPARAM(DisplayName = "Target Mesh") UDynamicMesh* 
+	SetMeshUVsFromBoxProjection( 
+		UDynamicMesh* TargetMesh, 
+		int UVSetIndex,
+		FTransform BoxTransform,
+		int MinIslandTriCount = 2,
+		UGeometryScriptDebug* Debug = nullptr );
+
+
+	UFUNCTION(BlueprintCallable, Category = "GeometryScript|UVs", meta=(ScriptMethod))
+	static UPARAM(DisplayName = "Target Mesh") UDynamicMesh* 
+	RepackMeshUVs( 
+		UDynamicMesh* TargetMesh, 
+		int UVSetIndex,
+		FGeometryScriptRepackUVsOptions RepackOptions,
+		UGeometryScriptDebug* Debug = nullptr );
+
 };
\ No newline at end of file
diff --git a/Engine/Plugins/Experimental/GeometryScripting/Source/GeometryScriptingCore/Public/GeometryScript/MeshVertexColorFunctions.h b/Engine/Plugins/Experimental/GeometryScripting/Source/GeometryScriptingCore/Public/GeometryScript/MeshVertexColorFunctions.h
--- a/Engine/Plugins/Experimental/GeometryScripting/Source/GeometryScriptingCore/Public/GeometryScript/MeshVertexColorFunctions.h
+++ b/Engine/Plugins/Experimental/GeometryScripting/Source/GeometryScriptingCore/Public/GeometryScript/MeshVertexColorFunctions.h
@@ -21,7 +21,7 @@ public:
 	static UPARAM(DisplayName = "Target Mesh") UDynamicMesh*
 	SetMeshPerVertexColors(
 		UDynamicMesh* TargetMesh,
-		const TArray<FLinearColor>& VertexColors,
+		UPARAM(ref) const TArray<FLinearColor>& VertexColors,
 		UGeometryScriptDebug* Debug = nullptr);
 
 };
diff --git a/Engine/Plugins/Experimental/GeometryScripting/Source/GeometryScriptingEditor/GeometryScriptingEditor.Build.cs b/Engine/Plugins/Experimental/GeometryScripting/Source/GeometryScriptingEditor/GeometryScriptingEditor.Build.cs
new file mode 100644
--- /dev/null
+++ b/Engine/Plugins/Experimental/GeometryScripting/Source/GeometryScriptingEditor/GeometryScriptingEditor.Build.cs
@@ -0,0 +1,65 @@
+// Copyright Epic Games, Inc. All Rights Reserved.
+
+using UnrealBuildTool;
+
+public class GeometryScriptingEditor : ModuleRules
+{
+	public GeometryScriptingEditor(ReadOnlyTargetRules Target) : base(Target)
+	{
+		PCHUsage = ModuleRules.PCHUsageMode.UseExplicitOrSharedPCHs;
+
+		PublicIncludePaths.AddRange(
+			new string[] {
+				// ... add public include paths required here ...
+			}
+			);
+
+
+		PrivateIncludePaths.AddRange(
+			new string[] {
+				// ... add other private include paths required here ...
+			}
+			);
+
+
+		PublicDependencyModuleNames.AddRange(
+			new string[]
+			{
+				"Core",
+				"PhysicsCore",
+				"RenderCore",
+				"GeometryCore",
+				"GeometryFramework",
+				"DynamicMesh",
+				"GeometryScriptingCore"
+			}
+			);
+
+
+		PrivateDependencyModuleNames.AddRange(
+			new string[]
+			{
+				"CoreUObject",
+				"Engine",
+				"MeshDescription",
+				"StaticMeshDescription",
+				"MeshConversion",
+				"GeometryAlgorithms",
+				"ModelingOperators",
+				"ModelingComponentsEditorOnly",
+				"EditorFramework",
+				"UnrealEd",
+				"BSPUtils"
+			}
+			);
+
+
+		DynamicallyLoadedModuleNames.AddRange(
+			new string[]
+			{
+				// ... add any modules that your module loads dynamically here ...
+			}
+			);
+
+	}
+}
diff --git a/Engine/Plugins/Experimental/GeometryScripting/Source/GeometryScriptingEditor/Private/CreateNewAssetUtilityFunctions.cpp b/Engine/Plugins/Experimental/GeometryScripting/Source/GeometryScriptingEditor/Private/CreateNewAssetUtilityFunctions.cpp
new file mode 100644
--- /dev/null
+++ b/Engine/Plugins/Experimental/GeometryScripting/Source/GeometryScriptingEditor/Private/CreateNewAssetUtilityFunctions.cpp
@@ -0,0 +1,175 @@
+// Copyright Epic Games, Inc. All Rights Reserved.
+
+#include "GeometryScript/CreateNewAssetUtilityFunctions.h"
+
+#include "DynamicMesh/DynamicMesh3.h"
+#include "UDynamicMesh.h"
+
+#include "Engine/StaticMesh.h"
+#include "Engine/SkeletalMesh.h"
+#include "Engine/StaticMeshActor.h"
+#include "StaticMeshAttributes.h"
+#include "StaticMeshOperations.h"
+#include "RenderingThread.h"
+
+#include "Engine/BlockingVolume.h"
+#include "Components/BrushComponent.h"
+#include "Engine/Polys.h"
+#include "Model.h"
+#include "BSPOps.h"		// in UnrealEd
+#include "Editor/EditorEngine.h"		// for FActorLabelUtilities
+#include "Editor.h"
+
+#include "MeshDescriptionToDynamicMesh.h"
+#include "DynamicMeshToMeshDescription.h"
+#include "ConversionUtils/DynamicMeshToVolume.h"
+#include "AssetUtils/CreateStaticMeshUtil.h"
+
+#include "ExplicitUseGeometryMathTypes.h"		// using UE::Geometry::(math types)
+using namespace UE::Geometry;
+
+#define LOCTEXT_NAMESPACE "UGeometryScriptLibrary_MeshAssetFunctions"
+
+
+AVolume* UGeometryScriptLibrary_CreateNewAssetFunctions::CreateNewVolumeFromMesh(
+	UDynamicMesh* FromDynamicMesh,
+	UWorld* CreateInWorld,
+	FTransform ActorTransform,
+	FString BaseActorName,
+	FGeometryScriptCreateNewVolumeFromMeshOptions Options,
+	TEnumAsByte<EGeometryScriptOutcomePins>& Outcome,
+	UGeometryScriptDebug* Debug)
+{
+	Outcome = EGeometryScriptOutcomePins::Failure;
+
+	if (FromDynamicMesh == nullptr)
+	{
+		UE::Geometry::AppendError(Debug, EGeometryScriptErrorType::InvalidInputs, LOCTEXT("CreateNewVolumeFromMesh_InvalidInput1", "CreateNewVolumeFromMesh: FromDynamicMesh is Null"));
+		return nullptr;
+	}
+	if (CreateInWorld == nullptr)
+	{
+		UE::Geometry::AppendError(Debug, EGeometryScriptErrorType::InvalidInputs, LOCTEXT("CreateNewVolumeFromMesh_InvalidInput2", "CreateNewVolumeFromMesh: CreateInWorld is Null"));
+		return nullptr;
+	}	
+	if (FromDynamicMesh->GetTriangleCount() < 4)
+	{
+		UE::Geometry::AppendError(Debug, EGeometryScriptErrorType::InvalidInputs, LOCTEXT("CreateNewVolumeFromMesh_InvalidInput3", "CreateNewVolumeFromMesh: FromDynamicMesh does not define a valid Volume"));
+		return nullptr;
+	}
+	// todo: other safety checks
+
+	// spawn new actor
+
+	UClass* VolumeClass = ABlockingVolume::StaticClass();
+	if (Options.VolumeType != nullptr 
+		&& Cast<AVolume>(Options.VolumeType->GetDefaultObject(false)) != nullptr )
+	{
+		VolumeClass = Options.VolumeType;
+	}
+
+	GEditor->BeginTransaction(LOCTEXT("CreateNewVolumeFromMesh_Transaction", "Create Volume"));
+
+	FActorSpawnParameters SpawnInfo;
+	FTransform NewActorTransform = FTransform::Identity;
+	AVolume* NewVolumeActor = (AVolume*)CreateInWorld->SpawnActor(VolumeClass, &NewActorTransform, SpawnInfo);
+
+	NewVolumeActor->BrushType = EBrushType::Brush_Add;
+	UModel* Model = NewObject<UModel>(NewVolumeActor);
+	NewVolumeActor->Brush = Model;
+	NewVolumeActor->GetBrushComponent()->Brush = NewVolumeActor->Brush;
+
+	UE::Conversion::FMeshToVolumeOptions ConvertOptions;
+	ConvertOptions.bAutoSimplify = true;
+	ConvertOptions.MaxTriangles = FMath::Max(1, Options.MaxTriangles);
+
+	FromDynamicMesh->ProcessMesh([&](const FDynamicMesh3& ReadMesh)
+	{
+		UE::Conversion::DynamicMeshToVolume(ReadMesh, NewVolumeActor, ConvertOptions);
+	});
+
+	NewVolumeActor->SetActorTransform(ActorTransform);
+	FActorLabelUtilities::SetActorLabelUnique(NewVolumeActor, BaseActorName);
+	NewVolumeActor->PostEditChange();
+
+	GEditor->EndTransaction();
+
+	Outcome = EGeometryScriptOutcomePins::Success;
+	return NewVolumeActor;
+}
+
+
+
+
+
+
+UStaticMesh* UGeometryScriptLibrary_CreateNewAssetFunctions::CreateNewStaticMeshAssetFromMesh(
+	UDynamicMesh* FromDynamicMesh, 
+	FString AssetPathAndName,
+	FGeometryScriptCreateNewStaticMeshAssetOptions Options,
+	TEnumAsByte<EGeometryScriptOutcomePins>& Outcome,
+	UGeometryScriptDebug* Debug)
+{
+	Outcome = EGeometryScriptOutcomePins::Failure;
+	if (FromDynamicMesh == nullptr)
+	{
+		UE::Geometry::AppendError(Debug, EGeometryScriptErrorType::InvalidInputs, LOCTEXT("CreateNewStaticMeshAssetFromMesh_InvalidInput1", "CreateNewStaticMeshAssetFromMesh: FromDynamicMesh is Null"));
+		return nullptr;
+	}
+	if (FromDynamicMesh->GetTriangleCount() == 0)
+	{
+		UE::Geometry::AppendError(Debug, EGeometryScriptErrorType::InvalidInputs, LOCTEXT("CreateNewStaticMeshAssetFromMesh_InvalidInput3", "CreateNewStaticMeshAssetFromMesh: FromDynamicMesh has zero triangles"));
+		return nullptr;
+	}
+	// todo: other safety checks
+
+	GEditor->BeginTransaction(LOCTEXT("CreateNewStaticMeshAssetFromMesh_Transaction", "Create StaticMesh"));
+
+	UE::AssetUtils::FStaticMeshAssetOptions AssetOptions;
+	AssetOptions.NewAssetPath = AssetPathAndName;
+
+	AssetOptions.NumSourceModels = 1;
+
+	// CreateStaticMeshAsset below will handle this, but we could allow passing in materials as an option...
+	//AssetOptions.NumMaterialSlots = CreateMeshParams.Materials.Num();
+	//AssetOptions.AssetMaterials = (CreateMeshParams.AssetMaterials.Num() == AssetOptions.NumMaterialSlots) ?
+	//	FilterMaterials(CreateMeshParams.AssetMaterials) : FilterMaterials(CreateMeshParams.Materials);
+
+	AssetOptions.bEnableRecomputeNormals = Options.bEnableRecomputeNormals;
+	AssetOptions.bEnableRecomputeTangents = Options.bEnableRecomputeTangents;
+	AssetOptions.bGenerateNaniteEnabledMesh = Options.bEnableNanite;
+	AssetOptions.NaniteProxyTrianglePercent = Options.NaniteProxyTrianglePercent;
+
+	AssetOptions.bCreatePhysicsBody = Options.bEnableCollision;
+	AssetOptions.CollisionType = Options.CollisionMode;
+
+	FDynamicMesh3 LODMesh;
+	FromDynamicMesh->ProcessMesh([&](const FDynamicMesh3& ReadMesh)
+	{
+		LODMesh = ReadMesh;
+	});
+	AssetOptions.SourceMeshes.DynamicMeshes.Add(&LODMesh);
+
+	UE::AssetUtils::FStaticMeshResults ResultData;
+	UE::AssetUtils::ECreateStaticMeshResult AssetResult = UE::AssetUtils::CreateStaticMeshAsset(AssetOptions, ResultData);
+
+	if (AssetResult != UE::AssetUtils::ECreateStaticMeshResult::Ok)
+	{
+		UE::Geometry::AppendError(Debug, EGeometryScriptErrorType::OperationFailed, LOCTEXT("CreateNewStaticMeshAssetFromMesh_Failed", "CreateNewStaticMeshAssetFromMesh: Failed to create new Asset"));
+		return nullptr;
+	}
+
+	UStaticMesh* NewStaticMesh = ResultData.StaticMesh;
+	NewStaticMesh->PostEditChange();
+
+	GEditor->EndTransaction();
+
+	Outcome = EGeometryScriptOutcomePins::Success;
+	return NewStaticMesh;
+}
+
+
+
+
+
+#undef LOCTEXT_NAMESPACE
\ No newline at end of file
diff --git a/Engine/Plugins/Experimental/GeometryScripting/Source/GeometryScriptingEditor/Private/GeometryScriptingCoreModule.cpp b/Engine/Plugins/Experimental/GeometryScripting/Source/GeometryScriptingEditor/Private/GeometryScriptingCoreModule.cpp
new file mode 100644
--- /dev/null
+++ b/Engine/Plugins/Experimental/GeometryScripting/Source/GeometryScriptingEditor/Private/GeometryScriptingCoreModule.cpp
@@ -0,0 +1,20 @@
+// Copyright Epic Games, Inc. All Rights Reserved.
+
+#include "GeometryScriptingEditorModule.h"
+
+#define LOCTEXT_NAMESPACE "FGeometryScriptingEditorModule"
+
+void FGeometryScriptingEditorModule::StartupModule()
+{
+	// This code will execute after your module is loaded into memory; the exact timing is specified in the .uplugin file per-module
+}
+
+void FGeometryScriptingEditorModule::ShutdownModule()
+{
+	// This function may be called during shutdown to clean up your module.  For modules that support dynamic reloading,
+	// we call this function before unloading the module.
+}
+
+#undef LOCTEXT_NAMESPACE
+	
+IMPLEMENT_MODULE(FGeometryScriptingEditorModule, GeometryScriptingEditor)
\ No newline at end of file
diff --git a/Engine/Plugins/Experimental/GeometryScripting/Source/GeometryScriptingEditor/Public/GeometryScript/CreateNewAssetUtilityFunctions.h b/Engine/Plugins/Experimental/GeometryScripting/Source/GeometryScriptingEditor/Public/GeometryScript/CreateNewAssetUtilityFunctions.h
new file mode 100644
--- /dev/null
+++ b/Engine/Plugins/Experimental/GeometryScripting/Source/GeometryScriptingEditor/Public/GeometryScript/CreateNewAssetUtilityFunctions.h
@@ -0,0 +1,88 @@
+// Copyright Epic Games, Inc. All Rights Reserved.
+
+#pragma once
+
+#include "CoreMinimal.h"
+#include "Kismet/BlueprintFunctionLibrary.h"
+#include "GeometryScript/GeometryScriptTypes.h"
+#include "Engine/BlockingVolume.h"
+#include "BodySetupEnums.h"
+#include "CreateNewAssetUtilityFunctions.generated.h"
+
+class UStaticMesh;
+class UDynamicMesh;
+class AVolume;
+
+
+
+USTRUCT(BlueprintType)
+struct GEOMETRYSCRIPTINGEDITOR_API FGeometryScriptCreateNewVolumeFromMeshOptions
+{
+	GENERATED_BODY()
+public:
+	UPROPERTY(BlueprintReadWrite, Category = Options)
+	TSubclassOf<class AVolume> VolumeType = ABlockingVolume::StaticClass();
+
+	UPROPERTY(BlueprintReadWrite, Category = Options)
+	bool bAutoSimplify = true;
+
+	UPROPERTY(BlueprintReadWrite, Category = Options)
+	int32 MaxTriangles = 250;
+};
+
+
+USTRUCT(BlueprintType)
+struct GEOMETRYSCRIPTINGEDITOR_API FGeometryScriptCreateNewStaticMeshAssetOptions
+{
+	GENERATED_BODY()
+public:
+	UPROPERTY(BlueprintReadWrite, Category = Options)
+	bool bEnableRecomputeNormals = false;
+
+	UPROPERTY(BlueprintReadWrite, Category = Options)
+	bool bEnableRecomputeTangents = false;
+
+	UPROPERTY(BlueprintReadWrite, Category = Options)
+	bool bEnableNanite = false;
+
+	UPROPERTY(BlueprintReadWrite, Category = Options)
+	float NaniteProxyTrianglePercent = 0.0f;
+
+	UPROPERTY(BlueprintReadWrite, Category = Options)
+	bool bEnableCollision = true;
+
+	UPROPERTY(BlueprintReadWrite, Category = Options)
+	TEnumAsByte<ECollisionTraceFlag> CollisionMode = ECollisionTraceFlag::CTF_UseDefault;
+};
+
+
+
+
+UCLASS(meta = (ScriptName = "GeometryScript_NewAssetUtils"))
+class GEOMETRYSCRIPTINGEDITOR_API UGeometryScriptLibrary_CreateNewAssetFunctions : public UBlueprintFunctionLibrary
+{
+	GENERATED_BODY()
+public:
+
+	UFUNCTION(BlueprintCallable, Category = "GeometryScript|AssetManagement", meta = (ExpandEnumAsExecs = "Outcome"))
+	static UPARAM(DisplayName = "Volume Actor") AVolume* 
+	CreateNewVolumeFromMesh(
+		UDynamicMesh* FromDynamicMesh, 
+		UPARAM(ref) UWorld* CreateInWorld,
+		FTransform ActorTransform,
+		FString BaseActorName,
+		FGeometryScriptCreateNewVolumeFromMeshOptions Options,
+		TEnumAsByte<EGeometryScriptOutcomePins>& Outcome,
+		UGeometryScriptDebug* Debug = nullptr);
+
+	UFUNCTION(BlueprintCallable, Category = "GeometryScript|AssetManagement", meta = (ExpandEnumAsExecs = "Outcome"))
+	static UPARAM(DisplayName = "Static Mesh Asset") UStaticMesh* 
+	CreateNewStaticMeshAssetFromMesh(
+		UDynamicMesh* FromDynamicMesh, 
+		FString AssetPathAndName,
+		FGeometryScriptCreateNewStaticMeshAssetOptions Options,
+		TEnumAsByte<EGeometryScriptOutcomePins>& Outcome,
+		UGeometryScriptDebug* Debug = nullptr);
+
+};
+
diff --git a/Engine/Plugins/Experimental/GeometryScripting/Source/GeometryScriptingEditor/Public/GeometryScriptingEditorModule.h b/Engine/Plugins/Experimental/GeometryScripting/Source/GeometryScriptingEditor/Public/GeometryScriptingEditorModule.h
new file mode 100644
--- /dev/null
+++ b/Engine/Plugins/Experimental/GeometryScripting/Source/GeometryScriptingEditor/Public/GeometryScriptingEditorModule.h
@@ -0,0 +1,17 @@
+// Copyright Epic Games, Inc. All Rights Reserved.
+
+#pragma once
+
+#include "CoreMinimal.h"
+#include "Modules/ModuleManager.h"
+
+class FGeometryScriptingEditorModule : public IModuleInterface
+{
+public:
+
+	/** IModuleInterface implementation */
+	virtual void StartupModule() override;
+	virtual void ShutdownModule() override;
+};
+
+
