GeometryScripting: ongoing evolution
- CopyMeshToStaticMesh now supports control of relevant build options on output, prefers to save incoming normals/tangents by default
- added ApplyMeshPlaneSlice
- added ApplyMathWarpToMesh (only sin waves currently), ApplyPerlinNoiseToMesh, ApplyIterativeSmoothingToMesh, ApplyDisplaceFromTextureMap
- added ComputeTangents
- added GetHasPolygroups, GetNumExtendedPolygroupLayers, EnablePolygroups, SetNumExtendedPolygroupLayers, ClearPolygroups, CopyPolygroupsLayer, ConvertUVIslandsToPolygroups, ComputePolygroupsFromAngleThreshold
- added AppendRevolvePath, converted various Primitive box/box2d inputs to separate floats, added Origin option to relevant primitives
- added ApplyRecursivePNTessellation, ApplyPolygroupCatmullClarkSubD, ApplyTriangleLoopSubD
- added SetMeshConstantVertexColor
- added CopyUVSet, SetMeshUVsFromCylinderProjection, RecomputeMeshUVs, AutoGeneratePatchBuilderMeshUVs, AutoGenerateXAtlasMeshUVs
#rnx

diff --git a/Engine/Plugins/Experimental/GeometryScripting/Source/GeometryScriptingCore/GeometryScriptingCore.Build.cs b/Engine/Plugins/Experimental/GeometryScripting/Source/GeometryScriptingCore/GeometryScriptingCore.Build.cs
--- a/Engine/Plugins/Experimental/GeometryScripting/Source/GeometryScriptingCore/GeometryScriptingCore.Build.cs
+++ b/Engine/Plugins/Experimental/GeometryScripting/Source/GeometryScriptingCore/GeometryScriptingCore.Build.cs
@@ -44,7 +44,8 @@ public class GeometryScriptingCore : ModuleRules
 				"StaticMeshDescription",
 				"MeshConversion",
 				"GeometryAlgorithms",
-				"ModelingOperators"
+				"ModelingOperators",
+				"ModelingComponents"
 			}
 			);
 
diff --git a/Engine/Plugins/Experimental/GeometryScripting/Source/GeometryScriptingCore/Private/MeshAssetFunctions.cpp b/Engine/Plugins/Experimental/GeometryScripting/Source/GeometryScriptingCore/Private/MeshAssetFunctions.cpp
--- a/Engine/Plugins/Experimental/GeometryScripting/Source/GeometryScriptingCore/Private/MeshAssetFunctions.cpp
+++ b/Engine/Plugins/Experimental/GeometryScripting/Source/GeometryScriptingCore/Private/MeshAssetFunctions.cpp
@@ -162,14 +162,28 @@ UDynamicMesh*  UGeometryScriptLibrary_StaticMeshFunctions::CopyMeshToStaticMesh(
 	ToStaticMeshAsset->SetFlags(RF_Transactional);
 	// mark as modified
 	ToStaticMeshAsset->Modify();
-	// mark mesh description for modify
-	ToStaticMeshAsset->ModifyMeshDescription(UseLODIndex);
+
+	if (ToStaticMeshAsset->GetNumSourceModels() < UseLODIndex+1)
+	{
+		ToStaticMeshAsset->SetNumSourceModels(UseLODIndex+1);
+	}
+
+	// configure build settings from options
+	FStaticMeshSourceModel& LODSourceModel = ToStaticMeshAsset->GetSourceModel(UseLODIndex);
+	FMeshBuildSettings& BuildSettings = LODSourceModel.BuildSettings;
+	BuildSettings.bRecomputeNormals = Options.bEnableRecomputeNormals;
+	BuildSettings.bRecomputeTangents = Options.bEnableRecomputeTangents;
+	BuildSettings.bRemoveDegenerates = Options.bEnableRemoveDegenerates;
 
 	FMeshDescription* MeshDescription = ToStaticMeshAsset->GetMeshDescription(UseLODIndex);
 	if (MeshDescription == nullptr)
 	{
 		MeshDescription = ToStaticMeshAsset->CreateMeshDescription(UseLODIndex);
 	}
+
+	// mark mesh description for modify
+	ToStaticMeshAsset->ModifyMeshDescription(UseLODIndex);
+
 	if (!ensure(MeshDescription != nullptr))
 	{
 		UE::Geometry::AppendError(Debug, EGeometryScriptErrorType::InvalidInputs,
@@ -178,15 +192,10 @@ UDynamicMesh*  UGeometryScriptLibrary_StaticMeshFunctions::CopyMeshToStaticMesh(
 	}
 
 	FConversionToMeshDescriptionOptions ConversionOptions;
-	//if (FromDynamicMesh->GetMesh()->HasVertexColors())
-	//{
-	//	ConversionOptions.bUpdateVtxColors = true;		// is this still relevant now that we have vertex color attribute?
-	//}
-
 	FDynamicMeshToMeshDescription Converter(ConversionOptions);
 	FromDynamicMesh->ProcessMesh([&](const FDynamicMesh3& ReadMesh)
 	{
-		Converter.Convert(&ReadMesh, *MeshDescription);
+		Converter.Convert(&ReadMesh, *MeshDescription, !BuildSettings.bRecomputeTangents);
 	});
 
 	ToStaticMeshAsset->CommitMeshDescription(UseLODIndex);
diff --git a/Engine/Plugins/Experimental/GeometryScripting/Source/GeometryScriptingCore/Private/MeshBooleanFunctions.cpp b/Engine/Plugins/Experimental/GeometryScripting/Source/GeometryScriptingCore/Private/MeshBooleanFunctions.cpp
--- a/Engine/Plugins/Experimental/GeometryScripting/Source/GeometryScriptingCore/Private/MeshBooleanFunctions.cpp
+++ b/Engine/Plugins/Experimental/GeometryScripting/Source/GeometryScriptingCore/Private/MeshBooleanFunctions.cpp
@@ -181,8 +181,7 @@ UDynamicMesh* UGeometryScriptLibrary_MeshBooleanFunctions::ApplyMeshSelfUnion(
 
 UDynamicMesh* UGeometryScriptLibrary_MeshBooleanFunctions::ApplyMeshPlaneCut(
 	UDynamicMesh* TargetMesh,
-	FVector CutPlaneOrigin,
-	FVector CutPlaneNormal,
+	FTransform CutFrame,
 	FGeometryScriptMeshPlaneCutOptions Options,
 	UGeometryScriptDebug* Debug)
 {
@@ -192,21 +191,19 @@ UDynamicMesh* UGeometryScriptLibrary_MeshBooleanFunctions::ApplyMeshPlaneCut(
 		return TargetMesh;
 	}
 
+	FFrame3d UseCutFrame(CutFrame);
+	FVector3d CutPlaneOrigin = UseCutFrame.Origin;
+	FVector3d UseNormal = UseCutFrame.Z();
 	if (Options.bFlipCutSide)
 	{
-		CutPlaneNormal = -CutPlaneNormal;
-	}
-	FVector3d UseNormal = CutPlaneNormal.GetSafeNormal();
-	if (FMath::Abs(1.0 - UseNormal.Length()) > 0.1)
-	{
-		UE::Geometry::AppendError(Debug, EGeometryScriptErrorType::InvalidInputs, LOCTEXT("ApplyMeshPlaneCut_InvalidNormal", "ApplyMeshPlaneCut: Normal vector is degenerate"));
-		UseNormal = FVector3d::UnitZ();
+		UseNormal = -UseNormal;
 	}
 
 	bool bSuccess = false;
 	TargetMesh->EditMesh([&](FDynamicMesh3& EditMesh)
 	{
 		FMeshPlaneCut Cut(&EditMesh, CutPlaneOrigin, UseNormal);
+		Cut.UVScaleFactor = 1.0 / Options.UVWorldDimension;
 		bSuccess = Cut.Cut();
 
 		if (Options.bFillHoles)
@@ -221,12 +218,78 @@ UDynamicMesh* UGeometryScriptLibrary_MeshBooleanFunctions::ApplyMeshPlaneCut(
 
 
 
+UDynamicMesh* UGeometryScriptLibrary_MeshBooleanFunctions::ApplyMeshPlaneSlice(
+	UDynamicMesh* TargetMesh,
+	FTransform CutFrame,
+	FGeometryScriptMeshPlaneSliceOptions Options,
+	UGeometryScriptDebug* Debug)
+{
+	const FName ObjectIndexAttribute = "ObjectIndexAttribute";
+
+	if (TargetMesh == nullptr)
+	{
+		UE::Geometry::AppendError(Debug, EGeometryScriptErrorType::InvalidInputs, LOCTEXT("ApplyMeshPlaneSlice_InvalidInput", "ApplyMeshPlaneSlice: TargetMesh is Null"));
+		return TargetMesh;
+	}
+
+	FFrame3d UseCutFrame(CutFrame);
+	FVector3d CutPlaneOrigin = UseCutFrame.Origin;
+	FVector3d UseNormal = UseCutFrame.Z();
+
+	bool bSuccess = false;
+	TargetMesh->EditMesh([&](FDynamicMesh3& EditMesh)
+	{
+		TDynamicMeshScalarTriangleAttribute<int>* SubObjectAttrib = new TDynamicMeshScalarTriangleAttribute<int>(&EditMesh);
+		SubObjectAttrib->Initialize(0);
+		EditMesh.Attributes()->AttachAttribute(ObjectIndexAttribute, SubObjectAttrib);
+
+		int MaxSubObjectID = -1;
+		for (int TID : EditMesh.TriangleIndicesItr())
+		{
+			MaxSubObjectID = FMath::Max(MaxSubObjectID, SubObjectAttrib->GetValue(TID));
+		}
+
+		FMeshPlaneCut Cut(&EditMesh, CutPlaneOrigin, UseNormal);
+		Cut.UVScaleFactor = 1.0 / Options.UVWorldDimension;
+
+		Cut.PlaneOrigin -= Cut.PlaneNormal * Options.GapWidth;
+		Cut.CutWithoutDelete(true, 0, SubObjectAttrib, MaxSubObjectID + 1, true, false);
+
+		int SecondCutMaxID = MaxSubObjectID;
+		for (int TID : EditMesh.TriangleIndicesItr())
+		{
+			SecondCutMaxID = FMath::Max(SecondCutMaxID, SubObjectAttrib->GetValue(TID));
+		}
+
+		Cut.PlaneOrigin += Cut.PlaneNormal * (2.0 * Options.GapWidth);
+		Cut.CutWithoutDelete(true, 0, SubObjectAttrib, SecondCutMaxID + 1, false, true);
+
+		for (int TID : EditMesh.TriangleIndicesItr())
+		{
+			int SubObjectID = SubObjectAttrib->GetValue(TID);
+			if (SubObjectID > MaxSubObjectID && SubObjectID <= SecondCutMaxID)
+			{
+				EditMesh.RemoveTriangle(TID);
+			}
+		}
+
+		if (Options.bFillHoles)
+		{
+			Cut.HoleFill(ConstrainedDelaunayTriangulate<double>, Options.bFillSpans);
+		}
+
+		EditMesh.Attributes()->RemoveAttribute(ObjectIndexAttribute);
+
+	}, EDynamicMeshChangeType::GeneralEdit, EDynamicMeshAttributeChangeFlags::Unknown, false);
+
+	return TargetMesh;
+}
+
 
 
 UDynamicMesh* UGeometryScriptLibrary_MeshBooleanFunctions::ApplyMeshMirror(
 	UDynamicMesh* TargetMesh,
-	FVector MirrorPlaneOrigin,
-	FVector MirrorPlaneNormal,
+	FTransform CutFrame,
 	FGeometryScriptMeshMirrorOptions Options,
 	UGeometryScriptDebug* Debug)
 {
@@ -238,15 +301,12 @@ UDynamicMesh* UGeometryScriptLibrary_MeshBooleanFunctions::ApplyMeshMirror(
 
 	const double PlaneTolerance = FMathf::ZeroTolerance * 10.0;
 
+	FFrame3d UseCutFrame(CutFrame);
+	FVector3d MirrorPlaneOrigin = UseCutFrame.Origin;
+	FVector3d UseNormal = UseCutFrame.Z();
 	if (Options.bApplyPlaneCut && Options.bFlipCutSide)
 	{
-		MirrorPlaneNormal = -MirrorPlaneNormal;
-	}
-	FVector3d UseNormal = MirrorPlaneNormal.GetSafeNormal();
-	if (FMath::Abs(1.0 - UseNormal.Length()) > 0.1)
-	{
-		UE::Geometry::AppendError(Debug, EGeometryScriptErrorType::InvalidInputs, LOCTEXT("ApplyMeshMirror_InvalidNormal", "ApplyMeshMirror: Normal vector is degenerate"));
-		UseNormal = FVector3d::UnitZ();
+		UseNormal = -UseNormal;
 	}
 
 	bool bSuccess = false;
diff --git a/Engine/Plugins/Experimental/GeometryScripting/Source/GeometryScriptingCore/Private/MeshDeformFunctions.cpp b/Engine/Plugins/Experimental/GeometryScripting/Source/GeometryScriptingCore/Private/MeshDeformFunctions.cpp
--- a/Engine/Plugins/Experimental/GeometryScripting/Source/GeometryScriptingCore/Private/MeshDeformFunctions.cpp
+++ b/Engine/Plugins/Experimental/GeometryScripting/Source/GeometryScriptingCore/Private/MeshDeformFunctions.cpp
@@ -7,6 +7,12 @@
 #include "SpaceDeformerOps/BendMeshOp.h"
 #include "SpaceDeformerOps/TwistMeshOp.h"
 #include "SpaceDeformerOps/FlareMeshOp.h"
+#include "DynamicMesh/MeshNormals.h"
+#include "Async/ParallelFor.h"
+
+#include "AssetUtils/Texture2DUtil.h"
+#include "Spatial/SampledScalarField2.h"
+
 #include "UDynamicMesh.h"
 
 #include "ExplicitUseGeometryMathTypes.h"		// using UE::Geometry::(math types)
@@ -162,5 +168,288 @@ UDynamicMesh* UGeometryScriptLibrary_MeshDeformFunctions::ApplyFlareWarpToMesh(
 
 
 
+UDynamicMesh* UGeometryScriptLibrary_MeshDeformFunctions::ApplyMathWarpToMesh(
+	UDynamicMesh* TargetMesh,
+	FTransform WarpOrientation,
+	EGeometryScriptMathWarpType WarpType,
+	FGeometryScriptMathWarpOptions Options,
+	UGeometryScriptDebug* Debug)
+{
+	if (TargetMesh == nullptr)
+	{
+		UE::Geometry::AppendError(Debug, EGeometryScriptErrorType::InvalidInputs, LOCTEXT("ApplyFlareWarpToMesh_InvalidInput", "ApplyFlareWarpToMesh: TargetMesh is Null"));
+		return TargetMesh;
+	}
+
+	UE::Geometry::FTransform3d WarpTransform(WarpOrientation);
+	double UseShift = (double)Options.FrequencyShift;
+	double UseFrequency = (double)Options.Frequency;
+	double UseMagnitude = (double)Options.Magnitude;
+
+	TargetMesh->EditMesh([&](FDynamicMesh3& EditMesh) 
+	{
+		ParallelFor(EditMesh.MaxVertexID(), [&](int32 vid)
+		{
+			if (EditMesh.IsVertex(vid))
+			{
+				FVector3d Pos = EditMesh.GetVertex(vid);
+				FVector3d LocalPos = WarpTransform.InverseTransformPosition(Pos);
+
+				FVector3d Displacement = FVector3d::Zero();
+				switch (WarpType)
+				{
+				case EGeometryScriptMathWarpType::SinWave1D:
+					{
+						LocalPos.Z += UseMagnitude * FMathd::Sin(UseFrequency * (LocalPos.X + UseShift));
+					}
+					break;
+				case EGeometryScriptMathWarpType::SinWave2D:
+					{
+						double Radius = FVector2d(LocalPos.X, LocalPos.Y).Length();
+						LocalPos.Z += UseMagnitude * FMathd::Sin(UseFrequency * (Radius + UseShift));
+					}
+					break;
+				}
+
+				Pos = WarpTransform.TransformPosition(LocalPos);
+				EditMesh.SetVertex(vid, Pos);
+			}
+		});
+
+	}, EDynamicMeshChangeType::GeneralEdit, EDynamicMeshAttributeChangeFlags::Unknown, false);
+
+	return TargetMesh;
+}
+
+
+UDynamicMesh* UGeometryScriptLibrary_MeshDeformFunctions::ApplyPerlinNoiseToMesh(
+	UDynamicMesh* TargetMesh,
+	FGeometryScriptPerlinNoiseOptions Options,
+	UGeometryScriptDebug* Debug)
+{
+	if (TargetMesh == nullptr)
+	{
+		UE::Geometry::AppendError(Debug, EGeometryScriptErrorType::InvalidInputs, LOCTEXT("ApplyPerlinNoiseToMesh_InvalidInput", "ApplyPerlinNoiseToMesh: TargetMesh is Null"));
+		return TargetMesh;
+	}
+
+	TargetMesh->EditMesh([&](FDynamicMesh3& EditMesh) 
+	{
+		FRandomStream Random(Options.BaseLayer.RandomSeed);
+		FVector3d Offsets[3];
+		for (int k = 0; k < 3; ++k)
+		{
+			const float RandomOffset = 10000.0f * Random.GetFraction();
+			Offsets[k] = FVector3d(RandomOffset, RandomOffset, RandomOffset);
+			Offsets[k] += (FVector3d)Options.BaseLayer.FrequencyShift;
+		}
+
+		FMeshNormals Normals(&EditMesh);
+		bool bAlignWithNormal = Options.bApplyAlongNormal;
+		if (bAlignWithNormal)
+		{
+			Normals.ComputeVertexNormals();
+		}
+
+		ParallelFor(EditMesh.MaxVertexID(), [&](int32 vid)
+		{
+			if (EditMesh.IsVertex(vid))
+			{
+				FVector3d Pos = EditMesh.GetVertex(vid);
+				if (bAlignWithNormal)
+				{
+					FVector NoisePos = (FVector)( (double)Options.BaseLayer.Frequency * (Pos + Offsets[0]) );
+					float Displacement = Options.BaseLayer.Magnitude * FMath::PerlinNoise3D(Options.BaseLayer.Frequency * NoisePos);
+					Pos += Displacement * Normals[vid];
+				}
+				else
+				{
+					FVector3d Displacement;
+					for (int32 k = 0; k < 3; ++k)
+					{
+						FVector NoisePos = (FVector)( (double)Options.BaseLayer.Frequency * (Pos + Offsets[k]) );
+						Displacement[k] = Options.BaseLayer.Magnitude * FMath::PerlinNoise3D(Options.BaseLayer.Frequency * NoisePos);
+					}
+					Pos += Displacement;
+				}
+				EditMesh.SetVertex(vid, Pos);
+			}
+		});
+
+	}, EDynamicMeshChangeType::GeneralEdit, EDynamicMeshAttributeChangeFlags::Unknown, false);
+
+	return TargetMesh;
+}
+
+
+
+
+
+
+
+UDynamicMesh* UGeometryScriptLibrary_MeshDeformFunctions::ApplyIterativeSmoothingToMesh(
+	UDynamicMesh* TargetMesh,
+	FGeometryScriptIterativeMeshSmoothingOptions Options,
+	UGeometryScriptDebug* Debug)
+{
+	if (TargetMesh == nullptr)
+	{
+		UE::Geometry::AppendError(Debug, EGeometryScriptErrorType::InvalidInputs, LOCTEXT("ApplyIterativeSmoothingToMesh_InvalidInput", "ApplyIterativeSmoothingToMesh: TargetMesh is Null"));
+		return TargetMesh;
+	}
+	double ClampAlpha = FMathd::Clamp((double)Options.Alpha, 0.0, 1.0);
+	int ClampIters = FMath::Clamp(Options.NumIterations, 0, 100);
+	if (ClampIters == 0 || ClampAlpha == 0.0)
+	{
+		return TargetMesh;
+	}
+
+	TargetMesh->EditMesh([&](FDynamicMesh3& EditMesh) 
+	{
+		int32 MaxVertices = EditMesh.MaxVertexID();
+		TArray<FVector3d> SmoothPositions;
+		SmoothPositions.SetNum(MaxVertices);
+
+		for (int32 k = 0; k < ClampIters; ++k)
+		{
+			ParallelFor(MaxVertices, [&](int32 VertexID)
+			{
+				if (EditMesh.IsVertex(VertexID))
+				{
+					FVector3d Centroid;
+					EditMesh.GetVtxOneRingCentroid(VertexID, Centroid);
+					SmoothPositions[VertexID] = Lerp(EditMesh.GetVertex(VertexID), Centroid, ClampAlpha);
+				}
+			});
+			for (int32 VertexID = 0; VertexID < MaxVertices; ++VertexID)
+			{
+				if (EditMesh.IsVertex(VertexID))
+				{
+					EditMesh.SetVertex(VertexID, SmoothPositions[VertexID]);
+				}
+			}
+		}		
+
+	}, EDynamicMeshChangeType::GeneralEdit, EDynamicMeshAttributeChangeFlags::Unknown, false);
+
+	return TargetMesh;
+}
+
+
+
+
+
+
+
+
+UDynamicMesh* UGeometryScriptLibrary_MeshDeformFunctions::ApplyDisplaceFromTextureMap(
+	UDynamicMesh* TargetMesh,
+	UTexture2D* TextureAsset,
+	FGeometryScriptDisplaceFromTextureOptions Options,
+	int32 UVLayer,
+	UGeometryScriptDebug* Debug)
+{
+	if (TargetMesh == nullptr)
+	{
+		UE::Geometry::AppendError(Debug, EGeometryScriptErrorType::InvalidInputs, LOCTEXT("ApplyDisplaceFromTextureMap_InvalidInput", "ApplyDisplaceFromTextureMap: TargetMesh is Null"));
+		return TargetMesh;
+	}
+	if (TextureAsset == nullptr)
+	{
+		UE::Geometry::AppendError(Debug, EGeometryScriptErrorType::InvalidInputs, LOCTEXT("ApplyDisplaceFromTextureMap_InvalidInput2", "ApplyDisplaceFromTextureMap: Texture is Null"));
+		return TargetMesh;
+	}
+
+	TImageBuilder<FVector4f> ImageData;
+	if (UE::AssetUtils::ReadTexture(TextureAsset, ImageData, false) == false)
+	{
+		UE::Geometry::AppendError(Debug, EGeometryScriptErrorType::OperationFailed, LOCTEXT("ApplyDisplaceFromTextureMap_TexReadFailed", "ApplyDisplaceFromTextureMap: Error reading source texture data"));
+		return TargetMesh;
+	}
+	int UseImageChannel = FMath::Clamp(Options.ImageChannel, 0, 3);
+
+	int64 TextureWidth = ImageData.GetDimensions().GetWidth();
+	int64 TextureHeight = ImageData.GetDimensions().GetHeight();
+	FSampledScalarField2f DisplaceField;
+	DisplaceField.Resize(TextureWidth, TextureHeight, 0.0f);
+	DisplaceField.SetCellSize(1.0f / (float)TextureWidth);	// todo: height of the texture will not be 1.0 if it was not square
+	for (int64 y = 0; y < TextureHeight; ++y)
+	{
+		for (int64 x = 0; x < TextureWidth; ++x)
+		{
+			DisplaceField.GridValues[y*TextureWidth + x] = ImageData.GetPixel(y*TextureWidth + x)[UseImageChannel];
+		}
+	}
+	TargetMesh->EditMesh([&](FDynamicMesh3& EditMesh)
+	{
+		if (EditMesh.HasAttributes() == false || !(UVLayer < EditMesh.Attributes()->NumUVLayers()))
+		{
+			UE::Geometry::AppendError(Debug, EGeometryScriptErrorType::InvalidInputs, LOCTEXT("ApplyDisplaceFromTextureMap_InvalidUVLayer", "ApplyDisplaceFromTextureMap: TargetMesh is missing requested UV Set"));
+			return;
+		}
+		FDynamicMeshUVOverlay* UVOverlay = EditMesh.Attributes()->GetUVLayer(UVLayer);
+		FDynamicMeshNormalOverlay* Normals = EditMesh.Attributes()->PrimaryNormals();
+		FVector2f UseUVScale = (FVector2f)Options.UVScale;
+		FVector2f UseUVOffset = (FVector2f)Options.UVOffset;
+
+		// We set things up such that DisplaceField goes from 0 to 1 in the U direction,
+		// but the V direction may be shorter or longer if the texture is not square
+		// (it will be 1/AspectRatio)
+		float VHeight = DisplaceField.Height() * DisplaceField.CellDimensions.Y;
+
+		TArray<FVector3d> DisplacedPositions;
+		DisplacedPositions.Init(FVector3d::Zero(), EditMesh.MaxVertexID());
+		TArray<int> Counts;
+		Counts.Init(0, EditMesh.MaxVertexID());
+
+		for (int tid : EditMesh.TriangleIndicesItr())
+		{
+			if (UVOverlay->IsSetTriangle(tid) == false || Normals->IsSetTriangle(tid) == false)
+			{
+				continue;
+			}
+
+			FIndex3i Tri = EditMesh.GetTriangle(tid);
+			FIndex3i UVTri = UVOverlay->GetTriangle(tid);
+			FIndex3i NormalTri = Normals->GetTriangle(tid);
+			for (int j = 0; j < 3; ++j)
+			{
+				int vid = Tri[j];
+				FVector2f UV = UVOverlay->GetElement(UVTri[j]);
+
+				// Adjust UV value and tile it. 
+				UV = UV * UseUVScale + UseUVOffset;
+				UV = UV - FVector2f(FMath::Floor(UV.X), FMath::Floor(UV.Y));
+				UV.Y *= VHeight;
+
+				double Offset = DisplaceField.BilinearSampleClamped(UV);
+				//if (AdjustmentCurve)
+				//{
+				//	Offset = AdjustmentCurve->Eval(Offset);
+				//}
+				Offset -= Options.Center;
+
+				FVector3f Normal = Normals->GetElement(NormalTri[j]);
+
+				double Intensity = Options.Magnitude;
+				DisplacedPositions[vid] += EditMesh.GetVertex(Tri[j]) + (Offset * Intensity * Normal);
+				Counts[vid]++;
+			}
+		}
+
+		for (int32 vid : EditMesh.VertexIndicesItr())
+		{
+			if (Counts[vid] != 0)
+			{
+				EditMesh.SetVertex(vid, DisplacedPositions[vid] / (double)Counts[vid]);
+			}
+		}
+
+
+	}, EDynamicMeshChangeType::GeneralEdit, EDynamicMeshAttributeChangeFlags::Unknown, false);
+
+	return TargetMesh;
+}
+
 
 #undef LOCTEXT_NAMESPACE
\ No newline at end of file
diff --git a/Engine/Plugins/Experimental/GeometryScripting/Source/GeometryScriptingCore/Private/MeshNormalsFunctions.cpp b/Engine/Plugins/Experimental/GeometryScripting/Source/GeometryScriptingCore/Private/MeshNormalsFunctions.cpp
--- a/Engine/Plugins/Experimental/GeometryScripting/Source/GeometryScriptingCore/Private/MeshNormalsFunctions.cpp
+++ b/Engine/Plugins/Experimental/GeometryScripting/Source/GeometryScriptingCore/Private/MeshNormalsFunctions.cpp
@@ -4,6 +4,7 @@
 
 #include "DynamicMesh/DynamicMesh3.h"
 #include "DynamicMesh/MeshNormals.h"
+#include "DynamicMesh/MeshTangents.h"
 #include "Operations/RepairOrientation.h"
 #include "DynamicMesh/DynamicMeshAABBTree3.h"
 #include "Polygroups/PolygroupSet.h"
@@ -197,6 +198,49 @@ UDynamicMesh* UGeometryScriptLibrary_MeshNormalsFunctions::ComputeSplitNormals(
 
 
 
+UDynamicMesh* UGeometryScriptLibrary_MeshNormalsFunctions::ComputeTangents( 
+	UDynamicMesh* TargetMesh, 
+	FGeometryScriptTangentsOptions Options,
+	UGeometryScriptDebug* Debug)
+{
+	if (TargetMesh == nullptr)
+	{
+		UE::Geometry::AppendError(Debug, EGeometryScriptErrorType::InvalidInputs, LOCTEXT("ComputeTangents_InvalidInput", "ComputeTangents: TargetMesh is Null"));
+		return TargetMesh;
+	}
+
+	// todo: publish correct change types
+	TargetMesh->EditMesh([&](FDynamicMesh3& EditMesh) 
+	{
+		if (EditMesh.HasAttributes() == false || !(Options.UVLayer < EditMesh.Attributes()->NumUVLayers()) || EditMesh.Attributes()->NumNormalLayers() == 0 )
+		{
+			UE::Geometry::AppendError(Debug, EGeometryScriptErrorType::InvalidInputs, LOCTEXT("ComputeTangents_NoUVLayers", "ComputeTangents: TargetMesh is missing UV Set or Normals required to compute Tangents"));
+			return;
+		}
+		if (EditMesh.Attributes()->HasTangentSpace() == false)
+		{
+			EditMesh.Attributes()->EnableTangents();
+		}
+
+		FComputeTangentsOptions TangentOptions;
+		TangentOptions.bAveraged = (Options.Type == EGeometryScriptTangentTypes::FastMikkT);
+
+		FMeshTangentsd Tangents(&EditMesh);
+		Tangents.ComputeTriVertexTangents(
+			EditMesh.Attributes()->PrimaryNormals(),
+			EditMesh.Attributes()->GetUVLayer(Options.UVLayer),
+			TangentOptions);
+		Tangents.CopyToOverlays(EditMesh);
+
+	}, EDynamicMeshChangeType::GeneralEdit, EDynamicMeshAttributeChangeFlags::Unknown, false);
+
+	return TargetMesh;
+}
+
+
+
+
+
 
 UDynamicMesh* UGeometryScriptLibrary_MeshNormalsFunctions::SetMeshTriangleNormals(
 	UDynamicMesh* TargetMesh,
diff --git a/Engine/Plugins/Experimental/GeometryScripting/Source/GeometryScriptingCore/Private/MeshPolygroupFunctions.cpp b/Engine/Plugins/Experimental/GeometryScripting/Source/GeometryScriptingCore/Private/MeshPolygroupFunctions.cpp
new file mode 100644
--- /dev/null
+++ b/Engine/Plugins/Experimental/GeometryScripting/Source/GeometryScriptingCore/Private/MeshPolygroupFunctions.cpp
@@ -0,0 +1,231 @@
+// Copyright Epic Games, Inc. All Rights Reserved.
+
+#include "GeometryScript/MeshPolygroupFunctions.h"
+
+#include "DynamicMesh/DynamicMesh3.h"
+#include "DynamicMesh/DynamicMeshAttributeSet.h"
+#include "Polygroups/PolygroupsGenerator.h"
+#include "UDynamicMesh.h"
+
+#include "ExplicitUseGeometryMathTypes.h"		// using UE::Geometry::(math types)
+using namespace UE::Geometry;
+
+#define LOCTEXT_NAMESPACE "UGeometryScriptLibrary_MeshPolygroupFunctions"
+
+
+
+
+UDynamicMesh* UGeometryScriptLibrary_MeshPolygroupFunctions::EnablePolygroups( 
+	UDynamicMesh* TargetMesh, 
+	UGeometryScriptDebug* Debug)
+{
+	if (TargetMesh == nullptr)
+	{
+		UE::Geometry::AppendError(Debug, EGeometryScriptErrorType::InvalidInputs, LOCTEXT("EnablePolygroups_InvalidInput", "EnablePolygroups: TargetMesh is Null"));
+		return TargetMesh;
+	}
+
+
+	TargetMesh->EditMesh([&](FDynamicMesh3& EditMesh) 
+	{
+		if (EditMesh.HasTriangleGroups() == false)
+		{
+			EditMesh.EnableTriangleGroups(0);
+		}
+	}, EDynamicMeshChangeType::GeneralEdit, EDynamicMeshAttributeChangeFlags::Unknown, false);
+
+	return TargetMesh;
+}
+
+
+UDynamicMesh* UGeometryScriptLibrary_MeshPolygroupFunctions::SetNumExtendedPolygroupLayers( 
+	UDynamicMesh* TargetMesh, 
+	int NumLayers,
+	UGeometryScriptDebug* Debug)
+{
+	if (TargetMesh == nullptr)
+	{
+		UE::Geometry::AppendError(Debug, EGeometryScriptErrorType::InvalidInputs, LOCTEXT("SetNumExtendedPolygroupLayers_InvalidInput", "SetNumExtendedPolygroupLayers: TargetMesh is Null"));
+		return TargetMesh;
+	}
+
+	TargetMesh->EditMesh([&](FDynamicMesh3& EditMesh) 
+	{
+		if (EditMesh.HasAttributes() == false)
+		{
+			EditMesh.EnableAttributes();
+		}
+		if (EditMesh.Attributes()->NumPolygroupLayers() != NumLayers)
+		{
+			EditMesh.Attributes()->SetNumPolygroupLayers(NumLayers);
+		}
+	}, EDynamicMeshChangeType::GeneralEdit, EDynamicMeshAttributeChangeFlags::Unknown, false);
+
+	return TargetMesh;
+}
+
+
+
+
+UDynamicMesh* UGeometryScriptLibrary_MeshPolygroupFunctions::ClearPolygroups( 
+	UDynamicMesh* TargetMesh, 
+	FGeometryScriptGroupLayer GroupLayer,
+	int ClearValue,
+	UGeometryScriptDebug* Debug)
+{
+	if (TargetMesh == nullptr)
+	{
+		UE::Geometry::AppendError(Debug, EGeometryScriptErrorType::InvalidInputs, LOCTEXT("ClearPolygroups_InvalidInput", "ClearPolygroups: TargetMesh is Null"));
+		return TargetMesh;
+	}
+	TargetMesh->EditMesh([&](FDynamicMesh3& EditMesh) 
+	{
+		TUniquePtr<FPolygroupSet> OutputGroups;
+		FPolygroupLayer InputGroupLayer{ GroupLayer.bDefaultLayer, GroupLayer.ExtendedLayerIndex };
+		if (InputGroupLayer.CheckExists(&EditMesh) == false)
+		{
+			UE::Geometry::AppendError(Debug, EGeometryScriptErrorType::InvalidInputs, LOCTEXT("ClearPolygroups_MissingGroups", "ClearPolygroups: Target Polygroup Layer does not exist"));
+			return;
+		}
+		OutputGroups = MakeUnique<FPolygroupSet>(&EditMesh, InputGroupLayer);
+		for (int32 tid : EditMesh.TriangleIndicesItr())
+		{
+			OutputGroups->SetGroup(tid, ClearValue, EditMesh);
+		}
+	}, EDynamicMeshChangeType::GeneralEdit, EDynamicMeshAttributeChangeFlags::Unknown, false);
+
+	return TargetMesh;
+}
+
+
+
+UDynamicMesh* UGeometryScriptLibrary_MeshPolygroupFunctions::CopyPolygroupsLayer( 
+	UDynamicMesh* TargetMesh, 
+	FGeometryScriptGroupLayer FromGroupLayer,
+	FGeometryScriptGroupLayer ToGroupLayer,
+	UGeometryScriptDebug* Debug)
+{
+	if (TargetMesh == nullptr)
+	{
+		UE::Geometry::AppendError(Debug, EGeometryScriptErrorType::InvalidInputs, LOCTEXT("CopyPolygroupsLayer_InvalidInput", "CopyPolygroupsLayer: TargetMesh is Null"));
+		return TargetMesh;
+	}
+	TargetMesh->EditMesh([&](FDynamicMesh3& EditMesh) 
+	{
+		FPolygroupLayer InputGroupLayer{ FromGroupLayer.bDefaultLayer, FromGroupLayer.ExtendedLayerIndex };
+		FPolygroupLayer OutputGroupLayer{ ToGroupLayer.bDefaultLayer, ToGroupLayer.ExtendedLayerIndex };
+		if (InputGroupLayer == OutputGroupLayer)
+		{
+			UE::Geometry::AppendError(Debug, EGeometryScriptErrorType::InvalidInputs, LOCTEXT("CopyPolygroupsLayer_SameGroups", "CopyPolygroupsLayer: tried to copy Polygroup Layer to itself"));
+			return;
+		}
+		if (InputGroupLayer.CheckExists(&EditMesh) == false)
+		{
+			UE::Geometry::AppendError(Debug, EGeometryScriptErrorType::InvalidInputs, LOCTEXT("CopyPolygroupsLayer_MissingFromGroups", "CopyPolygroupsLayer: From Polygroup Layer does not exist"));
+			return;
+		}
+		if (OutputGroupLayer.CheckExists(&EditMesh) == false)
+		{
+			UE::Geometry::AppendError(Debug, EGeometryScriptErrorType::InvalidInputs, LOCTEXT("CopyPolygroupsLayer_MissingToGroups", "CopyPolygroupsLayer: To Polygroup Layer does not exist"));
+			return;
+		}
+
+
+		TUniquePtr<FPolygroupSet> InputGroups = MakeUnique<FPolygroupSet>(&EditMesh, InputGroupLayer);
+		TUniquePtr<FPolygroupSet> OutputGroups = MakeUnique<FPolygroupSet>(&EditMesh, OutputGroupLayer);
+
+		for (int32 tid : EditMesh.TriangleIndicesItr())
+		{
+			OutputGroups->SetGroup(tid, InputGroups->GetGroup(tid), EditMesh);
+		}
+	}, EDynamicMeshChangeType::GeneralEdit, EDynamicMeshAttributeChangeFlags::Unknown, false);
+
+	return TargetMesh;
+}
+
+
+
+
+UDynamicMesh* UGeometryScriptLibrary_MeshPolygroupFunctions::ConvertUVIslandsToPolygroups( 
+	UDynamicMesh* TargetMesh, 
+	FGeometryScriptGroupLayer GroupLayer,
+	int UVLayer,
+	UGeometryScriptDebug* Debug)
+{
+	if (TargetMesh == nullptr)
+	{
+		UE::Geometry::AppendError(Debug, EGeometryScriptErrorType::InvalidInputs, LOCTEXT("ConvertUVIslandsToPolygroups_InvalidInput", "ConvertUVIslandsToPolygroups: TargetMesh is Null"));
+		return TargetMesh;
+	}
+
+	TargetMesh->EditMesh([&](FDynamicMesh3& EditMesh) 
+	{
+		if (EditMesh.HasAttributes() == false || !(UVLayer < EditMesh.Attributes()->NumUVLayers()) )
+		{
+			UE::Geometry::AppendError(Debug, EGeometryScriptErrorType::InvalidInputs, LOCTEXT("ConvertUVIslandsToPolygroups_InvalidUVLayers", "ConvertUVIslandsToPolygroups: Requested UV layer does not exist"));
+			return;
+		}
+
+		TUniquePtr<FPolygroupSet> OutputGroups;
+		FPolygroupLayer InputGroupLayer{ GroupLayer.bDefaultLayer, GroupLayer.ExtendedLayerIndex };
+		if (InputGroupLayer.CheckExists(&EditMesh) == false)
+		{
+			UE::Geometry::AppendError(Debug, EGeometryScriptErrorType::InvalidInputs, LOCTEXT("ConvertUVIslandsToPolygroups_MissingGroups", "ConvertUVIslandsToPolygroups: Target Polygroup Layer does not exist"));
+			return;
+		}
+		OutputGroups = MakeUnique<FPolygroupSet>(&EditMesh, InputGroupLayer);
+
+		FPolygroupsGenerator Generator(&EditMesh);
+		Generator.bApplyPostProcessing = false;
+		Generator.bCopyToMesh = false;
+		Generator.FindPolygroupsFromUVIslands(UVLayer);
+		Generator.CopyPolygroupsToPolygroupSet(*OutputGroups, EditMesh);
+
+	}, EDynamicMeshChangeType::GeneralEdit, EDynamicMeshAttributeChangeFlags::Unknown, false);
+
+	return TargetMesh;
+}
+
+
+
+
+UDynamicMesh* UGeometryScriptLibrary_MeshPolygroupFunctions::ComputePolygroupsFromAngleThreshold( 
+	UDynamicMesh* TargetMesh, 
+	FGeometryScriptGroupLayer GroupLayer,
+	float CreaseAngle,
+	int MinGroupSize,
+	UGeometryScriptDebug* Debug)
+{
+	if (TargetMesh == nullptr)
+	{
+		UE::Geometry::AppendError(Debug, EGeometryScriptErrorType::InvalidInputs, LOCTEXT("ComputePolygroupsFromAngleThresholde_InvalidInput", "ComputePolygroupsFromAngleThreshold: TargetMesh is Null"));
+		return TargetMesh;
+	}
+
+	TargetMesh->EditMesh([&](FDynamicMesh3& EditMesh) 
+	{
+		TUniquePtr<FPolygroupSet> OutputGroups;
+		FPolygroupLayer InputGroupLayer{ GroupLayer.bDefaultLayer, GroupLayer.ExtendedLayerIndex };
+		if (InputGroupLayer.CheckExists(&EditMesh) == false)
+		{
+			UE::Geometry::AppendError(Debug, EGeometryScriptErrorType::InvalidInputs, LOCTEXT("ComputePolygroupsFromAngleThreshold_MissingGroups", "ComputePolygroupsFromAngleThreshold: Target Polygroup Layer does not exist"));
+			return;
+		}
+		OutputGroups = MakeUnique<FPolygroupSet>(&EditMesh, InputGroupLayer);
+
+		FPolygroupsGenerator Generator(&EditMesh);
+		Generator.bApplyPostProcessing = (MinGroupSize > 1);
+		Generator.MinGroupSize = MinGroupSize;
+		Generator.bCopyToMesh = false;
+		double DotTolerance = 1.0 - FMathd::Cos((double)CreaseAngle * FMathd::DegToRad);
+		Generator.FindPolygroupsFromFaceNormals(DotTolerance);
+		Generator.CopyPolygroupsToPolygroupSet(*OutputGroups, EditMesh);
+
+	}, EDynamicMeshChangeType::GeneralEdit, EDynamicMeshAttributeChangeFlags::Unknown, false);
+
+	return TargetMesh;
+}
+
+
+
+#undef LOCTEXT_NAMESPACE
\ No newline at end of file
diff --git a/Engine/Plugins/Experimental/GeometryScripting/Source/GeometryScriptingCore/Private/MeshPrimitiveFunctions.cpp b/Engine/Plugins/Experimental/GeometryScripting/Source/GeometryScriptingCore/Private/MeshPrimitiveFunctions.cpp
--- a/Engine/Plugins/Experimental/GeometryScripting/Source/GeometryScriptingCore/Private/MeshPrimitiveFunctions.cpp
+++ b/Engine/Plugins/Experimental/GeometryScripting/Source/GeometryScriptingCore/Private/MeshPrimitiveFunctions.cpp
@@ -18,6 +18,9 @@
 #include "Generators/FlatTriangulationMeshGenerator.h"
 #include "ConstrainedDelaunay2.h"
 #include "Arrangement2d.h"
+#include "Util/RevolveUtil.h"
+
+#include "CompositionOps/CurveSweepOp.h"
 
 #include "ExplicitUseGeometryMathTypes.h"		// using UE::Geometry::(math types)
 using namespace UE::Geometry;
@@ -25,10 +28,20 @@ using namespace UE::Geometry;
 #define LOCTEXT_NAMESPACE "UGeometryScriptLibrary_MeshPrimitiveFunctions"
 
 
-static void AppendPrimitive(UDynamicMesh* TargetMesh, FMeshShapeGenerator* Generator, FTransform Transform, FGeometryScriptPrimitiveOptions PrimitiveOptions)
+static void AppendPrimitive(
+	UDynamicMesh* TargetMesh, 
+	FMeshShapeGenerator* Generator, 
+	FTransform Transform, 
+	FGeometryScriptPrimitiveOptions PrimitiveOptions,
+	FVector3d PreTranslate = FVector3d::Zero())
 {
-	auto ApplyOptionsToMesh = [&Transform, &PrimitiveOptions](FDynamicMesh3& Mesh)
+	auto ApplyOptionsToMesh = [&Transform, &PrimitiveOptions, PreTranslate](FDynamicMesh3& Mesh)
 	{
+		if (PreTranslate.SquaredLength() > 0)
+		{
+			MeshTransforms::Translate(Mesh, PreTranslate);
+		}
+
 		MeshTransforms::ApplyTransform(Mesh, (UE::Geometry::FTransform3d)Transform);
 		if (PrimitiveOptions.PolygroupMode == EGeometryScriptPrimitivePolygroupMode::SingleGroup)
 		{
@@ -80,10 +93,13 @@ UDynamicMesh* UGeometryScriptLibrary_MeshPrimitiveFunctions::AppendBox(
 	UDynamicMesh* TargetMesh,
 	FGeometryScriptPrimitiveOptions PrimitiveOptions,
 	FTransform Transform,
-	FBox Box,
+	float DimensionX,
+	float DimensionY,
+	float DimensionZ,
 	int32 StepsX,
 	int32 StepsY,
 	int32 StepsZ,
+	EGeometryScriptPrimitiveOriginMode Origin,
 	UGeometryScriptDebug* Debug)
 {
 	if (TargetMesh == nullptr)
@@ -92,13 +108,10 @@ UDynamicMesh* UGeometryScriptLibrary_MeshPrimitiveFunctions::AppendBox(
 		return TargetMesh;
 	}
 
-	UE::Geometry::FAxisAlignedBox3d ConvertBox(Box);
-	if (ConvertBox.MaxDim() == 0)
-	{
-		ConvertBox = UE::Geometry::FAxisAlignedBox3d(FVector3d(-50, -50, 0), FVector3d(50, 50, 100));
-	}
+	UE::Geometry::FAxisAlignedBox3d ConvertBox(
+		FVector3d(-DimensionX / 2, -DimensionY / 2, 0),
+		FVector3d(DimensionX / 2, DimensionY / 2, DimensionZ));
 	
-	// todo: could apply transform to box here and skip later?
 	// todo: if Steps X/Y/Z are zero, can use trivial box generator
 
 	FGridBoxMeshGenerator GridBoxGenerator;
@@ -107,7 +120,8 @@ UDynamicMesh* UGeometryScriptLibrary_MeshPrimitiveFunctions::AppendBox(
 	GridBoxGenerator.bPolygroupPerQuad = (PrimitiveOptions.PolygroupMode == EGeometryScriptPrimitivePolygroupMode::PerQuad);
 	GridBoxGenerator.Generate();
 
-	AppendPrimitive(TargetMesh, &GridBoxGenerator, Transform, PrimitiveOptions);
+	FVector3d OriginShift = (Origin == EGeometryScriptPrimitiveOriginMode::Center) ? FVector3d(0, 0, -DimensionZ/2) : FVector3d::Zero();
+	AppendPrimitive(TargetMesh, &GridBoxGenerator, Transform, PrimitiveOptions, OriginShift);
 
 	return TargetMesh;
 }
@@ -121,6 +135,7 @@ UDynamicMesh* UGeometryScriptLibrary_MeshPrimitiveFunctions::AppendSphereLatLong
 	float Radius,
 	int32 StepsPhi,
 	int32 StepsTheta,
+	EGeometryScriptPrimitiveOriginMode Origin,
 	UGeometryScriptDebug* Debug)
 {
 	if (TargetMesh == nullptr)
@@ -136,7 +151,8 @@ UDynamicMesh* UGeometryScriptLibrary_MeshPrimitiveFunctions::AppendSphereLatLong
 	SphereGenerator.bPolygroupPerQuad = (PrimitiveOptions.PolygroupMode == EGeometryScriptPrimitivePolygroupMode::PerQuad);
 	SphereGenerator.Generate();
 
-	AppendPrimitive(TargetMesh, &SphereGenerator, Transform, PrimitiveOptions);
+	FVector3d OriginShift = (Origin == EGeometryScriptPrimitiveOriginMode::Base) ? FVector3d(0, 0, Radius) : FVector3d::Zero();
+	AppendPrimitive(TargetMesh, &SphereGenerator, Transform, PrimitiveOptions, OriginShift);
 
 	return TargetMesh;
 }
@@ -151,6 +167,7 @@ UDynamicMesh* UGeometryScriptLibrary_MeshPrimitiveFunctions::AppendSphereBox(
 	int32 StepsX,
 	int32 StepsY,
 	int32 StepsZ,
+	EGeometryScriptPrimitiveOriginMode Origin,
 	UGeometryScriptDebug* Debug)
 {
 	if (TargetMesh == nullptr)
@@ -165,7 +182,8 @@ UDynamicMesh* UGeometryScriptLibrary_MeshPrimitiveFunctions::AppendSphereBox(
 	SphereGenerator.bPolygroupPerQuad = (PrimitiveOptions.PolygroupMode == EGeometryScriptPrimitivePolygroupMode::PerQuad);
 	SphereGenerator.Generate();
 
-	AppendPrimitive(TargetMesh, &SphereGenerator, Transform, PrimitiveOptions);
+	FVector3d OriginShift = (Origin == EGeometryScriptPrimitiveOriginMode::Base) ? FVector3d(0, 0, Radius) : FVector3d::Zero();
+	AppendPrimitive(TargetMesh, &SphereGenerator, Transform, PrimitiveOptions, OriginShift);
 
 	return TargetMesh;
 }
@@ -181,6 +199,7 @@ UDynamicMesh* UGeometryScriptLibrary_MeshPrimitiveFunctions::AppendCapsule(
 	float LineLength,
 	int32 HemisphereSteps,
 	int32 CircleSteps,
+	EGeometryScriptPrimitiveOriginMode Origin,
 	UGeometryScriptDebug* Debug)
 {
 	if (TargetMesh == nullptr)
@@ -197,7 +216,9 @@ UDynamicMesh* UGeometryScriptLibrary_MeshPrimitiveFunctions::AppendCapsule(
 	CapsuleGenerator.bPolygroupPerQuad = (PrimitiveOptions.PolygroupMode == EGeometryScriptPrimitivePolygroupMode::PerQuad);
 	CapsuleGenerator.Generate();
 
-	AppendPrimitive(TargetMesh, &CapsuleGenerator, Transform, PrimitiveOptions);
+	FVector3d OriginShift = FVector3d::Zero();
+	OriginShift.Z = (Origin == EGeometryScriptPrimitiveOriginMode::Center) ? -(LineLength / 2) : Radius;
+	AppendPrimitive(TargetMesh, &CapsuleGenerator, Transform, PrimitiveOptions, OriginShift);
 
 	return TargetMesh;
 }
@@ -214,6 +235,7 @@ UDynamicMesh* UGeometryScriptLibrary_MeshPrimitiveFunctions::AppendCylinder(
 	int32 RadialSteps,
 	int32 HeightSteps,
 	bool bCapped,
+	EGeometryScriptPrimitiveOriginMode Origin,
 	UGeometryScriptDebug* Debug)
 {
 	if (TargetMesh == nullptr)
@@ -232,7 +254,8 @@ UDynamicMesh* UGeometryScriptLibrary_MeshPrimitiveFunctions::AppendCylinder(
 	CylinderGenerator.bPolygroupPerQuad = (PrimitiveOptions.PolygroupMode == EGeometryScriptPrimitivePolygroupMode::PerQuad);
 	CylinderGenerator.Generate();
 
-	AppendPrimitive(TargetMesh, &CylinderGenerator, Transform, PrimitiveOptions);
+	FVector3d OriginShift = (Origin == EGeometryScriptPrimitiveOriginMode::Center) ? FVector3d(0, 0, -(Height/2)) : FVector3d::Zero();
+	AppendPrimitive(TargetMesh, &CylinderGenerator, Transform, PrimitiveOptions, OriginShift);
 
 	return TargetMesh;
 }
@@ -249,6 +272,7 @@ UDynamicMesh* UGeometryScriptLibrary_MeshPrimitiveFunctions::AppendCone(
 	int32 RadialSteps,
 	int32 HeightSteps,
 	bool bCapped,
+	EGeometryScriptPrimitiveOriginMode Origin,
 	UGeometryScriptDebug* Debug)
 {
 	if (TargetMesh == nullptr)
@@ -267,7 +291,8 @@ UDynamicMesh* UGeometryScriptLibrary_MeshPrimitiveFunctions::AppendCone(
 	CylinderGenerator.bPolygroupPerQuad = (PrimitiveOptions.PolygroupMode == EGeometryScriptPrimitivePolygroupMode::PerQuad);
 	CylinderGenerator.Generate();
 
-	AppendPrimitive(TargetMesh, &CylinderGenerator, Transform, PrimitiveOptions);
+	FVector3d OriginShift = (Origin == EGeometryScriptPrimitiveOriginMode::Center) ? FVector3d(0, 0, -(Height/2)) : FVector3d::Zero();
+	AppendPrimitive(TargetMesh, &CylinderGenerator, Transform, PrimitiveOptions, OriginShift);
 
 	return TargetMesh;
 }
@@ -281,13 +306,15 @@ UDynamicMesh* UGeometryScriptLibrary_MeshPrimitiveFunctions::AppendTorus(
 	float MinorRadius,
 	int32 MajorSteps,
 	int32 MinorSteps,
+	EGeometryScriptPrimitiveOriginMode Origin,
 	UGeometryScriptDebug* Debug)
 {
 	FPolygon2d Circle = FPolygon2d::MakeCircle(FMath::Max(FMathf::ZeroTolerance, MinorRadius), FMath::Max(3, MinorSteps));
 	TArray<FVector2D> PolygonVertices;
+	FVector2d Shift = (Origin == EGeometryScriptPrimitiveOriginMode::Base) ? FVector2d(0, MinorRadius) : FVector2d::Zero();
 	for (FVector2d v : Circle.GetVertices())
 	{
-		PolygonVertices.Add((FVector2D)v);
+		PolygonVertices.Add((FVector2D)(v+Shift));
 	}
 	return AppendSimpleRevolvePolygon(TargetMesh, PrimitiveOptions, Transform, PolygonVertices, MajorRadius, MajorSteps, Debug);
 }
@@ -339,6 +366,157 @@ UDynamicMesh* UGeometryScriptLibrary_MeshPrimitiveFunctions::AppendSimpleRevolve
 
 
 
+UDynamicMesh* UGeometryScriptLibrary_MeshPrimitiveFunctions::AppendRevolvePath(
+	UDynamicMesh* TargetMesh,
+	FGeometryScriptPrimitiveOptions PrimitiveOptions,
+	FTransform Transform,
+	const TArray<FVector2D>& PathVertices,
+	FGeometryScriptRevolveOptions RevolveOptions,
+	int32 Steps,
+	bool bCapped,
+	UGeometryScriptDebug* Debug)
+{
+	if (TargetMesh == nullptr)
+	{
+		UE::Geometry::AppendError(Debug, EGeometryScriptErrorType::InvalidInputs, LOCTEXT("AppendRevolvePath_NullMesh", "AppendRevolvePath: TargetMesh is Null"));
+		return TargetMesh;
+	}
+	if (PathVertices.Num() < 2)
+	{
+		UE::Geometry::AppendError(Debug, EGeometryScriptErrorType::InvalidInputs, LOCTEXT("AppendRevolvePath_InvalidPolygon", "AppendRevolvePath: PathVertices array requires at least 2 positions"));
+		return TargetMesh;
+	}
+
+	FCurveSweepOp CurveSweepOp;
+
+	FVector3d AxisDirection(0, 0, 1);
+	FVector3d AxisOrigin(0, 0, 0);
+	for (FVector2D Point : PathVertices)
+	{
+		CurveSweepOp.ProfileCurve.Add(FVector3d(Point.X, 0, Point.Y));
+	}
+	// unclear why but the sweep code seems to be written for a clockwise ordering...
+	Algo::Reverse(CurveSweepOp.ProfileCurve);
+
+	// Project first and last points onto the revolution axis to cap it
+	if (bCapped)
+	{
+		FVector3d FirstPoint = CurveSweepOp.ProfileCurve[0];
+		FVector3d LastPoint = CurveSweepOp.ProfileCurve.Last();
+
+		double DistanceAlongAxis = AxisDirection.Dot(LastPoint - AxisOrigin);
+		FVector3d ProjectedPoint = AxisOrigin + (AxisDirection * DistanceAlongAxis);
+		CurveSweepOp.ProfileCurve.Add(ProjectedPoint);
+
+		DistanceAlongAxis = AxisDirection.Dot(FirstPoint - AxisOrigin);
+		ProjectedPoint = AxisOrigin + (AxisDirection * DistanceAlongAxis);
+		CurveSweepOp.ProfileCurve.Add(ProjectedPoint);
+
+		CurveSweepOp.bProfileCurveIsClosed = true;
+	}
+	else
+	{
+		CurveSweepOp.bProfileCurveIsClosed = false;
+	}
+
+	//double TotalRevolutionDegrees = RevolveDegrees;// (AlongAxisOffsetPerDegree == 0) ? ClampedRevolutionDegrees : RevolutionDegrees;
+	double TotalRevolutionDegrees = FMath::Clamp(RevolveOptions.RevolveDegrees, 0.1f, 360.0f);
+
+	double DegreesPerStep = TotalRevolutionDegrees / (double)Steps;
+	double DegreesOffset = RevolveOptions.DegreeOffset;
+	if (RevolveOptions.ReverseDirection)
+	{
+		DegreesPerStep *= -1;
+		DegreesOffset *= -1;
+	}
+	//double DownAxisOffsetPerStep = TotalRevolutionDegrees * AlongAxisOffsetPerDegree / Steps;
+
+	if (RevolveOptions.bProfileAtMidpoint && DegreesPerStep != 0 && FMathd::Abs(DegreesPerStep) < 180)
+	{
+		RevolveUtil::MakeProfileCurveMidpointOfFirstStep(CurveSweepOp.ProfileCurve, DegreesPerStep, AxisOrigin, AxisDirection);
+	}
+
+	// Generate the sweep curve
+	//CurveSweepOpOut.bSweepCurveIsClosed = bWeldFullRevolution && AlongAxisOffsetPerDegree == 0 && TotalRevolutionDegrees == 360;
+	CurveSweepOp.bSweepCurveIsClosed = (TotalRevolutionDegrees == 360);
+	int32 NumSweepFrames = CurveSweepOp.bSweepCurveIsClosed ? Steps : Steps + 1; // If closed, last sweep frame is also first
+	CurveSweepOp.SweepCurve.Reserve(NumSweepFrames);
+	RevolveUtil::GenerateSweepCurve(AxisOrigin, AxisDirection, DegreesOffset,
+		DegreesPerStep, 0.0, NumSweepFrames, CurveSweepOp.SweepCurve);
+
+	// Weld any vertices that are on the axis
+	RevolveUtil::WeldPointsOnAxis(CurveSweepOp.ProfileCurve, AxisOrigin,
+		AxisDirection, 0.1, CurveSweepOp.ProfileVerticesToWeld);
+
+	CurveSweepOp.bSharpNormals = RevolveOptions.bHardNormals;
+	CurveSweepOp.SharpNormalAngleTolerance = RevolveOptions.HardNormalAngle;
+	//CurveSweepOpOut.DiagonalTolerance = DiagonalProportionTolerance;
+	//double UVScale = MaterialProperties.UVScale;
+	double UVScale = 1.0;
+	CurveSweepOp.UVScale = FVector2d(UVScale, UVScale);
+	//if (bReverseProfileCurve ^ bFlipVs)
+	//{
+	//	CurveSweepOpOut.UVScale[1] *= -1;
+	//	CurveSweepOpOut.UVOffset = FVector2d(0, UVScale);
+	//}
+	CurveSweepOp.bUVsSkipFullyWeldedEdges = true;
+	CurveSweepOp.bUVScaleRelativeWorld = true; // MaterialProperties.bWorldSpaceUVScale;
+	CurveSweepOp.UnitUVInWorldCoordinates = 100; // This seems to be the case in the AddPrimitiveTool
+	CurveSweepOp.QuadSplitMode = EProfileSweepQuadSplit::Uniform;
+
+	if (PrimitiveOptions.PolygroupMode == EGeometryScriptPrimitivePolygroupMode::PerQuad)
+	{
+		CurveSweepOp.PolygonGroupingMode = EProfileSweepPolygonGrouping::PerFace;
+	}
+	if (PrimitiveOptions.PolygroupMode == EGeometryScriptPrimitivePolygroupMode::PerFace)
+	{
+		CurveSweepOp.PolygonGroupingMode = EProfileSweepPolygonGrouping::PerProfileSegment;
+	}
+	else
+	{
+		CurveSweepOp.PolygonGroupingMode = EProfileSweepPolygonGrouping::Single;
+	}
+
+	CurveSweepOp.CapFillMode = (bCapped) ?
+		FCurveSweepOp::ECapFillMode::EarClipping : FCurveSweepOp::ECapFillMode::None;		// delaunay? hits ensure...
+
+	CurveSweepOp.CalculateResult(nullptr);
+	TUniquePtr<FDynamicMesh3> ResultMesh = CurveSweepOp.ExtractResult();
+	MeshTransforms::ApplyTransform(*ResultMesh, (UE::Geometry::FTransform3d)Transform);
+
+	if (PrimitiveOptions.bFlipOrientation)
+	{
+		ResultMesh->ReverseOrientation(true);
+		if (ResultMesh->HasAttributes() && ResultMesh->Attributes()->PrimaryNormals() != nullptr)
+		{
+			FDynamicMeshNormalOverlay* Normals = ResultMesh->Attributes()->PrimaryNormals();
+			for (int elemid : Normals->ElementIndicesItr())
+			{
+				Normals->SetElement(elemid, -Normals->GetElement(elemid));
+			}
+		}
+	}
+
+	if (TargetMesh->IsEmpty())
+	{
+		TargetMesh->SetMesh(MoveTemp(*ResultMesh));
+	}
+	else
+	{
+		TargetMesh->EditMesh([&](FDynamicMesh3& EditMesh)
+		{
+			FMeshIndexMappings TmpMappings;
+			FDynamicMeshEditor Editor(&EditMesh);
+			Editor.AppendMesh(ResultMesh.Get(), TmpMappings);
+
+		}, EDynamicMeshChangeType::GeneralEdit, EDynamicMeshAttributeChangeFlags::Unknown, false);
+	}
+
+	return TargetMesh;
+}
+
+
+
 UDynamicMesh* UGeometryScriptLibrary_MeshPrimitiveFunctions::AppendSimpleExtrudePolygon(
 	UDynamicMesh* TargetMesh,
 	FGeometryScriptPrimitiveOptions PrimitiveOptions,
@@ -347,6 +525,7 @@ UDynamicMesh* UGeometryScriptLibrary_MeshPrimitiveFunctions::AppendSimpleExtrude
 	float Height,
 	int32 HeightSteps,
 	bool bCapped,
+	EGeometryScriptPrimitiveOriginMode Origin,
 	UGeometryScriptDebug* Debug)
 {
 	if (TargetMesh == nullptr)
@@ -381,7 +560,8 @@ UDynamicMesh* UGeometryScriptLibrary_MeshPrimitiveFunctions::AppendSimpleExtrude
 	ExtrudeGen.bPolygroupPerQuad = (PrimitiveOptions.PolygroupMode == EGeometryScriptPrimitivePolygroupMode::PerQuad);
 	ExtrudeGen.Generate();
 
-	AppendPrimitive(TargetMesh, &ExtrudeGen, Transform, PrimitiveOptions);
+	FVector3d OriginShift = (Origin == EGeometryScriptPrimitiveOriginMode::Center) ? FVector3d(0, 0, -(Height/2)) : FVector3d::Zero();
+	AppendPrimitive(TargetMesh, &ExtrudeGen, Transform, PrimitiveOptions, OriginShift);
 	return TargetMesh;
 }
 
@@ -440,7 +620,8 @@ UDynamicMesh* UGeometryScriptLibrary_MeshPrimitiveFunctions::AppendRectangle(
 	UDynamicMesh* TargetMesh,
 	FGeometryScriptPrimitiveOptions PrimitiveOptions,
 	FTransform Transform,
-	FBox2D Box,
+	float DimensionX,
+	float DimensionY,
 	int32 StepsWidth,
 	int32 StepsHeight,
 	UGeometryScriptDebug* Debug)
@@ -451,18 +632,11 @@ UDynamicMesh* UGeometryScriptLibrary_MeshPrimitiveFunctions::AppendRectangle(
 		return TargetMesh;
 	}
 
-	UE::Geometry::FAxisAlignedBox2d ConvertBox(Box);
-	if (ConvertBox.MaxDim() == 0)
-	{
-		ConvertBox = UE::Geometry::FAxisAlignedBox2d(FVector2d(-50, -50), FVector2d(50, 50));
-	}
-	FVector2d Center = ConvertBox.Center();
-
 	FRectangleMeshGenerator RectGenerator;
-	RectGenerator.Origin = FVector3d(Center.X, Center.Y, 0);
+	RectGenerator.Origin = FVector3d(0, 0, 0);
 	RectGenerator.Normal = FVector3f::UnitZ();
-	RectGenerator.Width = ConvertBox.Width();
-	RectGenerator.Height = ConvertBox.Height();
+	RectGenerator.Width = DimensionX / 2;
+	RectGenerator.Height = DimensionY / 2;
 	RectGenerator.WidthVertexCount = FMath::Max(0, StepsWidth);
 	RectGenerator.HeightVertexCount = FMath::Max(0, StepsHeight);
 	RectGenerator.bSinglePolygroup = (PrimitiveOptions.PolygroupMode != EGeometryScriptPrimitivePolygroupMode::PerQuad);
@@ -480,7 +654,8 @@ UDynamicMesh* UGeometryScriptLibrary_MeshPrimitiveFunctions::AppendRoundRectangl
 	UDynamicMesh* TargetMesh,
 	FGeometryScriptPrimitiveOptions PrimitiveOptions,
 	FTransform Transform,
-	FBox2D Box,
+	float DimensionX,
+	float DimensionY,
 	float CornerRadius,
 	int32 StepsWidth,
 	int32 StepsHeight,
@@ -493,18 +668,11 @@ UDynamicMesh* UGeometryScriptLibrary_MeshPrimitiveFunctions::AppendRoundRectangl
 		return TargetMesh;
 	}
 
-	UE::Geometry::FAxisAlignedBox2d ConvertBox(Box);
-	if (ConvertBox.MaxDim() == 0)
-	{
-		ConvertBox = UE::Geometry::FAxisAlignedBox2d(FVector2d(-50, -50), FVector2d(50, 50));
-	}
-	FVector2d Center = ConvertBox.Center();
-
 	FRoundedRectangleMeshGenerator RectGenerator;
-	RectGenerator.Origin = FVector3d(Center.X, Center.Y, 0);
+	RectGenerator.Origin = FVector3d(0, 0, 0);
 	RectGenerator.Normal = FVector3f::UnitZ();
-	RectGenerator.Width = ConvertBox.Width();
-	RectGenerator.Height = ConvertBox.Height();
+	RectGenerator.Width = DimensionX / 2;
+	RectGenerator.Height = DimensionY / 2;
 	RectGenerator.WidthVertexCount = FMath::Max(0, StepsWidth);
 	RectGenerator.HeightVertexCount = FMath::Max(0, StepsHeight);
 	RectGenerator.Radius = FMath::Max(FMathf::ZeroTolerance, CornerRadius);
diff --git a/Engine/Plugins/Experimental/GeometryScripting/Source/GeometryScriptingCore/Private/MeshQueryFunctions.cpp b/Engine/Plugins/Experimental/GeometryScripting/Source/GeometryScriptingCore/Private/MeshQueryFunctions.cpp
--- a/Engine/Plugins/Experimental/GeometryScripting/Source/GeometryScriptingCore/Private/MeshQueryFunctions.cpp
+++ b/Engine/Plugins/Experimental/GeometryScripting/Source/GeometryScriptingCore/Private/MeshQueryFunctions.cpp
@@ -294,7 +294,19 @@ bool UGeometryScriptLibrary_MeshQueryFunctions::GetHasMaterialIDs( UDynamicMesh*
 	});
 }
 
+bool UGeometryScriptLibrary_MeshQueryFunctions::GetHasPolygroups( UDynamicMesh* TargetMesh )
+{
+	return SimpleMeshQuery<bool>(TargetMesh, false, [&](const FDynamicMesh3& Mesh) {
+		return Mesh.HasTriangleGroups();
+	});
+}
 
+int UGeometryScriptLibrary_MeshQueryFunctions::GetNumExtendedPolygroupLayers( UDynamicMesh* TargetMesh )
+{
+	return SimpleMeshQuery<int>(TargetMesh, 0, [&](const FDynamicMesh3& Mesh) {
+		return Mesh.HasAttributes() ? Mesh.Attributes()->NumPolygroupLayers() : 0;
+	});
+}
 
 
 
diff --git a/Engine/Plugins/Experimental/GeometryScripting/Source/GeometryScriptingCore/Private/MeshSubdivideFunctions.cpp b/Engine/Plugins/Experimental/GeometryScripting/Source/GeometryScriptingCore/Private/MeshSubdivideFunctions.cpp
new file mode 100644
--- /dev/null
+++ b/Engine/Plugins/Experimental/GeometryScripting/Source/GeometryScriptingCore/Private/MeshSubdivideFunctions.cpp
@@ -0,0 +1,54 @@
+// Copyright Epic Games, Inc. All Rights Reserved.
+
+#include "GeometryScript/MeshSubdivideFunctions.h"
+
+#include "DynamicMesh/DynamicMesh3.h"
+#include "DynamicMesh/DynamicMeshAttributeSet.h"
+#include "DynamicMesh/MeshNormals.h"
+#include "Operations/PNTriangles.h"
+#include "UDynamicMesh.h"
+
+#include "ExplicitUseGeometryMathTypes.h"		// using UE::Geometry::(math types)
+using namespace UE::Geometry;
+
+#define LOCTEXT_NAMESPACE "UGeometryScriptLibrary_MeshSubdivideFunctions"
+
+
+UDynamicMesh* UGeometryScriptLibrary_MeshSubdivideFunctions::ApplyRecursivePNTessellation(
+	UDynamicMesh* TargetMesh,
+	FGeometryScriptPNTessellateOptions Options,
+	int NumIterations,
+	UGeometryScriptDebug* Debug)
+{
+	if (TargetMesh == nullptr)
+	{
+		UE::Geometry::AppendError(Debug, EGeometryScriptErrorType::InvalidInputs, LOCTEXT("ApplyPNTessellation_InvalidInput", "ApplyPNTessellation: TargetMesh is Null"));
+		return TargetMesh;
+	}
+	if (NumIterations <= 0)
+	{
+		return TargetMesh;
+	}
+
+	TargetMesh->EditMesh([&](FDynamicMesh3& EditMesh) 
+	{
+		FPNTriangles Tessellator(&EditMesh);
+		Tessellator.TesselationLevel = NumIterations;
+		Tessellator.Compute();
+
+		if (Options.bRecomputeNormals && EditMesh.HasAttributes() && EditMesh.Attributes()->PrimaryNormals() != nullptr)
+		{
+			FMeshNormals MeshNormals(&EditMesh);
+			MeshNormals.RecomputeOverlayNormals(EditMesh.Attributes()->PrimaryNormals(), true, true);
+			MeshNormals.CopyToOverlay(EditMesh.Attributes()->PrimaryNormals(), false);
+		}
+
+	}, EDynamicMeshChangeType::GeneralEdit, EDynamicMeshAttributeChangeFlags::Unknown, false);
+
+	return TargetMesh;
+}
+
+
+
+
+#undef LOCTEXT_NAMESPACE
diff --git a/Engine/Plugins/Experimental/GeometryScripting/Source/GeometryScriptingCore/Private/MeshUVFunctions.cpp b/Engine/Plugins/Experimental/GeometryScripting/Source/GeometryScriptingCore/Private/MeshUVFunctions.cpp
--- a/Engine/Plugins/Experimental/GeometryScripting/Source/GeometryScriptingCore/Private/MeshUVFunctions.cpp
+++ b/Engine/Plugins/Experimental/GeometryScripting/Source/GeometryScriptingCore/Private/MeshUVFunctions.cpp
@@ -7,6 +7,9 @@
 #include "Polygroups/PolygroupSet.h"
 #include "Parameterization/DynamicMeshUVEditor.h"
 #include "Selections/MeshConnectedComponents.h"
+#include "Parameterization/PatchBasedMeshUVGenerator.h"
+#include "XAtlasWrapper.h"
+
 #include "Async/ParallelFor.h"
 #include "UDynamicMesh.h"
 
@@ -48,6 +51,45 @@ UDynamicMesh* UGeometryScriptLibrary_MeshUVFunctions::SetNumUVSets(
 }
 
 
+UDynamicMesh* UGeometryScriptLibrary_MeshUVFunctions::CopyUVSet(
+	UDynamicMesh* TargetMesh,
+	int FromUVSet,
+	int ToUVSet,
+	UGeometryScriptDebug* Debug)
+{
+	if (TargetMesh == nullptr)
+	{
+		UE::Geometry::AppendError(Debug, EGeometryScriptErrorType::InvalidInputs, LOCTEXT("CopyUVSet_InvalidInput", "CopyUVSet: TargetMesh is Null"));
+		return TargetMesh;
+	}
+	if (FromUVSet == ToUVSet)
+	{
+		UE::Geometry::AppendError(Debug, EGeometryScriptErrorType::InvalidInputs, LOCTEXT("SetNumUVSets_SameSet", "SetNumUVSets: From and To UV Sets have the same Index"));
+		return TargetMesh;
+	}
+	TargetMesh->EditMesh([&](FDynamicMesh3& EditMesh)
+	{
+		FDynamicMeshUVOverlay* FromUVOverlay = nullptr, *ToUVOverlay = nullptr;
+		if (EditMesh.HasAttributes())
+		{
+			FromUVOverlay = (FromUVSet < EditMesh.Attributes()->NumUVLayers()) ? EditMesh.Attributes()->GetUVLayer(FromUVSet) : nullptr;
+			ToUVOverlay = (ToUVSet < EditMesh.Attributes()->NumUVLayers()) ? EditMesh.Attributes()->GetUVLayer(ToUVSet) : nullptr;
+		}
+		if (FromUVOverlay == nullptr || ToUVOverlay == nullptr)
+		{
+			UE::Geometry::AppendError(Debug, EGeometryScriptErrorType::InvalidInputs, LOCTEXT("SetNumUVSets_CopyUVSet", "CopyUVSet: From or To UV Set does not Exist"));
+			return;
+		}
+		FDynamicMeshUVEditor UVEditor(&EditMesh, ToUVOverlay);
+		UVEditor.CopyUVLayer(FromUVOverlay);
+
+	}, EDynamicMeshChangeType::GeneralEdit, EDynamicMeshAttributeChangeFlags::Unknown, false);
+
+	return TargetMesh;
+}
+
+
+
 
 UDynamicMesh* UGeometryScriptLibrary_MeshUVFunctions::SetMeshTriangleUVs(
 	UDynamicMesh* TargetMesh,
@@ -296,6 +338,149 @@ UDynamicMesh* UGeometryScriptLibrary_MeshUVFunctions::SetMeshUVsFromBoxProjectio
 
 
 
+
+
+UDynamicMesh* UGeometryScriptLibrary_MeshUVFunctions::SetMeshUVsFromCylinderProjection(
+	UDynamicMesh* TargetMesh,
+	int UVSetIndex,
+	FTransform CylinderTransform,
+	float SplitAngle,
+	UGeometryScriptDebug* Debug)
+{
+	if (TargetMesh == nullptr)
+	{
+		UE::Geometry::AppendError(Debug, EGeometryScriptErrorType::InvalidInputs, LOCTEXT("SetMeshUVsFromCylinderProjection_InvalidInput", "SetMeshUVsFromCylinderProjection: TargetMesh is Null"));
+		return TargetMesh;
+	}
+
+	bool bHasUVSet = false;
+	ApplyMeshUVEditorOperation(TargetMesh, UVSetIndex, bHasUVSet, Debug,
+		[&](FDynamicMesh3& EditMesh, FDynamicMeshUVOverlay* UVOverlay, FDynamicMeshUVEditor& UVEditor)
+	{
+		TArray<int32> AllTriangles;
+		for (int32 tid : EditMesh.TriangleIndicesItr())
+		{
+			AllTriangles.Add(tid);
+		}
+
+		FFrame3d ProjectionFrame(CylinderTransform);
+		FVector Scale = CylinderTransform.GetScale3D();
+		FVector3d Dimensions = (FVector)Scale;
+		UVEditor.SetTriangleUVsFromCylinderProjection(AllTriangles, [&](const FVector3d& Pos) { return Pos; },
+			ProjectionFrame, Dimensions, SplitAngle);
+	});
+	if (bHasUVSet == false)
+	{
+		UE::Geometry::AppendError(Debug, EGeometryScriptErrorType::InvalidInputs, LOCTEXT("SetMeshUVsFromCylinderProjection_InvalidUVSet", "SetMeshUVsFromCylinderProjection: UVSetIndex does not exist on TargetMesh"));
+	}
+
+	return TargetMesh;
+}
+
+
+
+UDynamicMesh* UGeometryScriptLibrary_MeshUVFunctions::RecomputeMeshUVs( 
+	UDynamicMesh* TargetMesh, 
+	int UVSetIndex,
+	FGeometryScriptRecomputeUVsOptions Options,
+	UGeometryScriptDebug* Debug)
+{
+	if (TargetMesh == nullptr)
+	{
+		UE::Geometry::AppendError(Debug, EGeometryScriptErrorType::InvalidInputs, LOCTEXT("RecomputeMeshUVs_InvalidInput", "RecomputeMeshUVs: TargetMesh is Null"));
+		return TargetMesh;
+	}
+
+	bool bHasUVSet = false;
+	ApplyMeshUVEditorOperation(TargetMesh, UVSetIndex, bHasUVSet, Debug,
+		[&](FDynamicMesh3& EditMesh, FDynamicMeshUVOverlay* UVOverlay, FDynamicMeshUVEditor& UVEditor)
+	{
+
+
+		TUniquePtr<FPolygroupSet> IslandSourceGroups;
+		if (Options.IslandSource == EGeometryScriptUVIslandSource::PolyGroups)
+		{
+			FPolygroupLayer InputGroupLayer{ Options.GroupLayer.bDefaultLayer, Options.GroupLayer.ExtendedLayerIndex };
+			if (InputGroupLayer.CheckExists(&EditMesh))
+			{
+
+				IslandSourceGroups = MakeUnique<FPolygroupSet>(&EditMesh, InputGroupLayer);
+			}
+			else
+			{
+				UE::Geometry::AppendError(Debug, EGeometryScriptErrorType::OperationFailed, LOCTEXT("RecomputeMeshUVs_MissingGroups", "RecomputeMeshUVs: Requested Polygroup Layer does not exist"));
+				return;
+			}
+		}
+
+		// find group-connected-components
+		FMeshConnectedComponents ConnectedComponents(&EditMesh);
+		if (Options.IslandSource == EGeometryScriptUVIslandSource::PolyGroups)
+		{
+			ConnectedComponents.FindConnectedTriangles([&](int32 CurTri, int32 NbrTri) {
+				return IslandSourceGroups->GetTriangleGroup(CurTri) == IslandSourceGroups->GetTriangleGroup(NbrTri);
+			});
+		}
+		else
+		{
+			ConnectedComponents.FindConnectedTriangles([&](int32 Triangle0, int32 Triangle1) {
+				return UVOverlay->AreTrianglesConnected(Triangle0, Triangle1);
+			});
+		}
+
+		int32 NumComponents = ConnectedComponents.Num();
+		TArray<bool> bComponentSolved;
+		bComponentSolved.Init(false, NumComponents);
+		int32 SuccessCount = 0;
+		for (int32 k = 0; k < NumComponents; ++k)
+		{
+			const TArray<int32>& ComponentTris = ConnectedComponents[k].Indices;
+			bComponentSolved[k] = false;
+			switch (Options.Method)
+			{
+			case EGeometryScriptUVFlattenMethod::ExpMap:
+			{
+				FDynamicMeshUVEditor::FExpMapOptions ExpMapOptions;
+				ExpMapOptions.NormalSmoothingRounds = Options.ExpMapOptions.NormalSmoothingRounds;
+				ExpMapOptions.NormalSmoothingAlpha = Options.ExpMapOptions.NormalSmoothingAlpha;
+				bComponentSolved[k] = UVEditor.SetTriangleUVsFromExpMap(ComponentTris, ExpMapOptions);
+			}
+			break;
+
+			case EGeometryScriptUVFlattenMethod::Conformal:
+				bComponentSolved[k] = UVEditor.SetTriangleUVsFromFreeBoundaryConformal(ComponentTris);
+				if ( bComponentSolved[k] )
+				{
+					UVEditor.ScaleUVAreaTo3DArea(ComponentTris, true);
+				}
+				break;
+			}
+		}
+
+		if (Options.bAutoAlignIslandsWithAxes)
+		{
+			ParallelFor(NumComponents, [&](int32 k)
+			{
+				if (bComponentSolved[k])
+				{
+					const TArray<int32>& ComponentTris = ConnectedComponents[k].Indices;
+					UVEditor.AutoOrientUVArea(ComponentTris);
+				};
+			});
+		}
+
+	});
+	if (bHasUVSet == false)
+	{
+		UE::Geometry::AppendError(Debug, EGeometryScriptErrorType::InvalidInputs, LOCTEXT("RecomputeMeshUVs_InvalidUVSet", "RecomputeMeshUVs: UVSetIndex does not exist on TargetMesh"));
+	}
+
+	return TargetMesh;
+}
+
+
+
+
 UDynamicMesh* UGeometryScriptLibrary_MeshUVFunctions::RepackMeshUVs( 
 	UDynamicMesh* TargetMesh, 
 	int UVSetIndex,
@@ -337,4 +522,187 @@ UDynamicMesh* UGeometryScriptLibrary_MeshUVFunctions::RepackMeshUVs(
 
 
 
+
+UDynamicMesh* UGeometryScriptLibrary_MeshUVFunctions::AutoGeneratePatchBuilderMeshUVs( 
+	UDynamicMesh* TargetMesh, 
+	int UVSetIndex,
+	FGeometryScriptPatchBuilderOptions Options,
+	UGeometryScriptDebug* Debug)
+{
+	if (TargetMesh == nullptr)
+	{
+		UE::Geometry::AppendError(Debug, EGeometryScriptErrorType::InvalidInputs, LOCTEXT("AutoGeneratePatchBuilderMeshUVs_InvalidInput", "AutoGeneratePatchBuilderMeshUVs: TargetMesh is Null"));
+		return TargetMesh;
+	}
+
+	bool bHasUVSet = false;
+	ApplyMeshUVEditorOperation(TargetMesh, UVSetIndex, bHasUVSet, Debug,
+		[&](FDynamicMesh3& EditMesh, FDynamicMeshUVOverlay* UVOverlay, FDynamicMeshUVEditor& UVEditor)
+	{
+		FPatchBasedMeshUVGenerator UVGenerator;
+
+		TUniquePtr<FPolygroupSet> PolygroupConstraint;
+		if (Options.bRespectInputGroups)
+		{
+			FPolygroupLayer InputGroupLayer{ Options.GroupLayer.bDefaultLayer, Options.GroupLayer.ExtendedLayerIndex };
+			if (InputGroupLayer.CheckExists(&EditMesh))
+			{
+
+				PolygroupConstraint = MakeUnique<FPolygroupSet>(&EditMesh, InputGroupLayer);
+				UVGenerator.GroupConstraint = PolygroupConstraint.Get();
+			}
+			else
+			{
+				UE::Geometry::AppendError(Debug, EGeometryScriptErrorType::InvalidInputs, LOCTEXT("AutoGeneratePatchBuilderMeshUVs_MissingGruops", "AutoGeneratePatchBuilderMeshUVs: Requested Polygroup Layer does not exist"));
+			}
+		}
+
+		UVGenerator.TargetPatchCount = FMath::Max(1,Options.InitialPatchCount);
+		UVGenerator.bNormalWeightedPatches = true;
+		UVGenerator.PatchNormalWeight = FMath::Clamp(Options.PatchCurvatureAlignmentWeight, 0.0, 999999.0);
+		UVGenerator.MinPatchSize = FMath::Max(1,Options.MinPatchSize);
+
+		UVGenerator.MergingThreshold = FMath::Clamp(Options.PatchMergingMetricThresh, 0.001, 9999.0);
+		UVGenerator.MaxNormalDeviationDeg = FMath::Clamp(Options.PatchMergingAngleThresh, 0.0, 180.0);
+
+		UVGenerator.NormalSmoothingRounds = FMath::Clamp(Options.ExpMapOptions.NormalSmoothingRounds, 0, 9999);
+		UVGenerator.NormalSmoothingAlpha = FMath::Clamp(Options.ExpMapOptions.NormalSmoothingAlpha, 0.0, 1.0);
+
+		UVGenerator.bAutoPack = Options.bAutoPack;
+		if (Options.bAutoPack)
+		{
+			UVGenerator.bAutoAlignPatches = Options.PackingOptions.bOptimizeIslandRotation;
+			UVGenerator.PackingTextureResolution = FMath::Clamp(Options.PackingOptions.TargetImageWidth, 16, 4096);
+			UVGenerator.PackingGutterWidth = 1.0;
+		}
+		FGeometryResult Result = UVGenerator.AutoComputeUVs(*UVEditor.GetMesh(), *UVEditor.GetOverlay(), nullptr);
+
+		if (Result.HasFailed())
+		{
+			UE::Geometry::AppendError(Debug, EGeometryScriptErrorType::OperationFailed, LOCTEXT("AutoGeneratePatchBuilderMeshUVs_Failed", "AutoGeneratePatchBuilderMeshUVs: UV Generation Failed"));
+		}
+
+	});
+	if (bHasUVSet == false)
+	{
+		UE::Geometry::AppendError(Debug, EGeometryScriptErrorType::InvalidInputs, LOCTEXT("AutoGeneratePatchBuilderMeshUVs_InvalidUVSet", "AutoGeneratePatchBuilderMeshUVs: UVSetIndex does not exist on TargetMesh"));
+	}
+
+	return TargetMesh;
+}
+
+
+
+
+UDynamicMesh* UGeometryScriptLibrary_MeshUVFunctions::AutoGenerateXAtlasMeshUVs( 
+	UDynamicMesh* TargetMesh, 
+	int UVSetIndex,
+	FGeometryScriptXAtlasOptions Options,
+	UGeometryScriptDebug* Debug)
+{
+	if (TargetMesh == nullptr)
+	{
+		UE::Geometry::AppendError(Debug, EGeometryScriptErrorType::InvalidInputs, LOCTEXT("AutoGenerateXAtlasMeshUVs_InvalidInput", "AutoGenerateXAtlasMeshUVs: TargetMesh is Null"));
+		return TargetMesh;
+	}
+
+	bool bHasUVSet = false;
+	ApplyMeshUVEditorOperation(TargetMesh, UVSetIndex, bHasUVSet, Debug,
+		[&](FDynamicMesh3& EditMesh, FDynamicMeshUVOverlay* UVOverlay, FDynamicMeshUVEditor& UVEditor)
+	{
+		if (EditMesh.IsCompact() == false)
+		{
+			UE::Geometry::AppendError(Debug, EGeometryScriptErrorType::InvalidInputs, LOCTEXT("AutoGenerateXAtlasMeshUVs_NonCompact", "AutoGenerateXAtlasMeshUVs: TargetMesh must be Compacted before running XAtlas"));
+			return;
+		}
+
+
+		const bool bFixOrientation = false;
+		//const bool bFixOrientation = true;
+		//FDynamicMesh3 FlippedMesh(EMeshComponents::FaceGroups);
+		//FlippedMesh.Copy(Mesh, false, false, false, false);
+		//if (bFixOrientation)
+		//{
+		//	FlippedMesh.ReverseOrientation(false);
+		//}
+
+
+		int32 NumVertices = EditMesh.VertexCount();
+		TArray<FVector3f> VertexBuffer;
+		VertexBuffer.SetNum(NumVertices);
+		for (int32 k = 0; k < NumVertices; ++k)
+		{
+			VertexBuffer[k] = (FVector3f)EditMesh.GetVertex(k);
+		}
+
+		TArray<int32> IndexBuffer;
+		IndexBuffer.Reserve(EditMesh.TriangleCount()*3);
+		for (FIndex3i Triangle : EditMesh.TrianglesItr())
+		{
+			IndexBuffer.Add(Triangle.A);
+			IndexBuffer.Add(Triangle.B);
+			IndexBuffer.Add(Triangle.C);
+		}
+
+		TArray<FVector2D> UVVertexBuffer;
+		TArray<int32>     UVIndexBuffer;
+		TArray<int32>     VertexRemapArray; // This maps the UV vertices to the original position vertices.  Note multiple UV vertices might share the same positional vertex (due to UV boundaries)
+		XAtlasWrapper::XAtlasChartOptions ChartOptions;
+		ChartOptions.MaxIterations = Options.MaxIterations;
+		XAtlasWrapper::XAtlasPackOptions PackOptions;
+		bool bSuccess = XAtlasWrapper::ComputeUVs(IndexBuffer, VertexBuffer, ChartOptions, PackOptions,
+			UVVertexBuffer, UVIndexBuffer, VertexRemapArray);
+		if (bSuccess == false)
+		{
+			UE::Geometry::AppendError(Debug, EGeometryScriptErrorType::OperationFailed, LOCTEXT("AutoGenerateXAtlasMeshUVs_Failed", "AutoGenerateXAtlasMeshUVs: UV Generation Failed"));
+			return;
+		}
+
+		UVOverlay->ClearElements();
+
+		int32 NumUVs = UVVertexBuffer.Num();
+		TArray<int32> UVOffsetToElID;  UVOffsetToElID.Reserve(NumUVs);
+		for (int32 i = 0; i < NumUVs; ++i)
+		{
+			FVector2D UV = UVVertexBuffer[i];
+			const int32 VertOffset = VertexRemapArray[i];		// The associated VertID in the dynamic mesh
+			const int32 NewID = UVOverlay->AppendElement(FVector2f(UV));		// add the UV to the mesh overlay
+			UVOffsetToElID.Add(NewID);
+		}
+
+		int32 NumUVTris = UVIndexBuffer.Num() / 3;
+		for (int32 i = 0; i < NumUVTris; ++i)
+		{
+			int32 t = i * 3;
+			FIndex3i UVTri(UVIndexBuffer[t], UVIndexBuffer[t + 1], UVIndexBuffer[t + 2]);	// The triangle in UV space
+			FIndex3i TriVertIDs;				// the triangle in terms of the VertIDs in the DynamicMesh
+			for (int c = 0; c < 3; ++c)
+			{
+				int32 Offset = VertexRemapArray[UVTri[c]];		// the offset for this vertex in the LinearMesh
+				TriVertIDs[c] = Offset;
+			}
+
+			// NB: this could be slow.. 
+			int32 TriID = EditMesh.FindTriangle(TriVertIDs[0], TriVertIDs[1], TriVertIDs[2]);
+			if (TriID != IndexConstants::InvalidID)
+			{
+				FIndex3i ElTri = (bFixOrientation) ?
+					FIndex3i(UVOffsetToElID[UVTri[1]], UVOffsetToElID[UVTri[0]], UVOffsetToElID[UVTri[2]])
+					: FIndex3i(UVOffsetToElID[UVTri[0]], UVOffsetToElID[UVTri[1]], UVOffsetToElID[UVTri[2]]);
+				UVOverlay->SetTriangle(TriID, ElTri);
+			}
+		}
+
+	});
+	if (bHasUVSet == false)
+	{
+		UE::Geometry::AppendError(Debug, EGeometryScriptErrorType::InvalidInputs, LOCTEXT("AutoGenerateXAtlasMeshUVs_InvalidUVSet", "AutoGenerateXAtlasMeshUVs: UVSetIndex does not exist on TargetMesh"));
+	}
+
+	return TargetMesh;
+}
+
+
+
+
 #undef LOCTEXT_NAMESPACE
\ No newline at end of file
diff --git a/Engine/Plugins/Experimental/GeometryScripting/Source/GeometryScriptingCore/Private/MeshVertexColorFunctions.cpp b/Engine/Plugins/Experimental/GeometryScripting/Source/GeometryScriptingCore/Private/MeshVertexColorFunctions.cpp
--- a/Engine/Plugins/Experimental/GeometryScripting/Source/GeometryScriptingCore/Private/MeshVertexColorFunctions.cpp
+++ b/Engine/Plugins/Experimental/GeometryScripting/Source/GeometryScriptingCore/Private/MeshVertexColorFunctions.cpp
@@ -11,6 +11,77 @@ using namespace UE::Geometry;
 #define LOCTEXT_NAMESPACE "UGeometryScriptLibrary_MeshVertexColorFunctions"
 
 
+static FLinearColor CombineColors(FLinearColor ExistingColor, FLinearColor NewColor, FGeometryScriptColorFlags Flags)
+{
+	ExistingColor.R = (Flags.bRed) ? ExistingColor.R : NewColor.R;
+	ExistingColor.G = (Flags.bBlue) ? ExistingColor.G : NewColor.G;
+	ExistingColor.B = (Flags.bGreen) ? ExistingColor.B : NewColor.B;
+	ExistingColor.A = (Flags.bAlpha) ? ExistingColor.A : NewColor.A;
+	return ExistingColor;
+}
+
+UDynamicMesh* UGeometryScriptLibrary_MeshVertexColorFunctions::SetMeshConstantVertexColor(
+	UDynamicMesh* TargetMesh,
+	FLinearColor Color,
+	FGeometryScriptColorFlags Flags,
+	bool bClearExisting,
+	UGeometryScriptDebug* Debug)
+{
+	if (TargetMesh == nullptr)
+	{
+		UE::Geometry::AppendError(Debug, EGeometryScriptErrorType::InvalidInputs, LOCTEXT("SetMeshConstantVertexColor_InvalidInput", "SetMeshConstantVertexColor: TargetMesh is Null"));
+		return TargetMesh;
+	}
+
+	TargetMesh->EditMesh([&](FDynamicMesh3& EditMesh) 
+	{
+		bool bCreated = false;
+		if (EditMesh.HasAttributes() == false)
+		{
+			EditMesh.EnableAttributes();
+			bCreated = true;
+		}
+		if (EditMesh.Attributes()->HasPrimaryColors() == false)
+		{
+			EditMesh.Attributes()->EnablePrimaryColors();
+			bCreated = true;
+		}
+		FDynamicMeshColorOverlay* Colors = EditMesh.Attributes()->PrimaryColors();
+		if (bClearExisting && bCreated == false)
+		{
+			Colors->ClearElements();
+		}
+		if (Colors->ElementCount() == 0)
+		{
+			TArray<int32> ElemIDs;
+			ElemIDs.SetNum(EditMesh.MaxVertexID());
+			for (int32 VertexID : EditMesh.VertexIndicesItr())
+			{
+				ElemIDs[VertexID] = Colors->AppendElement(FVector4f(Color.R, Color.G, Color.B, Color.A));
+			}
+			for (int32 TriangleID : EditMesh.TriangleIndicesItr())
+			{
+				FIndex3i Triangle = EditMesh.GetTriangle(TriangleID);
+				Colors->SetTriangle(TriangleID, FIndex3i(ElemIDs[Triangle.A], ElemIDs[Triangle.B], ElemIDs[Triangle.C]) );
+			}
+		}
+		else
+		{
+			for (int32 ElementID : Colors->ElementIndicesItr())
+			{
+				FLinearColor Existing = ToLinearColor(Colors->GetElement(ElementID));
+				FLinearColor NewColor = CombineColors(Existing, Color, Flags);
+				Colors->SetElement(ElementID, ToVector4<float>(NewColor));
+			}
+		}
+
+	}, EDynamicMeshChangeType::GeneralEdit, EDynamicMeshAttributeChangeFlags::Unknown, false);
+
+	return TargetMesh;
+}
+
+
+
 UDynamicMesh* UGeometryScriptLibrary_MeshVertexColorFunctions::SetMeshPerVertexColors(
 	UDynamicMesh* TargetMesh,
 	const TArray<FLinearColor>& VertexColors,
diff --git a/Engine/Plugins/Experimental/GeometryScripting/Source/GeometryScriptingCore/Public/GeometryScript/GeometryScriptTypes.h b/Engine/Plugins/Experimental/GeometryScripting/Source/GeometryScriptingCore/Public/GeometryScript/GeometryScriptTypes.h
--- a/Engine/Plugins/Experimental/GeometryScripting/Source/GeometryScriptingCore/Public/GeometryScript/GeometryScriptTypes.h
+++ b/Engine/Plugins/Experimental/GeometryScripting/Source/GeometryScriptingCore/Public/GeometryScript/GeometryScriptTypes.h
@@ -129,6 +129,50 @@ public:
 
 
 
+
+//
+// Colors
+//
+
+USTRUCT(BlueprintType)
+struct GEOMETRYSCRIPTINGCORE_API FGeometryScriptColorFlags
+{
+	GENERATED_BODY()
+public:
+	UPROPERTY(BlueprintReadWrite, Category = Color)
+	bool bRed = true;
+
+	UPROPERTY(BlueprintReadWrite, Category = Color)
+	bool bGreen = true;
+
+	UPROPERTY(BlueprintReadWrite, Category = Color)
+	bool bBlue = true;
+
+	UPROPERTY(BlueprintReadWrite, Category = Color)
+	bool bAlpha = true;
+
+	bool AllSet() const { return bRed && bGreen && bBlue && bAlpha; }
+};
+
+
+
+//
+// Polygroups
+//
+
+USTRUCT(BlueprintType)
+struct GEOMETRYSCRIPTINGCORE_API FGeometryScriptGroupLayer
+{
+	GENERATED_BODY()
+public:
+	UPROPERTY(BlueprintReadWrite, Category = LOD)
+	bool bDefaultLayer = true;
+
+	UPROPERTY(BlueprintReadWrite, Category = LOD)
+	int ExtendedLayerIndex = 0;
+};
+
+
 //
 // Spatial data structures
 //
@@ -146,6 +190,7 @@ public:
 
 
 
+
 //
 // Errors/Debugging
 //
diff --git a/Engine/Plugins/Experimental/GeometryScripting/Source/GeometryScriptingCore/Public/GeometryScript/MeshAssetFunctions.h b/Engine/Plugins/Experimental/GeometryScripting/Source/GeometryScriptingCore/Public/GeometryScript/MeshAssetFunctions.h
--- a/Engine/Plugins/Experimental/GeometryScripting/Source/GeometryScriptingCore/Public/GeometryScript/MeshAssetFunctions.h
+++ b/Engine/Plugins/Experimental/GeometryScripting/Source/GeometryScriptingCore/Public/GeometryScript/MeshAssetFunctions.h
@@ -35,6 +35,15 @@ struct GEOMETRYSCRIPTINGCORE_API FGeometryScriptCopyMeshToAssetOptions
 {
 	GENERATED_BODY()
 public:
+	UPROPERTY(BlueprintReadWrite, Category = Options)
+	bool bEnableRecomputeNormals = false;
+
+	UPROPERTY(BlueprintReadWrite, Category = Options)
+	bool bEnableRecomputeTangents = false;
+
+	UPROPERTY(BlueprintReadWrite, Category = Options)
+	bool bEnableRemoveDegenerates = false;
+
 	UPROPERTY(BlueprintReadWrite, Category = Options)
 	bool bEmitTransaction = true;
 
diff --git a/Engine/Plugins/Experimental/GeometryScripting/Source/GeometryScriptingCore/Public/GeometryScript/MeshBooleanFunctions.h b/Engine/Plugins/Experimental/GeometryScripting/Source/GeometryScriptingCore/Public/GeometryScript/MeshBooleanFunctions.h
--- a/Engine/Plugins/Experimental/GeometryScripting/Source/GeometryScriptingCore/Public/GeometryScript/MeshBooleanFunctions.h
+++ b/Engine/Plugins/Experimental/GeometryScripting/Source/GeometryScriptingCore/Public/GeometryScript/MeshBooleanFunctions.h
@@ -71,6 +71,28 @@ public:
 
 	UPROPERTY(BlueprintReadWrite, Category = Options)
 	bool bFlipCutSide = false;
+
+	UPROPERTY(BlueprintReadWrite, Category = Options)
+	float UVWorldDimension = 1.0;
+};
+
+
+USTRUCT(BlueprintType)
+struct GEOMETRYSCRIPTINGCORE_API FGeometryScriptMeshPlaneSliceOptions
+{
+	GENERATED_BODY()
+public:
+	UPROPERTY(BlueprintReadWrite, Category = Options)
+	bool bFillHoles = true;
+
+	UPROPERTY(BlueprintReadWrite, Category = Options)
+	bool bFillSpans = true;
+
+	UPROPERTY(BlueprintReadWrite, Category = Options)
+	float GapWidth = 0.01;
+
+	UPROPERTY(BlueprintReadWrite, Category = Options)
+	float UVWorldDimension = 1.0;
 };
 
 
@@ -122,18 +144,24 @@ public:
 	static UPARAM(DisplayName = "Target Mesh") UDynamicMesh* 
 	ApplyMeshPlaneCut( 
 		UDynamicMesh* TargetMesh, 
-		FVector CutPlaneOrigin,
-		FVector CutPlaneNormal,
+		FTransform CutFrame,
 		FGeometryScriptMeshPlaneCutOptions Options,
 		UGeometryScriptDebug* Debug = nullptr);
 
+	UFUNCTION(BlueprintCallable, Category = "GeometryScript|Booleans", meta=(ScriptMethod))
+	static UPARAM(DisplayName = "Target Mesh") UDynamicMesh* 
+	ApplyMeshPlaneSlice( 
+		UDynamicMesh* TargetMesh, 
+		FTransform CutFrame,
+		FGeometryScriptMeshPlaneSliceOptions Options,
+		UGeometryScriptDebug* Debug = nullptr);
+
 
 	UFUNCTION(BlueprintCallable, Category = "GeometryScript|Booleans", meta=(ScriptMethod))
 	static UPARAM(DisplayName = "Target Mesh") UDynamicMesh* 
 	ApplyMeshMirror( 
 		UDynamicMesh* TargetMesh, 
-		FVector MirrorPlaneOrigin,
-		FVector MirrorPlaneNormal,
+		FTransform MirrorFrame,
 		FGeometryScriptMeshMirrorOptions Options,
 		UGeometryScriptDebug* Debug = nullptr);
 
diff --git a/Engine/Plugins/Experimental/GeometryScripting/Source/GeometryScriptingCore/Public/GeometryScript/MeshDeformFunctions.h b/Engine/Plugins/Experimental/GeometryScripting/Source/GeometryScriptingCore/Public/GeometryScript/MeshDeformFunctions.h
--- a/Engine/Plugins/Experimental/GeometryScripting/Source/GeometryScriptingCore/Public/GeometryScript/MeshDeformFunctions.h
+++ b/Engine/Plugins/Experimental/GeometryScripting/Source/GeometryScriptingCore/Public/GeometryScript/MeshDeformFunctions.h
@@ -81,6 +81,105 @@ public:
 
 
 
+USTRUCT(BlueprintType)
+struct GEOMETRYSCRIPTINGCORE_API FGeometryScriptPerlinNoiseLayerOptions
+{
+	GENERATED_BODY()
+public:
+	UPROPERTY(BlueprintReadWrite, Category = Options)
+	float Magnitude = 5.0;
+
+	UPROPERTY(BlueprintReadWrite, Category = Options)
+	float Frequency = 0.25;
+
+	UPROPERTY(BlueprintReadWrite, Category = Options)
+	FVector FrequencyShift = FVector::Zero();
+
+	UPROPERTY(BlueprintReadWrite, Category = Options)
+	int RandomSeed = 0;
+};
+
+
+
+UENUM(BlueprintType)
+enum class EGeometryScriptMathWarpType : uint8
+{
+	SinWave1D = 0,
+	SinWave2D = 1,
+	SinWave3D = 2
+};
+
+
+USTRUCT(BlueprintType)
+struct GEOMETRYSCRIPTINGCORE_API FGeometryScriptMathWarpOptions
+{
+	GENERATED_BODY()
+public:
+	UPROPERTY(BlueprintReadWrite, Category = Options)
+	float Magnitude = 5.0f;
+
+	UPROPERTY(BlueprintReadWrite, Category = Options)
+	float Frequency = 0.25f;
+
+	UPROPERTY(BlueprintReadWrite, Category = Options)
+	float FrequencyShift = 0.0;
+};
+
+
+
+
+USTRUCT(BlueprintType)
+struct GEOMETRYSCRIPTINGCORE_API FGeometryScriptPerlinNoiseOptions
+{
+	GENERATED_BODY()
+public:
+	UPROPERTY(BlueprintReadWrite, Category = Options)
+	FGeometryScriptPerlinNoiseLayerOptions BaseLayer;
+
+	UPROPERTY(BlueprintReadWrite, Category = Options)
+	bool bApplyAlongNormal = true;
+};
+
+
+USTRUCT(BlueprintType)
+struct GEOMETRYSCRIPTINGCORE_API FGeometryScriptIterativeMeshSmoothingOptions
+{
+	GENERATED_BODY()
+public:
+	UPROPERTY(BlueprintReadWrite, Category = Options)
+	int NumIterations = 10;
+
+	UPROPERTY(BlueprintReadWrite, Category = Options)
+	float Alpha = 0.2;
+};
+
+
+
+
+USTRUCT(BlueprintType)
+struct GEOMETRYSCRIPTINGCORE_API FGeometryScriptDisplaceFromTextureOptions
+{
+	GENERATED_BODY()
+public:
+	UPROPERTY(BlueprintReadWrite, Category = Options)
+	float Magnitude = 1.0f;
+
+	UPROPERTY(BlueprintReadWrite, Category = Options)
+	FVector2D UVScale = FVector2D(1,1);
+
+	UPROPERTY(BlueprintReadWrite, Category = Options)
+	FVector2D UVOffset = FVector2D(0,0);
+
+	UPROPERTY(BlueprintReadWrite, Category = Options)
+	float Center = 0.5;
+
+	UPROPERTY(BlueprintReadWrite, Category = Options)
+	int ImageChannel = 0;
+};
+
+
+
+
 
 UCLASS(meta = (ScriptName = "GeometryScript_MeshDeformers"))
 class GEOMETRYSCRIPTINGCORE_API UGeometryScriptLibrary_MeshDeformFunctions : public UBlueprintFunctionLibrary
@@ -118,4 +217,41 @@ public:
 		float FlarePercentY = 0,
 		float FlareExtent = 50,
 		UGeometryScriptDebug* Debug = nullptr);
+
+
+	UFUNCTION(BlueprintCallable, Category = "GeometryScript|Deformations", meta=(ScriptMethod))
+	static UPARAM(DisplayName = "Target Mesh") UDynamicMesh* 
+	ApplyMathWarpToMesh(  
+		UDynamicMesh* TargetMesh, 
+		FTransform WarpOrientation,
+		EGeometryScriptMathWarpType WarpType,
+		FGeometryScriptMathWarpOptions Options,
+		UGeometryScriptDebug* Debug = nullptr);
+
+
+	UFUNCTION(BlueprintCallable, Category = "GeometryScript|Deformations", meta=(ScriptMethod))
+	static UPARAM(DisplayName = "Target Mesh") UDynamicMesh* 
+	ApplyPerlinNoiseToMesh(  
+		UDynamicMesh* TargetMesh, 
+		FGeometryScriptPerlinNoiseOptions Options,
+		UGeometryScriptDebug* Debug = nullptr);
+
+	UFUNCTION(BlueprintCallable, Category = "GeometryScript|Deformations", meta=(ScriptMethod))
+	static UPARAM(DisplayName = "Target Mesh") UDynamicMesh* 
+	ApplyIterativeSmoothingToMesh(  
+		UDynamicMesh* TargetMesh, 
+		FGeometryScriptIterativeMeshSmoothingOptions Options,
+		UGeometryScriptDebug* Debug = nullptr);
+
+
+	UFUNCTION(BlueprintCallable, Category = "GeometryScript|Deformations", meta=(ScriptMethod))
+	static UPARAM(DisplayName = "Target Mesh") UDynamicMesh* 
+	ApplyDisplaceFromTextureMap(  
+		UDynamicMesh* TargetMesh, 
+		UTexture2D* Texture,
+		FGeometryScriptDisplaceFromTextureOptions Options,
+		int32 UVLayer = 0,
+		UGeometryScriptDebug* Debug = nullptr);
+
+	
 };
\ No newline at end of file
diff --git a/Engine/Plugins/Experimental/GeometryScripting/Source/GeometryScriptingCore/Public/GeometryScript/MeshNormalsFunctions.h b/Engine/Plugins/Experimental/GeometryScripting/Source/GeometryScriptingCore/Public/GeometryScript/MeshNormalsFunctions.h
--- a/Engine/Plugins/Experimental/GeometryScripting/Source/GeometryScriptingCore/Public/GeometryScript/MeshNormalsFunctions.h
+++ b/Engine/Plugins/Experimental/GeometryScripting/Source/GeometryScriptingCore/Public/GeometryScript/MeshNormalsFunctions.h
@@ -42,6 +42,28 @@ public:
 };
 
 
+UENUM(BlueprintType)
+enum class EGeometryScriptTangentTypes : uint8
+{
+	FastMikkT = 0,
+	PerTriangle = 1
+};
+
+
+USTRUCT(BlueprintType)
+struct GEOMETRYSCRIPTINGCORE_API FGeometryScriptTangentsOptions
+{
+	GENERATED_BODY()
+public:
+	UPROPERTY(BlueprintReadWrite, Category = Options)
+	EGeometryScriptTangentTypes Type = EGeometryScriptTangentTypes::FastMikkT;
+
+	UPROPERTY(BlueprintReadWrite, Category = Options)
+	int UVLayer = 0;
+};
+
+
+
 UCLASS(meta = (ScriptName = "GeometryScript_Normals"))
 class GEOMETRYSCRIPTINGCORE_API UGeometryScriptLibrary_MeshNormalsFunctions : public UBlueprintFunctionLibrary
 {
@@ -87,7 +109,12 @@ public:
 		FGeometryScriptCalculateNormalsOptions CalculateOptions,
 		UGeometryScriptDebug* Debug = nullptr);
 
-
+	UFUNCTION(BlueprintCallable, Category = "GeometryScript|Normals", meta=(ScriptMethod))
+	static UPARAM(DisplayName = "Target Mesh") UDynamicMesh* 
+	ComputeTangents( 
+		UDynamicMesh* TargetMesh, 
+		FGeometryScriptTangentsOptions Options,
+		UGeometryScriptDebug* Debug = nullptr);
 
 
 	UFUNCTION(BlueprintCallable, Category = "GeometryScript|Normals", meta=(ScriptMethod))
diff --git a/Engine/Plugins/Experimental/GeometryScripting/Source/GeometryScriptingCore/Public/GeometryScript/MeshPolygroupFunctions.h b/Engine/Plugins/Experimental/GeometryScripting/Source/GeometryScriptingCore/Public/GeometryScript/MeshPolygroupFunctions.h
new file mode 100644
--- /dev/null
+++ b/Engine/Plugins/Experimental/GeometryScripting/Source/GeometryScriptingCore/Public/GeometryScript/MeshPolygroupFunctions.h
@@ -0,0 +1,64 @@
+// Copyright Epic Games, Inc. All Rights Reserved.
+
+#pragma once
+
+#include "CoreMinimal.h"
+#include "Kismet/BlueprintFunctionLibrary.h"
+#include "GeometryScript/GeometryScriptTypes.h"
+#include "MeshPolygroupFunctions.generated.h"
+
+class UDynamicMesh;
+
+
+UCLASS(meta = (ScriptName = "GeometryScript_Polygroups"))
+class GEOMETRYSCRIPTINGCORE_API UGeometryScriptLibrary_MeshPolygroupFunctions : public UBlueprintFunctionLibrary
+{
+	GENERATED_BODY()
+public:
+
+	UFUNCTION(BlueprintCallable, Category = "GeometryScript|Polygroups", meta=(ScriptMethod))
+	static UPARAM(DisplayName = "Target Mesh") UDynamicMesh* 
+	EnablePolygroups( UDynamicMesh* TargetMesh, UGeometryScriptDebug* Debug = nullptr );
+
+	UFUNCTION(BlueprintCallable, Category = "GeometryScript|Polygroups", meta=(ScriptMethod))
+	static UPARAM(DisplayName = "Target Mesh") UDynamicMesh* 
+	SetNumExtendedPolygroupLayers( 
+		UDynamicMesh* TargetMesh, 
+		int NumLayers,
+		UGeometryScriptDebug* Debug = nullptr );
+
+	UFUNCTION(BlueprintCallable, Category = "GeometryScript|Polygroups", meta=(ScriptMethod))
+	static UPARAM(DisplayName = "Target Mesh") UDynamicMesh* 
+	ClearPolygroups( 
+		UDynamicMesh* TargetMesh, 
+		FGeometryScriptGroupLayer GroupLayer,
+		int ClearValue = 0,
+		UGeometryScriptDebug* Debug = nullptr );
+
+	UFUNCTION(BlueprintCallable, Category = "GeometryScript|Polygroups", meta=(ScriptMethod))
+	static UPARAM(DisplayName = "Target Mesh") UDynamicMesh* 
+	CopyPolygroupsLayer( 
+		UDynamicMesh* TargetMesh, 
+		FGeometryScriptGroupLayer FromGroupLayer,
+		FGeometryScriptGroupLayer ToGroupLayer,
+		UGeometryScriptDebug* Debug = nullptr );
+
+
+	UFUNCTION(BlueprintCallable, Category = "GeometryScript|Polygroups", meta=(ScriptMethod))
+	static UPARAM(DisplayName = "Target Mesh") UDynamicMesh* 
+	ConvertUVIslandsToPolygroups( 
+		UDynamicMesh* TargetMesh, 
+		FGeometryScriptGroupLayer GroupLayer,
+		int UVLayer = 0,
+		UGeometryScriptDebug* Debug = nullptr );
+
+	UFUNCTION(BlueprintCallable, Category = "GeometryScript|Polygroups", meta=(ScriptMethod))
+	static UPARAM(DisplayName = "Target Mesh") UDynamicMesh* 
+	ComputePolygroupsFromAngleThreshold( 
+		UDynamicMesh* TargetMesh, 
+		FGeometryScriptGroupLayer GroupLayer,
+		float CreaseAngle = 15,
+		int MinGroupSize = 2,
+		UGeometryScriptDebug* Debug = nullptr );
+
+};
\ No newline at end of file
diff --git a/Engine/Plugins/Experimental/GeometryScripting/Source/GeometryScriptingCore/Public/GeometryScript/MeshPrimitiveFunctions.h b/Engine/Plugins/Experimental/GeometryScripting/Source/GeometryScriptingCore/Public/GeometryScript/MeshPrimitiveFunctions.h
--- a/Engine/Plugins/Experimental/GeometryScripting/Source/GeometryScriptingCore/Public/GeometryScript/MeshPrimitiveFunctions.h
+++ b/Engine/Plugins/Experimental/GeometryScripting/Source/GeometryScriptingCore/Public/GeometryScript/MeshPrimitiveFunctions.h
@@ -16,6 +16,12 @@ enum class EGeometryScriptPrimitivePolygroupMode : uint8
 	PerQuad = 2
 };
 
+UENUM(BlueprintType)
+enum class EGeometryScriptPrimitiveOriginMode : uint8
+{
+	Center = 0,
+	Base = 1
+};
 
 
 USTRUCT(BlueprintType)
@@ -31,6 +37,31 @@ public:
 };
 
 
+USTRUCT(BlueprintType)
+struct GEOMETRYSCRIPTINGCORE_API FGeometryScriptRevolveOptions
+{
+	GENERATED_BODY()
+public:
+	UPROPERTY(BlueprintReadWrite, Category = Options)
+	float RevolveDegrees = 360.0f;
+
+	UPROPERTY(BlueprintReadWrite, Category = Options)
+	float DegreeOffset = 0.0f;
+
+	UPROPERTY(BlueprintReadWrite, Category = Options)
+	bool ReverseDirection = false;
+
+	UPROPERTY(BlueprintReadWrite, Category = Options)
+	bool bHardNormals = false;
+
+	UPROPERTY(BlueprintReadWrite, Category = Options)
+	float HardNormalAngle = 30.0f;
+
+	UPROPERTY(BlueprintReadWrite, Category = Options)
+	bool bProfileAtMidpoint = false;
+};
+
+
 UCLASS(meta = (ScriptName = "GeometryScript_Primitives"))
 class GEOMETRYSCRIPTINGCORE_API UGeometryScriptLibrary_MeshPrimitiveFunctions : public UBlueprintFunctionLibrary
 {
@@ -44,10 +75,13 @@ public:
 		UDynamicMesh* TargetMesh, 
 		FGeometryScriptPrimitiveOptions PrimitiveOptions,
 		FTransform Transform,
-		FBox Box,
+		float DimensionX = 100,
+		float DimensionY = 100,
+		float DimensionZ = 100,
 		int32 StepsX = 0,
 		int32 StepsY = 0,
 		int32 StepsZ = 0,
+		EGeometryScriptPrimitiveOriginMode Origin = EGeometryScriptPrimitiveOriginMode::Base,
 		UGeometryScriptDebug* Debug = nullptr);
 
 
@@ -58,8 +92,9 @@ public:
 		FGeometryScriptPrimitiveOptions PrimitiveOptions,
 		FTransform Transform,
 		float Radius = 50,
-		int32 StepsPhi = 8,
-		int32 StepsTheta = 8,
+		int32 StepsPhi = 10,
+		int32 StepsTheta = 16,
+		EGeometryScriptPrimitiveOriginMode Origin = EGeometryScriptPrimitiveOriginMode::Center,
 		UGeometryScriptDebug* Debug = nullptr);
 
 
@@ -73,6 +108,7 @@ public:
 		int32 StepsX = 6,
 		int32 StepsY = 6,
 		int32 StepsZ = 6,
+		EGeometryScriptPrimitiveOriginMode Origin = EGeometryScriptPrimitiveOriginMode::Center,
 		UGeometryScriptDebug* Debug = nullptr);
 
 
@@ -82,10 +118,11 @@ public:
 		UDynamicMesh* TargetMesh, 
 		FGeometryScriptPrimitiveOptions PrimitiveOptions,
 		FTransform Transform,
-		float Radius = 10,
-		float LineLength = 50,
-		int32 HemisphereSteps = 4,
+		float Radius = 30,
+		float LineLength = 75,
+		int32 HemisphereSteps = 5,
 		int32 CircleSteps = 8,
+		EGeometryScriptPrimitiveOriginMode Origin = EGeometryScriptPrimitiveOriginMode::Base,
 		UGeometryScriptDebug* Debug = nullptr);
 
 
@@ -95,11 +132,12 @@ public:
 		UDynamicMesh* TargetMesh, 
 		FGeometryScriptPrimitiveOptions PrimitiveOptions,
 		FTransform Transform,
-		float Radius = 25,
-		float Height = 50,
-		int32 RadialSteps = 8,
+		float Radius = 50,
+		float Height = 100,
+		int32 RadialSteps = 12,
 		int32 HeightSteps = 0,
 		bool bCapped = true,
+		EGeometryScriptPrimitiveOriginMode Origin = EGeometryScriptPrimitiveOriginMode::Base,
 		UGeometryScriptDebug* Debug = nullptr);
 
 
@@ -109,12 +147,13 @@ public:
 		UDynamicMesh* TargetMesh, 
 		FGeometryScriptPrimitiveOptions PrimitiveOptions,
 		FTransform Transform,
-		float BaseRadius = 25,
+		float BaseRadius = 50,
 		float TopRadius = 5,
-		float Height = 50,
-		int32 RadialSteps = 8,
-		int32 HeightSteps = 0,
+		float Height = 100,
+		int32 RadialSteps = 12,
+		int32 HeightSteps = 4,
 		bool bCapped = true,
+		EGeometryScriptPrimitiveOriginMode Origin = EGeometryScriptPrimitiveOriginMode::Base,
 		UGeometryScriptDebug* Debug = nullptr);
 
 
@@ -124,10 +163,11 @@ public:
 		UDynamicMesh* TargetMesh, 
 		FGeometryScriptPrimitiveOptions PrimitiveOptions,
 		FTransform Transform,
-		float MajorRadius = 100,
-		float MinorRadius = 20,
-		int32 MajorSteps = 8,
-		int32 MinorSteps = 4,
+		float MajorRadius = 50,
+		float MinorRadius = 25,
+		int32 MajorSteps = 16,
+		int32 MinorSteps = 8,
+		EGeometryScriptPrimitiveOriginMode Origin = EGeometryScriptPrimitiveOriginMode::Base,
 		UGeometryScriptDebug* Debug = nullptr);
 
 	/**
@@ -146,6 +186,19 @@ public:
 		int32 Steps = 8,
 		UGeometryScriptDebug* Debug = nullptr);
 
+
+	UFUNCTION(BlueprintCallable, Category = "GeometryScript|Primitives", meta=(ScriptMethod))
+	static UPARAM(DisplayName = "Target Mesh") UDynamicMesh* 
+	AppendRevolvePath( 
+		UDynamicMesh* TargetMesh, 
+		FGeometryScriptPrimitiveOptions PrimitiveOptions,
+		FTransform Transform,
+		const TArray<FVector2D>& PathVertices,
+		FGeometryScriptRevolveOptions RevolveOptions,
+		int32 Steps = 8,
+		bool bCapped = true,
+		UGeometryScriptDebug* Debug = nullptr);
+
 	/**
 	 * Polygon should be oriented counter-clockwise to produce a correctly-oriented shape, otherwise it will be inside-out
 	 * Polygon endpoint is not repeated.
@@ -160,6 +213,7 @@ public:
 		float Height = 100,
 		int32 HeightSteps = 0,
 		bool bCapped = true,
+		EGeometryScriptPrimitiveOriginMode Origin = EGeometryScriptPrimitiveOriginMode::Base,
 		UGeometryScriptDebug* Debug = nullptr);
 
 
@@ -185,7 +239,8 @@ public:
 		UDynamicMesh* TargetMesh, 
 		FGeometryScriptPrimitiveOptions PrimitiveOptions,
 		FTransform Transform,
-		FBox2D Box,
+		float DimensionX = 100,
+		float DimensionY = 100,
 		int32 StepsWidth = 0,
 		int32 StepsHeight = 0,
 		UGeometryScriptDebug* Debug = nullptr);
@@ -196,7 +251,8 @@ public:
 		UDynamicMesh* TargetMesh, 
 		FGeometryScriptPrimitiveOptions PrimitiveOptions,
 		FTransform Transform,
-		FBox2D Box,
+		float DimensionX = 100,
+		float DimensionY = 100,
 		float CornerRadius = 5,
 		int32 StepsWidth = 0,
 		int32 StepsHeight = 0,
diff --git a/Engine/Plugins/Experimental/GeometryScripting/Source/GeometryScriptingCore/Public/GeometryScript/MeshQueryFunctions.h b/Engine/Plugins/Experimental/GeometryScripting/Source/GeometryScriptingCore/Public/GeometryScript/MeshQueryFunctions.h
--- a/Engine/Plugins/Experimental/GeometryScripting/Source/GeometryScriptingCore/Public/GeometryScript/MeshQueryFunctions.h
+++ b/Engine/Plugins/Experimental/GeometryScripting/Source/GeometryScriptingCore/Public/GeometryScript/MeshQueryFunctions.h
@@ -100,4 +100,12 @@ public:
 	UFUNCTION(BlueprintPure, Category = "GeometryScript|MeshQueries", meta=(ScriptMethod))
 	static UPARAM(DisplayName = "Enabled") bool GetHasMaterialIDs( UDynamicMesh* TargetMesh );
 
+
+	UFUNCTION(BlueprintPure, Category = "GeometryScript|MeshQueries", meta=(ScriptMethod))
+	static UPARAM(DisplayName = "Enabled") bool GetHasPolygroups( UDynamicMesh* TargetMesh );
+
+	UFUNCTION(BlueprintPure, Category = "GeometryScript|MeshQueries", meta=(ScriptMethod))
+	static UPARAM(DisplayName = "Num Layers") int GetNumExtendedPolygroupLayers( UDynamicMesh* TargetMesh );
+
+
 };
\ No newline at end of file
diff --git a/Engine/Plugins/Experimental/GeometryScripting/Source/GeometryScriptingCore/Public/GeometryScript/MeshSubdivideFunctions.h b/Engine/Plugins/Experimental/GeometryScripting/Source/GeometryScriptingCore/Public/GeometryScript/MeshSubdivideFunctions.h
new file mode 100644
--- /dev/null
+++ b/Engine/Plugins/Experimental/GeometryScripting/Source/GeometryScriptingCore/Public/GeometryScript/MeshSubdivideFunctions.h
@@ -0,0 +1,39 @@
+// Copyright Epic Games, Inc. All Rights Reserved.
+
+#pragma once
+
+#include "CoreMinimal.h"
+#include "Kismet/BlueprintFunctionLibrary.h"
+#include "GeometryScript/GeometryScriptTypes.h"
+#include "MeshSubdivideFunctions.generated.h"
+
+class UDynamicMesh;
+
+
+
+USTRUCT(BlueprintType, meta = (DisplayName = "PN Tessellate Options"))
+struct GEOMETRYSCRIPTINGCORE_API FGeometryScriptPNTessellateOptions
+{
+	GENERATED_BODY()
+public:
+	UPROPERTY(BlueprintReadWrite, Category = Options)
+	bool bRecomputeNormals = true;
+};
+
+
+
+UCLASS(meta = (ScriptName = "GeometryScript_MeshSubdivide"))
+class GEOMETRYSCRIPTINGCORE_API UGeometryScriptLibrary_MeshSubdivideFunctions : public UBlueprintFunctionLibrary
+{
+	GENERATED_BODY()
+public:
+
+	UFUNCTION(BlueprintCallable, Category = "GeometryScript|Subdivide", meta=(ScriptMethod))
+	static UPARAM(DisplayName = "Target Mesh") UDynamicMesh* 
+	ApplyRecursivePNTessellation(
+		UDynamicMesh* TargetMesh,
+		FGeometryScriptPNTessellateOptions Options,
+		int NumIterations = 3,
+		UGeometryScriptDebug* Debug = nullptr );
+
+};
\ No newline at end of file
diff --git a/Engine/Plugins/Experimental/GeometryScripting/Source/GeometryScriptingCore/Public/GeometryScript/MeshUVFunctions.h b/Engine/Plugins/Experimental/GeometryScripting/Source/GeometryScriptingCore/Public/GeometryScript/MeshUVFunctions.h
--- a/Engine/Plugins/Experimental/GeometryScripting/Source/GeometryScriptingCore/Public/GeometryScript/MeshUVFunctions.h
+++ b/Engine/Plugins/Experimental/GeometryScripting/Source/GeometryScriptingCore/Public/GeometryScript/MeshUVFunctions.h
@@ -22,6 +22,109 @@ public:
 	bool bOptimizeIslandRotation = true;
 };
 
+UENUM(BlueprintType)
+enum class EGeometryScriptUVFlattenMethod : uint8
+{
+	ExpMap = 0,
+	Conformal = 1
+};
+
+UENUM(BlueprintType)
+enum class EGeometryScriptUVIslandSource : uint8
+{
+	PolyGroups = 0,
+	UVIslands = 1
+};
+
+
+USTRUCT(BlueprintType)
+struct GEOMETRYSCRIPTINGCORE_API FGeometryScriptExpMapUVOptions
+{
+	GENERATED_BODY()
+
+	UPROPERTY(BlueprintReadWrite, Category = Options)
+	int NormalSmoothingRounds = 0;
+
+	UPROPERTY(BlueprintReadWrite, Category = Options)
+	float NormalSmoothingAlpha = 0.25f;
+};
+
+
+USTRUCT(BlueprintType)
+struct GEOMETRYSCRIPTINGCORE_API FGeometryScriptRecomputeUVsOptions
+{
+	GENERATED_BODY()
+public:
+	UPROPERTY(BlueprintReadWrite, Category = Options)
+	EGeometryScriptUVFlattenMethod Method = EGeometryScriptUVFlattenMethod::Conformal;
+
+	UPROPERTY(BlueprintReadWrite, Category = Options)
+	EGeometryScriptUVIslandSource IslandSource = EGeometryScriptUVIslandSource::UVIslands;
+
+	UPROPERTY(BlueprintReadWrite, Category = Options)
+	FGeometryScriptExpMapUVOptions ExpMapOptions;
+
+	UPROPERTY(BlueprintReadWrite, Category = Options)
+	FGeometryScriptGroupLayer GroupLayer;
+
+	UPROPERTY(BlueprintReadWrite, Category = Options)
+	bool bAutoAlignIslandsWithAxes = true;
+};
+
+
+
+
+
+USTRUCT(BlueprintType)
+struct GEOMETRYSCRIPTINGCORE_API FGeometryScriptPatchBuilderOptions
+{
+	GENERATED_BODY()
+public:
+	UPROPERTY(BlueprintReadWrite, Category = Options)
+	int InitialPatchCount = 100;
+
+	UPROPERTY(BlueprintReadWrite, Category = Options)
+	int MinPatchSize = 2;
+
+	UPROPERTY(BlueprintReadWrite, Category = Options)
+	float PatchCurvatureAlignmentWeight = 1.0f;
+
+	UPROPERTY(BlueprintReadWrite, Category = Options)
+	float PatchMergingMetricThresh = 1.5f;
+
+	UPROPERTY(BlueprintReadWrite, Category = Options)
+	float PatchMergingAngleThresh = 45.0f;
+
+	UPROPERTY(BlueprintReadWrite, Category = Options)
+	FGeometryScriptExpMapUVOptions ExpMapOptions;
+
+	UPROPERTY(BlueprintReadWrite, Category = Options)
+	bool bRespectInputGroups = false;
+
+	UPROPERTY(BlueprintReadWrite, Category = Options)
+	FGeometryScriptGroupLayer GroupLayer;
+
+
+	UPROPERTY(BlueprintReadWrite, Category = Options)
+	bool bAutoPack = true;
+
+	UPROPERTY(BlueprintReadWrite, Category = Options)
+	FGeometryScriptRepackUVsOptions PackingOptions;
+};
+
+
+
+
+USTRUCT(BlueprintType)
+struct GEOMETRYSCRIPTINGCORE_API FGeometryScriptXAtlasOptions
+{
+	GENERATED_BODY()
+public:
+	UPROPERTY(BlueprintReadWrite, Category = Options)
+	int MaxIterations = 2;
+};
+
+
 
 UCLASS(meta = (ScriptName = "GeometryScript_UVs"))
 class GEOMETRYSCRIPTINGCORE_API UGeometryScriptLibrary_MeshUVFunctions : public UBlueprintFunctionLibrary
@@ -36,6 +139,14 @@ public:
 		int NumUVSets,
 		UGeometryScriptDebug* Debug = nullptr );
 
+	UFUNCTION(BlueprintCallable, Category = "GeometryScript|UVs", meta=(ScriptMethod))
+	static UPARAM(DisplayName = "Target Mesh") UDynamicMesh* 
+	CopyUVSet( 
+		UDynamicMesh* TargetMesh, 
+		int FromUVSet,
+		int ToUVSet,
+		UGeometryScriptDebug* Debug = nullptr );
+
 
 	UFUNCTION(BlueprintCallable, Category = "GeometryScript|UVs", meta=(ScriptMethod))
 	static UPARAM(DisplayName = "Target Mesh") UDynamicMesh* 
@@ -97,6 +208,25 @@ public:
 		UGeometryScriptDebug* Debug = nullptr );
 
 
+	UFUNCTION(BlueprintCallable, Category = "GeometryScript|UVs", meta=(ScriptMethod))
+	static UPARAM(DisplayName = "Target Mesh") UDynamicMesh* 
+	SetMeshUVsFromCylinderProjection( 
+		UDynamicMesh* TargetMesh, 
+		int UVSetIndex,
+		FTransform CylinderTransform,
+		float SplitAngle = 45.0,
+		UGeometryScriptDebug* Debug = nullptr );
+
+
+
+	UFUNCTION(BlueprintCallable, Category = "GeometryScript|UVs", meta=(ScriptMethod))
+	static UPARAM(DisplayName = "Target Mesh") UDynamicMesh* 
+	RecomputeMeshUVs( 
+		UDynamicMesh* TargetMesh, 
+		int UVSetIndex,
+		FGeometryScriptRecomputeUVsOptions Options,
+		UGeometryScriptDebug* Debug = nullptr );
+
 	UFUNCTION(BlueprintCallable, Category = "GeometryScript|UVs", meta=(ScriptMethod))
 	static UPARAM(DisplayName = "Target Mesh") UDynamicMesh* 
 	RepackMeshUVs( 
@@ -105,4 +235,20 @@ public:
 		FGeometryScriptRepackUVsOptions RepackOptions,
 		UGeometryScriptDebug* Debug = nullptr );
 
+	UFUNCTION(BlueprintCallable, Category = "GeometryScript|UVs", meta=(ScriptMethod))
+	static UPARAM(DisplayName = "Target Mesh") UDynamicMesh* 
+	AutoGeneratePatchBuilderMeshUVs( 
+		UDynamicMesh* TargetMesh, 
+		int UVSetIndex,
+		FGeometryScriptPatchBuilderOptions Options,
+		UGeometryScriptDebug* Debug = nullptr );
+
+	UFUNCTION(BlueprintCallable, Category = "GeometryScript|UVs", meta=(ScriptMethod))
+	static UPARAM(DisplayName = "Target Mesh") UDynamicMesh* 
+	AutoGenerateXAtlasMeshUVs( 
+		UDynamicMesh* TargetMesh, 
+		int UVSetIndex,
+		FGeometryScriptXAtlasOptions Options,
+		UGeometryScriptDebug* Debug = nullptr );
+
 };
\ No newline at end of file
diff --git a/Engine/Plugins/Experimental/GeometryScripting/Source/GeometryScriptingCore/Public/GeometryScript/MeshVertexColorFunctions.h b/Engine/Plugins/Experimental/GeometryScripting/Source/GeometryScriptingCore/Public/GeometryScript/MeshVertexColorFunctions.h
--- a/Engine/Plugins/Experimental/GeometryScripting/Source/GeometryScriptingCore/Public/GeometryScript/MeshVertexColorFunctions.h
+++ b/Engine/Plugins/Experimental/GeometryScripting/Source/GeometryScriptingCore/Public/GeometryScript/MeshVertexColorFunctions.h
@@ -17,12 +17,22 @@ class GEOMETRYSCRIPTINGCORE_API UGeometryScriptLibrary_MeshVertexColorFunctions
 	GENERATED_BODY()
 public:
 
-	UFUNCTION(BlueprintCallable, Category = "GeometryScript|Repair", meta=(ScriptMethod))
+	UFUNCTION(BlueprintCallable, Category = "GeometryScript|VertexColor", meta=(ScriptMethod))
+	static UPARAM(DisplayName = "Target Mesh") UDynamicMesh*
+	SetMeshConstantVertexColor(
+		UDynamicMesh* TargetMesh,
+		FLinearColor Color,
+		FGeometryScriptColorFlags Flags,
+		bool bClearExisting = false,
+		UGeometryScriptDebug* Debug = nullptr);
+
+	UFUNCTION(BlueprintCallable, Category = "GeometryScript|VertexColor", meta=(ScriptMethod))
 	static UPARAM(DisplayName = "Target Mesh") UDynamicMesh*
 	SetMeshPerVertexColors(
 		UDynamicMesh* TargetMesh,
 		UPARAM(ref) const TArray<FLinearColor>& VertexColors,
 		UGeometryScriptDebug* Debug = nullptr);
 
+
 };
 
diff --git a/Engine/Plugins/Experimental/GeometryScripting/Source/GeometryScriptingEditor/Private/OpenSubdivUtilityFunctions.cpp b/Engine/Plugins/Experimental/GeometryScripting/Source/GeometryScriptingEditor/Private/OpenSubdivUtilityFunctions.cpp
new file mode 100644
--- /dev/null
+++ b/Engine/Plugins/Experimental/GeometryScripting/Source/GeometryScriptingEditor/Private/OpenSubdivUtilityFunctions.cpp
@@ -0,0 +1,138 @@
+// Copyright Epic Games, Inc. All Rights Reserved.
+
+#include "GeometryScript/OpenSubdivUtilityFunctions.h"
+
+#include "DynamicMesh/DynamicMesh3.h"
+#include "Operations/SubdividePoly.h"
+#include "Polygroups/PolygroupSet.h"
+#include "GroupTopology.h"
+#include "UDynamicMesh.h"
+
+#include "ExplicitUseGeometryMathTypes.h"		// using UE::Geometry::(math types)
+using namespace UE::Geometry;
+
+#define LOCTEXT_NAMESPACE "UGeometryScriptLibrary_MeshAssetFunctions"
+
+
+UDynamicMesh* UGeometryScriptLibrary_OpenSubdivFunctions::ApplyPolygroupCatmullClarkSubD(
+	UDynamicMesh* TargetMesh,
+	int32 Subdivisions,
+	FGeometryScriptGroupLayer GroupLayer,
+	UGeometryScriptDebug* Debug) 
+{
+	if (TargetMesh == nullptr)
+	{
+		UE::Geometry::AppendError(Debug, EGeometryScriptErrorType::InvalidInputs, LOCTEXT("ApplyPolygroupCatmullClarkSubD_InvalidInput", "ApplyPolygroupCatmullClarkSubD: TargetMesh is Null"));
+		return TargetMesh;
+	}
+	if (Subdivisions <= 0)
+	{
+		return TargetMesh;
+	}
+	if (Subdivisions > 6)
+	{
+		UE::Geometry::AppendError(Debug, EGeometryScriptErrorType::InvalidInputs, LOCTEXT("ApplyPolygroupCatmullClarkSubD_MaxSubdivisions", "ApplyPolygroupCatmullClarkSubD: Clamping Subdivision to 6 levels to avoid catastrophic memory usage. Use repeated Subdivisions if higher levels are necessary."));
+		Subdivisions = 6;
+	}
+
+	TargetMesh->EditMesh([&](FDynamicMesh3& EditMesh)
+	{
+		FPolygroupLayer InputGroupLayer{ GroupLayer.bDefaultLayer, GroupLayer.ExtendedLayerIndex };
+		if (InputGroupLayer.CheckExists(&EditMesh) == false)
+		{
+			UE::Geometry::AppendError(Debug, EGeometryScriptErrorType::InvalidInputs, LOCTEXT("ApplyPolygroupCatmullClarkSubD_MissingGroups", "ApplyPolygroupCatmullClarkSubD: Target Polygroup Layer does not exist"));
+			return;
+		}
+
+		TUniquePtr<FGroupTopology> Topo;
+		if (GroupLayer.bDefaultLayer)
+		{
+			Topo = MakeUnique<FGroupTopology>(&EditMesh, true);
+		}
+		else
+		{
+			Topo = MakeUnique<FGroupTopology>(&EditMesh, EditMesh.Attributes()->GetPolygroupLayer(GroupLayer.ExtendedLayerIndex), true);
+		}
+
+		FSubdividePoly SubD(*Topo, EditMesh, Subdivisions);
+
+		FSubdividePoly::ETopologyCheckResult CheckResult = SubD.ValidateTopology();
+		if (CheckResult != FSubdividePoly::ETopologyCheckResult::Ok)
+		{
+			UE::Geometry::AppendError(Debug, EGeometryScriptErrorType::InvalidInputs, LOCTEXT("ApplyPolygroupCatmullClarkSubD_NoPolygons", "ApplyPolygroupCatmullClarkSubD: Target Polygroup Layer does not define valid Polygons for Subdivision"));
+			return;
+		}
+
+		SubD.SubdivisionScheme = ESubdivisionScheme::CatmullClark;
+		SubD.NormalComputationMethod = ESubdivisionOutputNormals::Interpolated;
+		SubD.UVComputationMethod = ESubdivisionOutputUVs::Interpolated;
+
+		FDynamicMesh3 SubDMesh;
+		bool bOK = SubD.ComputeTopologySubdivision();
+		bOK = bOK && SubD.ComputeSubdividedMesh(SubDMesh);
+		if (bOK)
+		{
+			EditMesh = MoveTemp(SubDMesh);
+		}
+		else
+		{
+			UE::Geometry::AppendError(Debug, EGeometryScriptErrorType::OperationFailed, LOCTEXT("ApplyPolygroupCatmullClarkSubD_Failed", "ApplyPolygroupCatmullClarkSubD: Subdivision Failed"));
+		}
+	});
+
+	return TargetMesh;
+}
+
+
+
+
+UDynamicMesh* UGeometryScriptLibrary_OpenSubdivFunctions::ApplyTriangleLoopSubD(
+	UDynamicMesh* TargetMesh,
+	int32 Subdivisions,
+	UGeometryScriptDebug* Debug) 
+{
+	if (TargetMesh == nullptr)
+	{
+		UE::Geometry::AppendError(Debug, EGeometryScriptErrorType::InvalidInputs, LOCTEXT("ApplyTriangleLoopSubD_InvalidInput", "ApplyTriangleLoopSubD: TargetMesh is Null"));
+		return TargetMesh;
+	}
+	if (Subdivisions <= 0)
+	{
+		return TargetMesh;
+	}
+	if (Subdivisions > 6)
+	{
+		UE::Geometry::AppendError(Debug, EGeometryScriptErrorType::InvalidInputs, LOCTEXT("ApplyTriangleLoopSubD_MaxSubdivisions", "ApplyTriangleLoopSubD: Clamping Subdivision to 6 levels to avoid catastrophic memory usage. Use repeated Subdivisions if higher levels are necessary."));
+		Subdivisions = 6;
+	}
+
+	TargetMesh->EditMesh([&](FDynamicMesh3& EditMesh)
+	{
+		FGroupTopology Topo;
+		FSubdividePoly SubD(Topo, EditMesh, Subdivisions);
+		SubD.SubdivisionScheme = ESubdivisionScheme::Loop;
+		SubD.NormalComputationMethod = ESubdivisionOutputNormals::Interpolated;
+		SubD.UVComputationMethod = ESubdivisionOutputUVs::Interpolated;
+
+		FDynamicMesh3 SubDMesh;
+		bool bOK = SubD.ComputeTopologySubdivision();
+		bOK = bOK && SubD.ComputeSubdividedMesh(SubDMesh);
+		if (bOK)
+		{
+			EditMesh = MoveTemp(SubDMesh);
+		}
+		else
+		{
+			UE::Geometry::AppendError(Debug, EGeometryScriptErrorType::OperationFailed, LOCTEXT("ApplyTriangleLoopSubD_Failed", "ApplyTriangleLoopSubD: Subdivision Failed"));
+		}
+	});
+
+	return TargetMesh;
+}
+
+
+
+
+
+
+#undef LOCTEXT_NAMESPACE
\ No newline at end of file
diff --git a/Engine/Plugins/Experimental/GeometryScripting/Source/GeometryScriptingEditor/Public/GeometryScript/OpenSubdivUtilityFunctions.h b/Engine/Plugins/Experimental/GeometryScripting/Source/GeometryScriptingEditor/Public/GeometryScript/OpenSubdivUtilityFunctions.h
new file mode 100644
--- /dev/null
+++ b/Engine/Plugins/Experimental/GeometryScripting/Source/GeometryScriptingEditor/Public/GeometryScript/OpenSubdivUtilityFunctions.h
@@ -0,0 +1,33 @@
+// Copyright Epic Games, Inc. All Rights Reserved.
+
+#pragma once
+
+#include "CoreMinimal.h"
+#include "Kismet/BlueprintFunctionLibrary.h"
+#include "GeometryScript/GeometryScriptTypes.h"
+#include "OpenSubdivUtilityFunctions.generated.h"
+
+UCLASS(meta = (ScriptName = "GeometryScript_OpenSubdiv"))
+class GEOMETRYSCRIPTINGEDITOR_API UGeometryScriptLibrary_OpenSubdivFunctions : public UBlueprintFunctionLibrary
+{
+	GENERATED_BODY()
+public:
+
+
+	UFUNCTION(BlueprintCallable, Category = "GeometryScript|OpenSubdiv", meta=(ScriptMethod))
+	static UPARAM(DisplayName = "Target Mesh") UDynamicMesh* 
+	ApplyPolygroupCatmullClarkSubD(
+		UDynamicMesh* FromDynamicMesh, 
+		int32 Subdivisions,
+		FGeometryScriptGroupLayer GroupLayer,
+		UGeometryScriptDebug* Debug = nullptr);
+
+
+	UFUNCTION(BlueprintCallable, Category = "GeometryScript|OpenSubdiv", meta=(ScriptMethod))
+	static UPARAM(DisplayName = "Target Mesh") UDynamicMesh* 
+	ApplyTriangleLoopSubD(
+		UDynamicMesh* FromDynamicMesh, 
+		int32 Subdivisions,
+		UGeometryScriptDebug* Debug = nullptr);
+
+};
\ No newline at end of file
diff --git a/Engine/Plugins/Experimental/MeshModelingToolsetExp/Source/ModelingOperators/Private/CompositionOps/CurveSweepOp.cpp b/Engine/Plugins/Experimental/MeshModelingToolsetExp/Source/ModelingOperators/Private/CompositionOps/CurveSweepOp.cpp
--- a/Engine/Plugins/Experimental/MeshModelingToolsetExp/Source/ModelingOperators/Private/CompositionOps/CurveSweepOp.cpp
+++ b/Engine/Plugins/Experimental/MeshModelingToolsetExp/Source/ModelingOperators/Private/CompositionOps/CurveSweepOp.cpp
@@ -15,8 +15,6 @@ using namespace UE::Geometry;
 
 void FCurveSweepOp::CalculateResult(FProgressCancel* Progress)
 {
-	check(Progress != nullptr);
-
 	if (Progress && Progress->Cancelled())
 	{
 		return;
