GeometryScript: add Shell, Offset, and Extrude modeling operations
#rnx

diff --git a/Engine/Plugins/Experimental/GeometryScripting/Source/GeometryScriptingCore/Private/MeshModelingFunctions.cpp b/Engine/Plugins/Experimental/GeometryScripting/Source/GeometryScriptingCore/Private/MeshModelingFunctions.cpp
new file mode 100644
--- /dev/null
+++ b/Engine/Plugins/Experimental/GeometryScripting/Source/GeometryScriptingCore/Private/MeshModelingFunctions.cpp
@@ -0,0 +1,428 @@
+// Copyright Epic Games, Inc. All Rights Reserved.
+
+#include "GeometryScript/MeshModelingFunctions.h"
+
+#include "DynamicMesh/DynamicMesh3.h"
+#include "DynamicMesh/DynamicMeshAttributeSet.h"
+#include "DynamicMesh/MeshNormals.h"
+#include "MeshWeights.h"
+#include "Async/ParallelFor.h"
+#include "Util/ProgressCancel.h"
+#include "Operations/JoinMeshLoops.h"
+#include "MeshBoundaryLoops.h"
+#include "DynamicMeshEditor.h"
+#include "DynamicMesh/DynamicMeshAABBTree3.h"
+#include "Operations/OffsetMeshRegion.h"
+#include "UDynamicMesh.h"
+
+#include "ExplicitUseGeometryMathTypes.h"		// using UE::Geometry::(math types)
+using namespace UE::Geometry;
+
+#define LOCTEXT_NAMESPACE "UGeometryScriptLibrary_MeshModelingFunctions"
+
+
+
+
+class FMeshOffset
+{
+public:
+
+	FDynamicMesh3* Mesh;
+
+	int32 Steps = 1;
+	double SmoothAlpha = 0.1;
+	bool bReprojectSmooth = false;
+	double BoundaryAlpha = 0.2;		// should not be > 0.9
+	bool bFixedBoundary = false;
+
+	FMeshNormals* PerVertexNormals = nullptr;
+
+	FDynamicMesh3 ResultMesh;
+
+	double OffsetDistance = 1.0;
+	double GetOffsetDistance(int32 VertexID)
+	{
+		return OffsetDistance;
+	}
+
+protected:
+	TArray<FVector3d> PositionBuffer;
+	FMeshNormals* BaseNormals = nullptr;
+
+	TArray<bool> bIsBoundary;
+	TArray<int32> BoundaryVerts;
+
+public:
+
+	FMeshOffset(FDynamicMesh3* MeshIn)
+	{
+		Mesh = MeshIn;
+	}
+
+	virtual ~FMeshOffset() {}
+
+	virtual void Apply(FProgressCancel* Progress)
+	{
+		BaseNormals = PerVertexNormals;
+		FMeshNormals LocalNormals;
+		if (PerVertexNormals == nullptr || PerVertexNormals->GetNormals().Num() != Mesh->VertexCount())
+		{
+			LocalNormals.SetMesh(Mesh);
+			LocalNormals.ComputeVertexNormals();
+			BaseNormals = &LocalNormals;
+		}
+
+		ResultMesh = *Mesh;
+		PositionBuffer.SetNum(ResultMesh.MaxVertexID());
+
+		// probably could be more efficient...
+		bIsBoundary.Init(false, ResultMesh.MaxVertexID());
+		BoundaryVerts.Reset();
+		for (int32 vid : ResultMesh.VertexIndicesItr())
+		{
+			if ((bIsBoundary[vid] = ResultMesh.IsBoundaryVertex(vid)))
+			{
+				BoundaryVerts.Add(vid);
+			}
+		}
+
+		if (Steps > 1 && SmoothAlpha > 0)
+		{
+			Offset_Smoothed(Progress);
+		}
+		else
+		{
+			Offset(Progress);
+		}
+
+		for (int32 vid : ResultMesh.VertexIndicesItr())
+		{
+			ResultMesh.SetVertex(vid, PositionBuffer[vid]);
+		}
+	}
+
+protected:
+	void Offset(FProgressCancel* Progress)
+	{
+		int32 UseSteps = FMath::Max(1, Steps);
+		if (UseSteps == 1)
+		{
+			for (int32 vid : ResultMesh.VertexIndicesItr())
+			{
+				double Dist = GetOffsetDistance(vid);
+
+				FVector3d OffsetPosition = ResultMesh.GetVertex(vid);
+				if (bFixedBoundary == false || bIsBoundary[vid] == false)
+				{
+					OffsetPosition += Dist * BaseNormals->GetNormals()[vid];
+				}
+				PositionBuffer[vid] = OffsetPosition;
+			}
+		}
+		else
+		{
+			FMeshNormals StepNormals(&ResultMesh);
+			for (int32 k = 0; k < UseSteps; ++k)
+			{
+				FMeshNormals& UseNormals = (k == 0) ? *BaseNormals : StepNormals;
+
+				// do offset step
+				int32 NumVertices = ResultMesh.MaxVertexID();
+				for (int32 vid : ResultMesh.VertexIndicesItr())
+				{
+					double Dist = GetOffsetDistance(vid);
+					double OffsetPerStep = Dist / (double)UseSteps;
+
+					FVector3d OffsetPosition = ResultMesh.GetVertex(vid);
+					if (bFixedBoundary == false || bIsBoundary[vid] == false)
+					{
+						OffsetPosition += OffsetPerStep * UseNormals.GetNormals()[vid];
+					}
+					PositionBuffer[vid] = OffsetPosition;
+				}
+
+				// bake positions and re-calc normals
+				for (int32 vid : ResultMesh.VertexIndicesItr())
+				{
+					ResultMesh.SetVertex(vid, PositionBuffer[vid]);
+				}
+				StepNormals.ComputeVertexNormals();
+			}
+
+			for (int32 vid : ResultMesh.VertexIndicesItr())
+			{
+				PositionBuffer[vid] = ResultMesh.GetVertex(vid);
+			}
+		}
+	}
+
+
+	void Offset_Smoothed(FProgressCancel* Progress)
+	{
+		int32 UseSteps = FMath::Max(1, Steps);
+		TArray<FVector3d> SmoothedBuffer = PositionBuffer;
+
+		bool bHasSmoothing = (SmoothAlpha > 0);
+		FDynamicMesh3 ProjectMesh;
+		FDynamicMeshAABBTree3 Spatial;
+		if (bReprojectSmooth && bHasSmoothing)
+		{
+			ProjectMesh.Copy(ResultMesh, false, false, false, false);
+			Spatial.SetMesh(&ProjectMesh, true);
+		}
+
+		FMeshNormals StepNormals(&ResultMesh);
+		for (int32 k = 0; k < UseSteps; ++k)
+		{
+			FMeshNormals& UseNormals = (k == 0) ? *BaseNormals : StepNormals;
+
+			// do offset step
+			int32 NumVertices = ResultMesh.MaxVertexID();
+			for (int32 vid : ResultMesh.VertexIndicesItr())
+			{
+				double Dist = GetOffsetDistance(vid);
+				double OffsetPerStep = Dist / (double)UseSteps;
+
+				FVector3d OffsetPosition = ResultMesh.GetVertex(vid);
+				if (bFixedBoundary == false || bIsBoundary[vid] == false)
+				{
+					OffsetPosition += OffsetPerStep * UseNormals.GetNormals()[vid];
+				}
+				PositionBuffer[vid] = OffsetPosition;
+			}
+
+			if (bReprojectSmooth && bHasSmoothing)
+			{
+				for (int32 vid : ResultMesh.VertexIndicesItr())
+				{
+					ProjectMesh.SetVertex(vid, PositionBuffer[vid]);
+				}
+				Spatial.Build();
+			}
+
+			// do smooth step
+			if (bHasSmoothing)
+			{
+				ParallelFor(NumVertices, [&](int32 vid)
+				{
+					if (ResultMesh.IsReferencedVertex(vid) == false || bIsBoundary[vid])
+					{
+						SmoothedBuffer[vid] = PositionBuffer[vid];
+						return;
+					}
+					FVector3d Centroid = FMeshWeights::UniformCentroid(ResultMesh, vid, [&](int32 nbrvid) { return PositionBuffer[nbrvid]; });
+					SmoothedBuffer[vid] = UE::Geometry::Lerp(PositionBuffer[vid], Centroid, SmoothAlpha);
+					if (bReprojectSmooth)
+					{
+						SmoothedBuffer[vid] = Spatial.FindNearestPoint(SmoothedBuffer[vid]);
+					}
+				} /*, EParallelForFlags::ForceSingleThread*/);
+
+				if (bFixedBoundary == false)
+				{
+					ParallelFor(BoundaryVerts.Num(), [&](int32 idx)
+					{
+						int32 vid = BoundaryVerts[idx];
+						FVector3d Centroid = FMeshWeights::FilteredUniformCentroid(ResultMesh, vid,
+							[&](int32 nbrvid) { return PositionBuffer[nbrvid]; },
+							[&](int32 nbrvid) { return bIsBoundary[nbrvid]; });
+						SmoothedBuffer[vid] = UE::Geometry::Lerp(PositionBuffer[vid], Centroid, BoundaryAlpha);
+					});
+				}
+
+				for (int32 vid : ResultMesh.VertexIndicesItr())
+				{
+					PositionBuffer[vid] = SmoothedBuffer[vid];
+				}
+			}
+
+			// bake positions and re-calc normals
+			for (int32 vid : ResultMesh.VertexIndicesItr())
+			{
+				ResultMesh.SetVertex(vid, PositionBuffer[vid]);
+			}
+			StepNormals.ComputeVertexNormals();
+		}
+
+		for (int32 vid : ResultMesh.VertexIndicesItr())
+		{
+			PositionBuffer[vid] = ResultMesh.GetVertex(vid);
+		}
+	}
+
+};
+
+
+
+
+UDynamicMesh* UGeometryScriptLibrary_MeshModelingFunctions::ApplyMeshOffset(
+	UDynamicMesh* TargetMesh,
+	FGeometryScriptMeshOffsetOptions Options,
+	UGeometryScriptDebug* Debug)
+{
+	if (TargetMesh == nullptr)
+	{
+		UE::Geometry::AppendError(Debug, EGeometryScriptErrorType::InvalidInputs, LOCTEXT("ApplyMeshOffset_InvalidInput", "ApplyMeshOffset: TargetMesh is Null"));
+		return TargetMesh;
+	}
+
+
+	TargetMesh->EditMesh([&](FDynamicMesh3& EditMesh) 
+	{
+		FMeshOffset Offset(&EditMesh);
+		Offset.bFixedBoundary = Options.bFixedBoundary;
+		Offset.OffsetDistance = Options.OffsetDistance;
+		Offset.Steps = FMath::Clamp(Options.SolveSteps, 1, 1000);
+		Offset.SmoothAlpha = FMath::Clamp((double)Options.SmoothAlpha, 0, 1.0);
+		Offset.bReprojectSmooth = Options.bReprojectDuringSmoothing;
+		Offset.BoundaryAlpha = FMath::Clamp((double)Options.BoundaryAlpha, 0, 0.9);
+
+		Offset.Apply(nullptr);
+		EditMesh = MoveTemp(Offset.ResultMesh);
+
+		if (EditMesh.HasAttributes() && EditMesh.Attributes()->PrimaryNormals() != nullptr)
+		{
+			FMeshNormals Normals(&EditMesh);
+			FDynamicMeshNormalOverlay* NormalOverlay = EditMesh.Attributes()->PrimaryNormals();
+			Normals.RecomputeOverlayNormals(NormalOverlay);
+			Normals.CopyToOverlay(NormalOverlay);
+		}
+
+	}, EDynamicMeshChangeType::GeneralEdit, EDynamicMeshAttributeChangeFlags::Unknown, false);
+
+	return TargetMesh;
+}
+
+
+
+
+UDynamicMesh* UGeometryScriptLibrary_MeshModelingFunctions::ApplyMeshShell(
+	UDynamicMesh* TargetMesh,
+	FGeometryScriptMeshOffsetOptions Options,
+	UGeometryScriptDebug* Debug)
+{
+	if (TargetMesh == nullptr)
+	{
+		UE::Geometry::AppendError(Debug, EGeometryScriptErrorType::InvalidInputs, LOCTEXT("ApplyMeshShell_InvalidInput", "ApplyMeshShell: TargetMesh is Null"));
+		return TargetMesh;
+	}
+
+
+	TargetMesh->EditMesh([&](FDynamicMesh3& EditMesh) 
+	{
+		UE::Geometry::FMeshBoundaryLoops BoundaryLoops(&EditMesh, true);
+		
+		double OffsetSign = (Options.OffsetDistance >= 0) ? 1 : -1;
+
+		FDynamicMesh3 InnerMesh;
+		InnerMesh.Copy(EditMesh);
+		bool bIsPositiveOffset = (OffsetSign >= 0);
+		if (bIsPositiveOffset)
+		{
+			InnerMesh.ReverseOrientation();
+		}
+
+		FMeshOffset Offset(&EditMesh);
+		Offset.bFixedBoundary = Options.bFixedBoundary;
+		Offset.OffsetDistance = Options.OffsetDistance;
+		Offset.Steps = FMath::Clamp(Options.SolveSteps, 1, 1000);
+		Offset.SmoothAlpha = FMath::Clamp((double)Options.SmoothAlpha, 0, 1.0);
+		Offset.bReprojectSmooth = Options.bReprojectDuringSmoothing;
+		Offset.BoundaryAlpha = FMath::Clamp((double)Options.BoundaryAlpha, 0, 0.9);
+
+		Offset.Apply(nullptr);
+		EditMesh = MoveTemp(Offset.ResultMesh);
+
+		if (!bIsPositiveOffset)
+		{
+			EditMesh.ReverseOrientation();
+		}
+
+		UE::Geometry::FDynamicMeshEditor Editor(&EditMesh);
+		UE::Geometry::FMeshIndexMappings MeshMap;
+		Editor.AppendMesh(&InnerMesh, MeshMap);
+
+		bool bWeldStitch = Options.bFixedBoundary;
+
+		// join the boundary loops by weld or stitch
+		for (UE::Geometry::FEdgeLoop& BaseLoop : BoundaryLoops.Loops)
+		{
+			int32 LoopCount = BaseLoop.GetVertexCount();
+			TArray<int32> OffsetLoop;
+			OffsetLoop.SetNum(LoopCount);
+			for (int k = 0; k < LoopCount; ++k)
+			{
+				OffsetLoop[k] = MeshMap.GetNewVertex(BaseLoop.Vertices[k]);
+			}
+
+			TArray<int32>& LoopA = (bIsPositiveOffset) ? BaseLoop.Vertices : OffsetLoop;
+			TArray<int32>& LoopB = (bIsPositiveOffset) ? OffsetLoop : BaseLoop.Vertices;
+
+			if (bWeldStitch)
+			{
+				bool bStitchSuccess = Editor.WeldVertexLoops(LoopA, LoopB);
+			}
+			else
+			{
+				FJoinMeshLoops Join(&EditMesh, LoopA, LoopB);
+				Join.Apply();
+			}
+		}
+
+		if (EditMesh.HasAttributes() && EditMesh.Attributes()->PrimaryNormals() != nullptr)
+		{
+			FMeshNormals Normals(&EditMesh);
+			FDynamicMeshNormalOverlay* NormalOverlay = EditMesh.Attributes()->PrimaryNormals();
+			Normals.RecomputeOverlayNormals(NormalOverlay);
+			Normals.CopyToOverlay(NormalOverlay);
+		}
+
+	}, EDynamicMeshChangeType::GeneralEdit, EDynamicMeshAttributeChangeFlags::Unknown, false);
+
+	return TargetMesh;
+}
+
+
+
+
+
+
+UDynamicMesh* UGeometryScriptLibrary_MeshModelingFunctions::ApplyMeshExtrude(
+	UDynamicMesh* TargetMesh,
+	FGeometryScriptMeshExtrudeOptions Options,
+	UGeometryScriptDebug* Debug)
+{
+	if (TargetMesh == nullptr)
+	{
+		UE::Geometry::AppendError(Debug, EGeometryScriptErrorType::InvalidInputs, LOCTEXT("ApplyMeshExtrude_InvalidInput", "ApplyMeshExtrude: TargetMesh is Null"));
+		return TargetMesh;
+	}
+	FVector3d ExtrudeVec = (double)Options.ExtrudeDistance * (FVector3d)Options.ExtrudeDirection;
+
+	TargetMesh->EditMesh([&](FDynamicMesh3& EditMesh) 
+	{
+		FOffsetMeshRegion Extruder(&EditMesh);
+		for (int32 tid : EditMesh.TriangleIndicesItr())
+		{
+			Extruder.Triangles.Add(tid);
+		}
+		Extruder.OffsetPositionFunc = [ExtrudeVec](const FVector3d& Position, const FVector3d& VertexVector, int VertexID)
+		{
+			return Position + ExtrudeVec;
+		};
+		Extruder.bIsPositiveOffset = (Options.ExtrudeDistance > 0);
+		
+		Extruder.UVScaleFactor = Options.UVScale;
+		Extruder.bOffsetFullComponentsAsSolids = Options.bSolidsToShells;
+
+		Extruder.Apply();
+
+	}, EDynamicMeshChangeType::GeneralEdit, EDynamicMeshAttributeChangeFlags::Unknown, false);
+
+	return TargetMesh;
+}
+
+
+
+
+#undef LOCTEXT_NAMESPACE
diff --git a/Engine/Plugins/Experimental/GeometryScripting/Source/GeometryScriptingCore/Private/MeshNormalsFunctions.cpp b/Engine/Plugins/Experimental/GeometryScripting/Source/GeometryScriptingCore/Private/MeshNormalsFunctions.cpp
--- a/Engine/Plugins/Experimental/GeometryScripting/Source/GeometryScriptingCore/Private/MeshNormalsFunctions.cpp
+++ b/Engine/Plugins/Experimental/GeometryScripting/Source/GeometryScriptingCore/Private/MeshNormalsFunctions.cpp
@@ -167,7 +167,17 @@ UDynamicMesh* UGeometryScriptLibrary_MeshNormalsFunctions::ComputeSplitNormals(
 		FaceNormals.ComputeTriangleNormals();
 		const TArray<FVector3d>& Normals = FaceNormals.GetNormals();
 
-		FPolygroupSet FaceGroups(&EditMesh, SplitOptions.FaceGroupLayerName);
+		TUniquePtr<FPolygroupSet> SplitByGroups;
+		if (SplitOptions.bSplitByFaceGroup)
+		{
+			FPolygroupLayer InputGroupLayer{ SplitOptions.GroupLayer.bDefaultLayer, SplitOptions.GroupLayer.ExtendedLayerIndex };
+			if (InputGroupLayer.CheckExists(&EditMesh) == false)
+			{
+				UE::Geometry::AppendError(Debug, EGeometryScriptErrorType::InvalidInputs, LOCTEXT("ComputeSplitNormals_MissingGroups", "ComputeSplitNormals: Polygroup Layer does not exist"));
+				return;
+			}
+			SplitByGroups = MakeUnique<FPolygroupSet>(&EditMesh, InputGroupLayer);
+		}
 
 		// create new normal overlay topology based on splitting criteria
 		EditMesh.Attributes()->PrimaryNormals()->CreateFromPredicate([&](int VID, int TA, int TB)
@@ -176,7 +186,7 @@ UDynamicMesh* UGeometryScriptLibrary_MeshNormalsFunctions::ComputeSplitNormals(
 			{
 				return false;
 			}
-			if (SplitOptions.bSplitByFaceGroup && FaceGroups.GetTriangleGroup(TA) != FaceGroups.GetTriangleGroup(TB))
+			if (SplitOptions.bSplitByFaceGroup && SplitByGroups->GetTriangleGroup(TA) != SplitByGroups->GetTriangleGroup(TB))
 			{
 				return false;
 			}
diff --git a/Engine/Plugins/Experimental/GeometryScripting/Source/GeometryScriptingCore/Private/MeshSimplifyFunctions.cpp b/Engine/Plugins/Experimental/GeometryScripting/Source/GeometryScriptingCore/Private/MeshSimplifyFunctions.cpp
--- a/Engine/Plugins/Experimental/GeometryScripting/Source/GeometryScriptingCore/Private/MeshSimplifyFunctions.cpp
+++ b/Engine/Plugins/Experimental/GeometryScripting/Source/GeometryScriptingCore/Private/MeshSimplifyFunctions.cpp
@@ -8,6 +8,10 @@
 #include "MeshSimplification.h"
 #include "MeshConstraintsUtil.h"
 #include "ProjectionTargets.h"
+#include "Polygroups/PolygroupSet.h"
+#include "GroupTopology.h"
+#include "Operations/PolygroupRemesh.h"
+#include "ConstrainedDelaunay2.h"
 
 #include "UDynamicMesh.h"
 
@@ -45,6 +49,49 @@ UDynamicMesh* UGeometryScriptLibrary_MeshSimplifyFunctions::ApplySimplifyToPlana
 
 
 
+UDynamicMesh* UGeometryScriptLibrary_MeshSimplifyFunctions::ApplySimplifyToPolygroupTopology(
+	UDynamicMesh* TargetMesh,
+	FGeometryScriptPolygroupSimplifyOptions Options,
+	FGeometryScriptGroupLayer GroupLayer,
+	UGeometryScriptDebug* Debug)
+{
+	if (TargetMesh == nullptr)
+	{
+		UE::Geometry::AppendError(Debug, EGeometryScriptErrorType::InvalidInputs, LOCTEXT("ApplySimplifyToPolygroupTopology_InvalidInput", "ApplySimplifyToPolygroupTopology: TargetMesh is Null"));
+		return TargetMesh;
+	}
+
+	TargetMesh->EditMesh([&](FDynamicMesh3& EditMesh) 
+	{
+		FPolygroupLayer InputGroupLayer{ GroupLayer.bDefaultLayer, GroupLayer.ExtendedLayerIndex };
+		if (InputGroupLayer.CheckExists(&EditMesh) == false)
+		{
+			UE::Geometry::AppendError(Debug, EGeometryScriptErrorType::InvalidInputs, LOCTEXT("ApplySimplifyToPolygroupTopology_MissingGroups", "ApplySimplifyToPolygroupTopology: Target Polygroup Layer does not exist"));
+			return;
+		}
+
+		TUniquePtr<FGroupTopology> Topo;
+		if (GroupLayer.bDefaultLayer)
+		{
+			Topo = MakeUnique<FGroupTopology>(&EditMesh, true);
+		}
+		else
+		{
+			Topo = MakeUnique<FGroupTopology>(&EditMesh, EditMesh.Attributes()->GetPolygroupLayer(GroupLayer.ExtendedLayerIndex), true);
+		}
+
+		FPolygroupRemesh Simplifier(&EditMesh, Topo.Get(), ConstrainedDelaunayTriangulate<double>);
+		Simplifier.SimplificationAngleTolerance = Options.AngleThreshold;
+		Simplifier.Compute();
+
+	}, EDynamicMeshChangeType::GeneralEdit, EDynamicMeshAttributeChangeFlags::Unknown, false);
+
+	return TargetMesh;
+}
+
+
+
+
 template<typename SimplificationType>
 void DoSimplifyMesh(
 	FDynamicMesh3& EditMesh, 
diff --git a/Engine/Plugins/Experimental/GeometryScripting/Source/GeometryScriptingCore/Public/GeometryScript/MeshModelingFunctions.h b/Engine/Plugins/Experimental/GeometryScripting/Source/GeometryScriptingCore/Public/GeometryScript/MeshModelingFunctions.h
new file mode 100644
--- /dev/null
+++ b/Engine/Plugins/Experimental/GeometryScripting/Source/GeometryScriptingCore/Public/GeometryScript/MeshModelingFunctions.h
@@ -0,0 +1,88 @@
+// Copyright Epic Games, Inc. All Rights Reserved.
+
+#pragma once
+
+#include "CoreMinimal.h"
+#include "Kismet/BlueprintFunctionLibrary.h"
+#include "GeometryScript/GeometryScriptTypes.h"
+#include "MeshModelingFunctions.generated.h"
+
+class UDynamicMesh;
+
+
+USTRUCT(BlueprintType)
+struct GEOMETRYSCRIPTINGCORE_API FGeometryScriptMeshOffsetOptions
+{
+	GENERATED_BODY()
+public:
+	UPROPERTY(BlueprintReadWrite, Category = Options)
+	float OffsetDistance = 1.0;
+
+	UPROPERTY(BlueprintReadWrite, Category = Options)
+	bool bFixedBoundary = false;
+
+	UPROPERTY(BlueprintReadWrite, Category = Options)
+	int SolveSteps = 5;
+
+	UPROPERTY(BlueprintReadWrite, Category = Options)
+	float SmoothAlpha = 0.1f;
+
+	UPROPERTY(BlueprintReadWrite, Category = Options)
+	bool bReprojectDuringSmoothing = false;
+
+	// should not be > 0.9
+	UPROPERTY(BlueprintReadWrite, Category = Options)
+	float BoundaryAlpha = 0.2f;
+};
+
+
+
+USTRUCT(BlueprintType)
+struct GEOMETRYSCRIPTINGCORE_API FGeometryScriptMeshExtrudeOptions
+{
+	GENERATED_BODY()
+public:
+	UPROPERTY(BlueprintReadWrite, Category = Options)
+	float ExtrudeDistance = 1.0;
+
+	UPROPERTY(BlueprintReadWrite, Category = Options)
+	FVector ExtrudeDirection = FVector(0,0,1);
+
+	UPROPERTY(BlueprintReadWrite, Category = Options)
+	float UVScale = 1.0f;
+
+	UPROPERTY(BlueprintReadWrite, Category = Options)
+	bool bSolidsToShells = true;
+};
+
+
+
+
+UCLASS(meta = (ScriptName = "GeometryScript_MeshModeling"))
+class GEOMETRYSCRIPTINGCORE_API UGeometryScriptLibrary_MeshModelingFunctions : public UBlueprintFunctionLibrary
+{
+	GENERATED_BODY()
+public:
+
+	UFUNCTION(BlueprintCallable, Category = "GeometryScript|Modeling", meta=(ScriptMethod))
+	static UPARAM(DisplayName = "Target Mesh") UDynamicMesh* 
+	ApplyMeshOffset(
+		UDynamicMesh* TargetMesh,
+		FGeometryScriptMeshOffsetOptions Options,
+		UGeometryScriptDebug* Debug = nullptr );
+
+	UFUNCTION(BlueprintCallable, Category = "GeometryScript|Modeling", meta=(ScriptMethod))
+	static UPARAM(DisplayName = "Target Mesh") UDynamicMesh* 
+	ApplyMeshShell(
+		UDynamicMesh* TargetMesh,
+		FGeometryScriptMeshOffsetOptions Options,
+		UGeometryScriptDebug* Debug = nullptr );
+
+	UFUNCTION(BlueprintCallable, Category = "GeometryScript|Modeling", meta=(ScriptMethod))
+	static UPARAM(DisplayName = "Target Mesh") UDynamicMesh* 
+	ApplyMeshExtrude(
+		UDynamicMesh* TargetMesh,
+		FGeometryScriptMeshExtrudeOptions Options,
+		UGeometryScriptDebug* Debug = nullptr );
+
+};
\ No newline at end of file
diff --git a/Engine/Plugins/Experimental/GeometryScripting/Source/GeometryScriptingCore/Public/GeometryScript/MeshNormalsFunctions.h b/Engine/Plugins/Experimental/GeometryScripting/Source/GeometryScriptingCore/Public/GeometryScript/MeshNormalsFunctions.h
--- a/Engine/Plugins/Experimental/GeometryScripting/Source/GeometryScriptingCore/Public/GeometryScript/MeshNormalsFunctions.h
+++ b/Engine/Plugins/Experimental/GeometryScripting/Source/GeometryScriptingCore/Public/GeometryScript/MeshNormalsFunctions.h
@@ -38,7 +38,7 @@ public:
 	bool bSplitByFaceGroup = false;
 
 	UPROPERTY(BlueprintReadWrite, Category = Options)
-	FName FaceGroupLayerName = FName();
+	FGeometryScriptGroupLayer GroupLayer;
 };
 
 
diff --git a/Engine/Plugins/Experimental/GeometryScripting/Source/GeometryScriptingCore/Public/GeometryScript/MeshSimplifyFunctions.h b/Engine/Plugins/Experimental/GeometryScripting/Source/GeometryScriptingCore/Public/GeometryScript/MeshSimplifyFunctions.h
--- a/Engine/Plugins/Experimental/GeometryScripting/Source/GeometryScriptingCore/Public/GeometryScript/MeshSimplifyFunctions.h
+++ b/Engine/Plugins/Experimental/GeometryScripting/Source/GeometryScriptingCore/Public/GeometryScript/MeshSimplifyFunctions.h
@@ -19,6 +19,15 @@ public:
 	float AngleThreshold = 0.001;
 };
 
+USTRUCT(BlueprintType)
+struct GEOMETRYSCRIPTINGCORE_API FGeometryScriptPolygroupSimplifyOptions
+{
+	GENERATED_BODY()
+public:
+	/** */
+	UPROPERTY(BlueprintReadWrite, Category = Options)
+	float AngleThreshold = 0.001;
+};
 
 
 UENUM(BlueprintType)
@@ -71,6 +80,14 @@ public:
 		FGeometryScriptPlanarSimplifyOptions Options,
 		UGeometryScriptDebug* Debug = nullptr);
 
+	UFUNCTION(BlueprintCallable, Category = "GeometryScript|Simplification", meta=(ScriptMethod))
+	static UPARAM(DisplayName = "Target Mesh") UDynamicMesh* 
+	ApplySimplifyToPolygroupTopology(  
+		UDynamicMesh* TargetMesh, 
+		FGeometryScriptPolygroupSimplifyOptions Options,
+		FGeometryScriptGroupLayer GroupLayer,
+		UGeometryScriptDebug* Debug = nullptr);
+
 	UFUNCTION(BlueprintCallable, Category = "GeometryScript|Simplification", meta=(ScriptMethod))
 	static UPARAM(DisplayName = "Target Mesh") UDynamicMesh* 
 	ApplySimplifyToTriangleCount(  
