GeometryScript: Add functions for setting simple collision on StaticMesh and DynamicMeshComponent
#rnx

diff --git a/Engine/Plugins/Experimental/GeometryScripting/Source/GeometryScriptingCore/Private/CollisionFunctions.cpp b/Engine/Plugins/Experimental/GeometryScripting/Source/GeometryScriptingCore/Private/CollisionFunctions.cpp
new file mode 100644
--- /dev/null
+++ b/Engine/Plugins/Experimental/GeometryScripting/Source/GeometryScriptingCore/Private/CollisionFunctions.cpp
@@ -0,0 +1,293 @@
+// Copyright Epic Games, Inc. All Rights Reserved.
+
+#include "GeometryScript/CollisionFunctions.h"
+
+#include "DynamicMesh/DynamicMesh3.h"
+#include "DynamicMesh/MeshNormals.h"
+#include "Operations/MeshConvexHull.h"
+#include "Operations/MeshProjectionHull.h"
+#include "UDynamicMesh.h"
+#include "Components/DynamicMeshComponent.h"
+
+#include "Selections/MeshConnectedComponents.h"
+#include "DynamicSubmesh3.h"
+#include "Polygroups/PolygroupUtil.h"
+
+#include "ShapeApproximation/ShapeDetection3.h"
+#include "ShapeApproximation/MeshSimpleShapeApproximation.h"
+
+// physics data
+#include "Engine/StaticMesh.h"
+#include "Components/StaticMeshComponent.h"
+#include "PhysicsEngine/BodySetup.h"
+
+// requires ModelingComponents
+#include "Physics/PhysicsDataCollection.h"
+
+#if WITH_EDITOR
+#include "Editor.h"
+#endif
+
+#include "ExplicitUseGeometryMathTypes.h"		// using UE::Geometry::(math types)
+using namespace UE::Geometry;
+
+#define LOCTEXT_NAMESPACE "UGeometryScriptLibrary_CollisionFunctions"
+
+
+namespace UELocal
+{
+
+void ComputeCollisionFromMesh(
+	const FDynamicMesh3& Mesh, 
+	FKAggregateGeom& GeneratedCollision, 
+	FGeometryScriptCollisionFromMeshOptions& Options)
+{
+	FPhysicsDataCollection NewCollision;
+
+	FMeshConnectedComponents Components(&Mesh);
+	Components.FindConnectedTriangles();
+
+	TArray<FDynamicMesh3> Submeshes;
+	TArray<const FDynamicMesh3*> SubmeshPointers;
+	SubmeshPointers.SetNum(Submeshes.Num());
+
+	if (Components.Num() == 1)
+	{
+		SubmeshPointers.Add(&Mesh);
+	}
+	else
+	{
+		Submeshes.SetNum(Components.Num());
+		ParallelFor(Submeshes.Num(), [&](int32 k)
+		{
+			FDynamicSubmesh3 Submesh(&Mesh, Components[k].Indices, (int32)EMeshComponents::None, false);
+			Submeshes[k] = MoveTemp(Submesh.GetSubmesh());
+			SubmeshPointers[k] = &Submeshes[k];
+		});
+	}
+
+	FMeshSimpleShapeApproximation ShapeGenerator;
+	ShapeGenerator.InitializeSourceMeshes(SubmeshPointers);
+
+	ShapeGenerator.bDetectSpheres = Options.bAutoDetectSpheres;
+	ShapeGenerator.bDetectBoxes = Options.bAutoDetectBoxes;
+	ShapeGenerator.bDetectCapsules = Options.bAutoDetectCapsules;
+
+	ShapeGenerator.MinDimension = Options.MinThickness;
+
+	switch (Options.Method)
+	{
+	case EGeometryScriptCollisionGenerationMethod::AlignedBoxes:
+		ShapeGenerator.Generate_AlignedBoxes(NewCollision.Geometry);
+		break;
+	case EGeometryScriptCollisionGenerationMethod::OrientedBoxes:
+		ShapeGenerator.Generate_OrientedBoxes(NewCollision.Geometry);
+		break;
+	case EGeometryScriptCollisionGenerationMethod::MinimalSpheres:
+		ShapeGenerator.Generate_MinimalSpheres(NewCollision.Geometry);
+		break;
+	case EGeometryScriptCollisionGenerationMethod::Capsules:
+		ShapeGenerator.Generate_Capsules(NewCollision.Geometry);
+		break;
+	case EGeometryScriptCollisionGenerationMethod::ConvexHulls:
+		ShapeGenerator.bSimplifyHulls = Options.bSimplifyHulls;
+		ShapeGenerator.HullTargetFaceCount = Options.ConvexHullTargetFaceCount;
+		ShapeGenerator.Generate_ConvexHulls(NewCollision.Geometry);
+		break;
+	case EGeometryScriptCollisionGenerationMethod::SweptHulls:
+		ShapeGenerator.bSimplifyHulls = Options.bSimplifyHulls;
+		ShapeGenerator.HullSimplifyTolerance = Options.SweptHullSimplifyTolerance;
+		ShapeGenerator.Generate_ProjectedHulls(NewCollision.Geometry, 
+			static_cast<FMeshSimpleShapeApproximation::EProjectedHullAxisMode>(Options.SweptHullAxis));
+		break;
+	case EGeometryScriptCollisionGenerationMethod::MinVolumeShapes:
+		ShapeGenerator.Generate_MinVolume(NewCollision.Geometry);
+		break;
+	}
+
+	if (Options.bRemoveFullyContainedShapes && Components.Num() > 1)
+	{
+		NewCollision.Geometry.RemoveContainedGeometry();
+	}
+
+	if (Options.MaxShapeCount > 0 && Options.MaxShapeCount < Components.Num())
+	{
+		NewCollision.Geometry.FilterByVolume(Options.MaxShapeCount);
+	}
+
+	NewCollision.CopyGeometryToAggregate();
+	GeneratedCollision = NewCollision.AggGeom;
+}
+
+}		// end namespace UELocal
+
+
+UDynamicMesh* UGeometryScriptLibrary_CollisionFunctions::SetStaticMeshCollisionFromMesh(
+	UDynamicMesh* FromDynamicMesh,
+	UStaticMesh* ToStaticMeshAsset,
+	FGeometryScriptCollisionFromMeshOptions Options,
+	UGeometryScriptDebug* Debug)
+{
+	if (FromDynamicMesh == nullptr)
+	{
+		UE::Geometry::AppendError(Debug, EGeometryScriptErrorType::InvalidInputs, LOCTEXT("SetStaticMeshCollisionFromMesh_InvalidInput1", "SetStaticMeshCollisionFromMesh: FromDynamicMesh is Null"));
+		return FromDynamicMesh;
+	}
+	if (ToStaticMeshAsset == nullptr)
+	{
+		UE::Geometry::AppendError(Debug, EGeometryScriptErrorType::InvalidInputs, LOCTEXT("SetStaticMeshCollisionFromMesh_InvalidInput2", "SetStaticMeshCollisionFromMesh: ToStaticMeshAsset is Null"));
+		return FromDynamicMesh;
+	}
+
+	FKAggregateGeom NewCollision;
+	FromDynamicMesh->ProcessMesh([&](const FDynamicMesh3& ReadMesh)
+	{
+		UELocal::ComputeCollisionFromMesh(ReadMesh, NewCollision, Options);
+	});
+
+#if WITH_EDITOR
+	if (Options.bEmitTransaction)
+	{
+		GEditor->BeginTransaction(LOCTEXT("UpdateStaticMesh", "Set Simple Collision"));
+	}
+
+	ToStaticMeshAsset->Modify();
+#endif
+
+	UBodySetup* BodySetup = ToStaticMeshAsset->GetBodySetup();
+	if (BodySetup != nullptr)
+	{
+		// mark the BodySetup for modification. Do we need to modify the UStaticMesh??
+#if WITH_EDITOR
+		if (Options.bEmitTransaction)
+		{
+			BodySetup->Modify();
+		}
+#endif
+
+		// clear existing simple collision. This will call BodySetup->InvalidatePhysicsData()
+		BodySetup->RemoveSimpleCollision();
+
+		// set new collision geometry
+		BodySetup->AggGeom = NewCollision;
+
+		// update collision type
+		//BodySetup->CollisionTraceFlag = (ECollisionTraceFlag)(int32)Settings->SetCollisionType;
+
+		// rebuild physics meshes
+		BodySetup->CreatePhysicsMeshes();
+
+		// rebuild nav collision (? StaticMeshEditor does this)
+		ToStaticMeshAsset->CreateNavCollision(/*bIsUpdate=*/true);
+
+		// update physics state on all components using this StaticMesh
+		for (FThreadSafeObjectIterator Iter(UStaticMeshComponent::StaticClass()); Iter; ++Iter)
+		{
+			UStaticMeshComponent* SMComponent = Cast<UStaticMeshComponent>(*Iter);
+			if (SMComponent->GetStaticMesh() == ToStaticMeshAsset)
+			{
+				if (SMComponent->IsPhysicsStateCreated())
+				{
+					SMComponent->RecreatePhysicsState();
+				}
+			}
+		}
+
+		// do we need to do a post edit change here??
+
+		// mark static mesh as dirty so it gets resaved?
+		ToStaticMeshAsset->MarkPackageDirty();
+
+#if WITH_EDITORONLY_DATA
+		// mark the static mesh as having customized collision so it is not regenerated on reimport
+		ToStaticMeshAsset->bCustomizedCollision = true;
+#endif // WITH_EDITORONLY_DATA
+	}
+
+#if WITH_EDITOR
+	if (Options.bEmitTransaction)
+	{
+		GEditor->EndTransaction();
+	}
+#endif
+
+
+	return FromDynamicMesh;
+}
+
+
+
+
+
+
+
+UDynamicMesh* UGeometryScriptLibrary_CollisionFunctions::SetDynamicMeshCollisionFromMesh(
+	UDynamicMesh* FromDynamicMesh,
+	UDynamicMeshComponent* DynamicMeshComponent,
+	FGeometryScriptCollisionFromMeshOptions Options,
+	UGeometryScriptDebug* Debug)
+{
+	if (FromDynamicMesh == nullptr)
+	{
+		UE::Geometry::AppendError(Debug, EGeometryScriptErrorType::InvalidInputs, LOCTEXT("SetDynamicMeshCollisionFromMesh_InvalidInput1", "SetDynamicMeshCollisionFromMesh: FromDynamicMesh is Null"));
+		return FromDynamicMesh;
+	}
+	if (DynamicMeshComponent == nullptr)
+	{
+		UE::Geometry::AppendError(Debug, EGeometryScriptErrorType::InvalidInputs, LOCTEXT("SetDynamicMeshCollisionFromMesh_InvalidInput2", "SetDynamicMeshCollisionFromMesh: ToDynamicMeshComponent is Null"));
+		return FromDynamicMesh;
+	}
+
+	FKAggregateGeom NewCollision;
+	FromDynamicMesh->ProcessMesh([&](const FDynamicMesh3& ReadMesh)
+	{
+		UELocal::ComputeCollisionFromMesh(ReadMesh, NewCollision, Options);
+	});
+
+#if WITH_EDITOR
+	if (Options.bEmitTransaction)
+	{
+		GEditor->BeginTransaction(LOCTEXT("UpdateDynamicMesh", "Set Simple Collision"));
+	}
+
+	DynamicMeshComponent->Modify();
+#endif
+
+	UBodySetup* BodySetup = DynamicMeshComponent->GetBodySetup();
+	if (BodySetup != nullptr)
+	{
+		// mark the BodySetup for modification. Do we need to modify the UStaticMesh??
+#if WITH_EDITOR
+		if (Options.bEmitTransaction)
+		{
+			BodySetup->Modify();
+		}
+#endif
+
+		// clear existing simple collision. This will call BodySetup->InvalidatePhysicsData()
+		BodySetup->RemoveSimpleCollision();
+
+		// set new collision geometry
+		BodySetup->AggGeom = NewCollision;
+
+		// update collision type
+		//BodySetup->CollisionTraceFlag = (ECollisionTraceFlag)(int32)Settings->SetCollisionType;
+
+		// do we need to do a post edit change here??
+		DynamicMeshComponent->UpdateCollision(false);
+	}
+
+#if WITH_EDITOR
+	if (Options.bEmitTransaction)
+	{
+		GEditor->EndTransaction();
+	}
+#endif
+
+	return FromDynamicMesh;
+}
+
+
+
+
+#undef LOCTEXT_NAMESPACE
\ No newline at end of file
diff --git a/Engine/Plugins/Experimental/GeometryScripting/Source/GeometryScriptingCore/Private/SceneUtilityFunctions.cpp b/Engine/Plugins/Experimental/GeometryScripting/Source/GeometryScriptingCore/Private/SceneUtilityFunctions.cpp
--- a/Engine/Plugins/Experimental/GeometryScripting/Source/GeometryScriptingCore/Private/SceneUtilityFunctions.cpp
+++ b/Engine/Plugins/Experimental/GeometryScripting/Source/GeometryScriptingCore/Private/SceneUtilityFunctions.cpp
@@ -16,6 +16,11 @@ using namespace UE::Geometry;
 
 #define LOCTEXT_NAMESPACE "UGeometryScriptLibrary_SceneUtilityFunctions"
 
+UDynamicMeshPool* UGeometryScriptLibrary_SceneUtilityFunctions::CreateDynamicMeshPool()
+{
+	return NewObject<UDynamicMeshPool>();
+}
+
 
 UDynamicMesh* UGeometryScriptLibrary_SceneUtilityFunctions::CopyMeshFromComponent(
 	USceneComponent* Component,
diff --git a/Engine/Plugins/Experimental/GeometryScripting/Source/GeometryScriptingCore/Public/GeometryScript/CollisionFunctions.h b/Engine/Plugins/Experimental/GeometryScripting/Source/GeometryScriptingCore/Public/GeometryScript/CollisionFunctions.h
new file mode 100644
--- /dev/null
+++ b/Engine/Plugins/Experimental/GeometryScripting/Source/GeometryScriptingCore/Public/GeometryScript/CollisionFunctions.h
@@ -0,0 +1,109 @@
+// Copyright Epic Games, Inc. All Rights Reserved.
+
+#pragma once
+
+#include "CoreMinimal.h"
+#include "Kismet/BlueprintFunctionLibrary.h"
+#include "GeometryScript/GeometryScriptTypes.h"
+#include "CollisionFunctions.generated.h"
+
+class UDynamicMesh;
+class UDynamicMeshComponent;
+class UStaticMesh;
+
+UENUM(BlueprintType)
+enum class EGeometryScriptCollisionGenerationMethod : uint8
+{
+	AlignedBoxes = 0,
+	OrientedBoxes = 1,
+	MinimalSpheres = 2,
+	Capsules = 3,
+	ConvexHulls = 4,
+	SweptHulls = 5,
+	MinVolumeShapes = 6
+};
+
+
+UENUM(BlueprintType)
+enum class EGeometryScriptSweptHullAxis : uint8
+{
+	X = 0,
+	Y = 1,
+	Z = 2,
+	/** Use X/Y/Z axis with smallest axis-aligned-bounding-box dimension */
+	SmallestBoxDimension = 3,
+	/** Compute projected hull for each of X/Y/Z axes and use the one that has the smallest volume  */
+	SmallestVolume = 4
+};
+
+
+USTRUCT(BlueprintType)
+struct GEOMETRYSCRIPTINGCORE_API FGeometryScriptCollisionFromMeshOptions
+{
+	GENERATED_BODY()
+public:
+	UPROPERTY(BlueprintReadWrite, Category = Options)
+	bool bEmitTransaction = true;
+
+	UPROPERTY(BlueprintReadWrite, Category = Options)
+	EGeometryScriptCollisionGenerationMethod Method = EGeometryScriptCollisionGenerationMethod::MinVolumeShapes;
+
+	UPROPERTY(BlueprintReadWrite, Category = Options)
+	bool bAutoDetectSpheres = true;
+
+	UPROPERTY(BlueprintReadWrite, Category = Options)
+	bool bAutoDetectBoxes = true;
+
+	UPROPERTY(BlueprintReadWrite, Category = Options)
+	bool bAutoDetectCapsules = true;
+
+	UPROPERTY(BlueprintReadWrite, Category = Options)
+	float MinThickness = 1.0;
+
+	UPROPERTY(BlueprintReadWrite, Category = Options)
+	bool bSimplifyHulls = true;
+
+	UPROPERTY(BlueprintReadWrite, Category = Options)
+	int ConvexHullTargetFaceCount = 25;
+
+	UPROPERTY(BlueprintReadWrite, Category = Options)
+	float SweptHullSimplifyTolerance = 0.1;
+
+	UPROPERTY(BlueprintReadWrite, Category = Options)
+	EGeometryScriptSweptHullAxis SweptHullAxis = EGeometryScriptSweptHullAxis::Z;
+
+
+	UPROPERTY(BlueprintReadWrite, Category = Options)
+	bool bRemoveFullyContainedShapes = true;
+
+	UPROPERTY(BlueprintReadWrite, Category = Options)
+	int MaxShapeCount = 0;
+};
+
+
+
+
+UCLASS(meta = (ScriptName = "GeometryScript_Collision"))
+class GEOMETRYSCRIPTINGCORE_API UGeometryScriptLibrary_CollisionFunctions : public UBlueprintFunctionLibrary
+{
+	GENERATED_BODY()
+public:
+
+
+	UFUNCTION(BlueprintCallable, Category = "GeometryScript|Collision")
+	static UPARAM(DisplayName = "Dynamic Mesh") UDynamicMesh* 
+	SetStaticMeshCollisionFromMesh(
+		UDynamicMesh* FromDynamicMesh, 
+		UStaticMesh* ToStaticMeshAsset, 
+		FGeometryScriptCollisionFromMeshOptions Options,
+		UGeometryScriptDebug* Debug = nullptr);
+
+	UFUNCTION(BlueprintCallable, Category = "GeometryScript|Collision")
+	static UPARAM(DisplayName = "Dynamic Mesh") UDynamicMesh* 
+	SetDynamicMeshCollisionFromMesh(
+		UDynamicMesh* FromDynamicMesh, 
+		UDynamicMeshComponent* ToDynamicMeshComponent,
+		FGeometryScriptCollisionFromMeshOptions Options,
+		UGeometryScriptDebug* Debug = nullptr);
+
+};
\ No newline at end of file
diff --git a/Engine/Plugins/Experimental/GeometryScripting/Source/GeometryScriptingCore/Public/GeometryScript/SceneUtilityFunctions.h b/Engine/Plugins/Experimental/GeometryScripting/Source/GeometryScriptingCore/Public/GeometryScript/SceneUtilityFunctions.h
--- a/Engine/Plugins/Experimental/GeometryScripting/Source/GeometryScriptingCore/Public/GeometryScript/SceneUtilityFunctions.h
+++ b/Engine/Plugins/Experimental/GeometryScripting/Source/GeometryScriptingCore/Public/GeometryScript/SceneUtilityFunctions.h
@@ -9,7 +9,7 @@
 
 class UStaticMesh;
 class UDynamicMesh;
-
+class UDynamicMeshPool;
 
 USTRUCT(BlueprintType)
 struct GEOMETRYSCRIPTINGCORE_API FGeometryScriptCopyMeshFromComponentOptions
@@ -33,6 +33,12 @@ class GEOMETRYSCRIPTINGCORE_API UGeometryScriptLibrary_SceneUtilityFunctions : p
 	GENERATED_BODY()
 public:
 
+	UFUNCTION(BlueprintCallable, Category = "GeometryScript|Utility")
+	static UPARAM(DisplayName = "Dynamic Mesh Pool") UDynamicMeshPool*
+	CreateDynamicMeshPool();
+
+
+
 	UFUNCTION(BlueprintCallable, Category = "GeometryScript|Scene", meta = (ExpandEnumAsExecs = "Outcome"))
 	static UPARAM(DisplayName = "Dynamic Mesh") UDynamicMesh*
 	CopyMeshFromComponent(
