GeometryScript: fix uninitialized ustruct errors
#rnx

diff --git a/Engine/Plugins/Experimental/GeometryScripting/Source/GeometryScriptingCore/Public/GeometryScript/GeometryScriptTypes.h b/Engine/Plugins/Experimental/GeometryScripting/Source/GeometryScriptingCore/Public/GeometryScript/GeometryScriptTypes.h
--- a/Engine/Plugins/Experimental/GeometryScripting/Source/GeometryScriptingCore/Public/GeometryScript/GeometryScriptTypes.h
+++ b/Engine/Plugins/Experimental/GeometryScripting/Source/GeometryScriptingCore/Public/GeometryScript/GeometryScriptTypes.h
@@ -83,13 +83,13 @@ struct GEOMETRYSCRIPTINGCORE_API FGeometryScriptTriangle
 	GENERATED_BODY()
 public:
 	UPROPERTY(BlueprintReadWrite, Category = Options)
-	FVector Vector0;
+	FVector Vector0 = FVector::ZeroVector;
 
 	UPROPERTY(BlueprintReadWrite, Category = Options)
-	FVector Vector1;
+	FVector Vector1 = FVector::ZeroVector;
 
 	UPROPERTY(BlueprintReadWrite, Category = Options)
-	FVector Vector2;
+	FVector Vector2 = FVector::ZeroVector;
 };
 
 
@@ -105,10 +105,10 @@ public:
 	int TriangleID = -1;
 
 	UPROPERTY(BlueprintReadWrite, Category = Options)
-	FVector Position;
+	FVector Position = FVector::ZeroVector;
 
 	UPROPERTY(BlueprintReadWrite, Category = Options)
-	FVector BaryCoords;
+	FVector BaryCoords = FVector::ZeroVector;
 };
 
 
@@ -118,13 +118,13 @@ struct GEOMETRYSCRIPTINGCORE_API FGeometryScriptUVTriangle
 	GENERATED_BODY()
 public:
 	UPROPERTY(BlueprintReadWrite, Category = Options)
-	FVector2D UV0;
+	FVector2D UV0 = FVector2D::ZeroVector;
 
 	UPROPERTY(BlueprintReadWrite, Category = Options)
-	FVector2D UV1;
+	FVector2D UV1 = FVector2D::ZeroVector;
 
 	UPROPERTY(BlueprintReadWrite, Category = Options)
-	FVector2D UV2;
+	FVector2D UV2 = FVector2D::ZeroVector;
 };
 
 
diff --git a/Engine/Plugins/Experimental/GeometryScripting/Source/GeometryScriptingCore/Public/GeometryScript/MeshSpatialFunctions.h b/Engine/Plugins/Experimental/GeometryScripting/Source/GeometryScriptingCore/Public/GeometryScript/MeshSpatialFunctions.h
--- a/Engine/Plugins/Experimental/GeometryScripting/Source/GeometryScriptingCore/Public/GeometryScript/MeshSpatialFunctions.h
+++ b/Engine/Plugins/Experimental/GeometryScripting/Source/GeometryScriptingCore/Public/GeometryScript/MeshSpatialFunctions.h
@@ -44,10 +44,10 @@ public:
 	int HitTriangleID = -1;
 
 	UPROPERTY(BlueprintReadWrite, Category = Options)
-	FVector HitPosition;
+	FVector HitPosition = FVector::ZeroVector;
 
 	UPROPERTY(BlueprintReadWrite, Category = Options)
-	FVector HitBaryCoords;
+	FVector HitBaryCoords = FVector::ZeroVector;
 };
 
 
diff --git a/Engine/Plugins/Experimental/GeometryScripting/Source/GeometryScriptingCore/Public/GeometryScript/MeshVoxelFunctions.h b/Engine/Plugins/Experimental/GeometryScripting/Source/GeometryScriptingCore/Public/GeometryScript/MeshVoxelFunctions.h
--- a/Engine/Plugins/Experimental/GeometryScripting/Source/GeometryScriptingCore/Public/GeometryScript/MeshVoxelFunctions.h
+++ b/Engine/Plugins/Experimental/GeometryScripting/Source/GeometryScriptingCore/Public/GeometryScript/MeshVoxelFunctions.h
@@ -51,7 +51,7 @@ public:
 	float WindingThreshold = 0.5;
 
 	UPROPERTY(BlueprintReadWrite, Category = Options)
-	bool bSolidAtBoundaries;
+	bool bSolidAtBoundaries = true;
 
 	UPROPERTY(BlueprintReadWrite, Category = Options)
 	float ExtendBounds = 1.0;
