GeometryScript: fix up python name collision
#rnx

diff --git a/Engine/Plugins/Experimental/GeometryScripting/Source/GeometryScriptingCore/Private/MeshQueryFunctions.cpp b/Engine/Plugins/Experimental/GeometryScripting/Source/GeometryScriptingCore/Private/MeshQueryFunctions.cpp
--- a/Engine/Plugins/Experimental/GeometryScripting/Source/GeometryScriptingCore/Private/MeshQueryFunctions.cpp
+++ b/Engine/Plugins/Experimental/GeometryScripting/Source/GeometryScriptingCore/Private/MeshQueryFunctions.cpp
@@ -127,11 +127,6 @@ int32 UGeometryScriptLibrary_MeshQueryFunctions::GetNumConnectedComponents( UDyn
 
 
 
-int32 UGeometryScriptLibrary_MeshQueryFunctions::GetTriangleCount( UDynamicMesh* TargetMesh )
-{
-	return SimpleMeshQuery<int32>(TargetMesh, 0, [&](const FDynamicMesh3& Mesh) { return Mesh.TriangleCount(); });
-}
-
 int32 UGeometryScriptLibrary_MeshQueryFunctions::GetNumTriangleIDs( UDynamicMesh* TargetMesh )
 {
 	return SimpleMeshQuery<int32>(TargetMesh, 0, [&](const FDynamicMesh3& Mesh) { return Mesh.MaxTriangleID(); });
diff --git a/Engine/Plugins/Experimental/GeometryScripting/Source/GeometryScriptingCore/Public/GeometryScript/MeshQueryFunctions.h b/Engine/Plugins/Experimental/GeometryScripting/Source/GeometryScriptingCore/Public/GeometryScript/MeshQueryFunctions.h
--- a/Engine/Plugins/Experimental/GeometryScripting/Source/GeometryScriptingCore/Public/GeometryScript/MeshQueryFunctions.h
+++ b/Engine/Plugins/Experimental/GeometryScripting/Source/GeometryScriptingCore/Public/GeometryScript/MeshQueryFunctions.h
@@ -40,9 +40,9 @@ public:
 	UFUNCTION(BlueprintPure, Category = "GeometryScript|MeshQueries", meta=(ScriptMethod))
 	static UPARAM(DisplayName = "Num Edges") int32 GetNumConnectedComponents( UDynamicMesh* TargetMesh );
 
-
-	UFUNCTION(BlueprintPure, Category = "GeometryScript|MeshQueries", meta=(ScriptMethod))
-	static UPARAM(DisplayName = "Triangle Count") int32 GetTriangleCount( UDynamicMesh* TargetMesh );
+	// UDynamicMesh already has this function
+	//UFUNCTION(BlueprintPure, Category = "GeometryScript|MeshQueries", meta=(ScriptMethod))
+	//static UPARAM(DisplayName = "Triangle Count") int32 GetTriangleCount( UDynamicMesh* TargetMesh );
 
 	UFUNCTION(BlueprintPure, Category = "GeometryScript|MeshQueries", meta=(ScriptMethod))
 	static UPARAM(DisplayName = "Num Triangle IDs") int32 GetNumTriangleIDs( UDynamicMesh* TargetMesh );
diff --git a/Engine/Source/Runtime/GeometryFramework/Public/UDynamicMesh.h b/Engine/Source/Runtime/GeometryFramework/Public/UDynamicMesh.h
--- a/Engine/Source/Runtime/GeometryFramework/Public/UDynamicMesh.h
+++ b/Engine/Source/Runtime/GeometryFramework/Public/UDynamicMesh.h
@@ -110,7 +110,7 @@ public:
 	 * This *does not* allocate a new mesh, so any existing mesh pointers/refs are still valid
 	 */
 	UFUNCTION(BlueprintCallable, Category = "Dynamic Mesh")
-	UPARAM(DisplayName = "Self") UDynamicMesh* 
+	UPARAM(DisplayName = "Target") UDynamicMesh* 
 	Reset();
 
 	/**
@@ -118,7 +118,7 @@ public:
 	 * This this instead of Reset() if an initially-empty mesh is undesirable (eg for a Component)
 	 */
 	UFUNCTION(BlueprintCallable, Category = "Dynamic Mesh")
-	UPARAM(DisplayName = "Self") UDynamicMesh* 
+	UPARAM(DisplayName = "Target") UDynamicMesh* 
 	ResetToCube();
 
 	//
@@ -140,7 +140,7 @@ public:
 	 * @return number of triangles in the mesh
 	 */
 	UFUNCTION(BlueprintCallable, Category = "Dynamic Mesh")
-	int32 GetTriangleCount() const;
+	UPARAM(DisplayName = "Triangle Count") int32 GetTriangleCount() const;
 
 
 	/** 
