GeometryScript: add CreateUniqueNewAssetPathName Function
#rnx

diff --git a/Engine/Plugins/Editor/ModelingToolsEditorMode/Source/ModelingToolsEditorMode/Private/ModelingModeAssetUtils.cpp b/Engine/Plugins/Editor/ModelingToolsEditorMode/Source/ModelingToolsEditorMode/Private/ModelingModeAssetUtils.cpp
--- a/Engine/Plugins/Editor/ModelingToolsEditorMode/Source/ModelingToolsEditorMode/Private/ModelingModeAssetUtils.cpp
+++ b/Engine/Plugins/Editor/ModelingToolsEditorMode/Source/ModelingToolsEditorMode/Private/ModelingModeAssetUtils.cpp
@@ -8,6 +8,7 @@
 #include "AssetToolsModule.h"
 #include "FileHelpers.h"
 #include "Misc/Paths.h"
+#include "ModelingObjectsCreationAPI.h"
 
 // for content-browser things
 #include "ContentBrowserModule.h"
@@ -21,49 +22,6 @@ namespace UE
 namespace Local
 {
 
-/**
- * Generate a N-letter GUID string that contains only hex digits,
- * and contains at least one letter and one number
- */
-static FString GenerateRandomShortHexString(int32 NumChars = 8)
-{
-	int32 FailCount = 0;
-	while (FailCount++ < 10)
-	{
-		FGuid Guid = FGuid::NewGuid();
-		FString GuidString = Guid.ToString(EGuidFormats::UniqueObjectGuid).ToUpper();
-		FString Result;
-		int32 Digits = 0, Letters = 0;
-		for (int32 k = 0; k < GuidString.Len(); ++k)
-		{
-			TCHAR Character = GuidString[k];
-			if (FChar::IsHexDigit(Character))
-			{
-				Result.AppendChar(Character);
-				if (FChar::IsDigit(Character))
-				{
-					Digits++;
-				}
-				else
-				{
-					Letters++;
-				}
-			}
-			if (Result.Len() == NumChars)
-			{
-				if (Digits > 0 && Letters > 0)
-				{
-					return Result;
-				}
-				break;		// exit loop
-			}
-		}
-	}
-	return TEXT("BADGUID1");
-}
-
-
-
 FString GetWorldRelativeAssetRootPath(const UWorld* World)
 {
 	if (ensure(World->GetOutermost() != nullptr) == false)
@@ -204,7 +162,7 @@ FString UE::Modeling::GetNewAssetPathName(const FString& BaseNameIn, const UWorl
 	FString UseBaseName = ObjectBaseName;
 	if (Settings->bAppendRandomStringToName)
 	{
-		FString GuidString = UE::Local::GenerateRandomShortHexString();
+		FString GuidString = UE::Modeling::GenerateRandomShortHexString();
 		UseBaseName = FString::Printf(TEXT("%s_%s"), *UseBaseName, *GuidString);
 	}
 
diff --git a/Engine/Plugins/Experimental/GeometryScripting/Source/GeometryScriptingEditor/GeometryScriptingEditor.Build.cs b/Engine/Plugins/Experimental/GeometryScripting/Source/GeometryScriptingEditor/GeometryScriptingEditor.Build.cs
--- a/Engine/Plugins/Experimental/GeometryScripting/Source/GeometryScriptingEditor/GeometryScriptingEditor.Build.cs
+++ b/Engine/Plugins/Experimental/GeometryScripting/Source/GeometryScriptingEditor/GeometryScriptingEditor.Build.cs
@@ -47,6 +47,7 @@ public class GeometryScriptingEditor : ModuleRules
 				"MeshConversion",
 				"GeometryAlgorithms",
 				"ModelingOperators",
+				"ModelingComponents",
 				"ModelingComponentsEditorOnly",
 				"EditorFramework",
 				"UnrealEd",
diff --git a/Engine/Plugins/Experimental/GeometryScripting/Source/GeometryScriptingEditor/Private/CreateNewAssetUtilityFunctions.cpp b/Engine/Plugins/Experimental/GeometryScripting/Source/GeometryScriptingEditor/Private/CreateNewAssetUtilityFunctions.cpp
--- a/Engine/Plugins/Experimental/GeometryScripting/Source/GeometryScriptingEditor/Private/CreateNewAssetUtilityFunctions.cpp
+++ b/Engine/Plugins/Experimental/GeometryScripting/Source/GeometryScriptingEditor/Private/CreateNewAssetUtilityFunctions.cpp
@@ -19,11 +19,14 @@
 #include "BSPOps.h"		// in UnrealEd
 #include "Editor/EditorEngine.h"		// for FActorLabelUtilities
 #include "Editor.h"
+#include "AssetRegistryModule.h"
+#include "Misc/Paths.h"
 
 #include "MeshDescriptionToDynamicMesh.h"
 #include "DynamicMeshToMeshDescription.h"
 #include "ConversionUtils/DynamicMeshToVolume.h"
 #include "AssetUtils/CreateStaticMeshUtil.h"
+#include "ModelingObjectsCreationAPI.h"
 
 #include "ExplicitUseGeometryMathTypes.h"		// using UE::Geometry::(math types)
 using namespace UE::Geometry;
@@ -31,6 +34,42 @@ using namespace UE::Geometry;
 #define LOCTEXT_NAMESPACE "UGeometryScriptLibrary_CreateNewAssetUtilityFunctions"
 
 
+
+
+
+
+void UGeometryScriptLibrary_CreateNewAssetFunctions::CreateUniqueNewAssetPathName(
+	FString AssetFolderPath,
+	FString BaseAssetName,
+	FString& UniqueAssetPathAndName,
+	FString& UniqueAssetName,
+	FGeometryScriptUniqueAssetNameOptions Options,
+	TEnumAsByte<EGeometryScriptOutcomePins>& Outcome,
+	UGeometryScriptDebug* Debug)
+{
+	FAssetRegistryModule& AssetRegistryModule = FModuleManager::LoadModuleChecked<FAssetRegistryModule>("AssetRegistry");
+
+	Outcome = EGeometryScriptOutcomePins::Failure;
+	int Attempts = 0;
+	while (Attempts++ < 10)
+	{
+		FString UUID = UE::Modeling::GenerateRandomShortHexString(Options.UniqueIDDigits);
+		UniqueAssetName = FString::Printf(TEXT("%s_%s"), *BaseAssetName, *UUID);
+		UniqueAssetPathAndName = FPaths::Combine(AssetFolderPath, UniqueAssetName);
+
+		// if asset does not exist at this path, we can use it
+		FAssetData AssetData = AssetRegistryModule.Get().GetAssetByObjectPath(*UniqueAssetPathAndName);
+		if (AssetData.IsValid() == false)
+		{
+			Outcome = EGeometryScriptOutcomePins::Success;
+			return;
+		}
+	}
+
+	UE::Geometry::AppendError(Debug, EGeometryScriptErrorType::OperationFailed, LOCTEXT("CreateUniqueNewAssetPathName_Failed", "Failed to find available unique Asset Path/Name"));
+}
+
+
 AVolume* UGeometryScriptLibrary_CreateNewAssetFunctions::CreateNewVolumeFromMesh(
 	UDynamicMesh* FromDynamicMesh,
 	UWorld* CreateInWorld,
@@ -164,6 +203,9 @@ UStaticMesh* UGeometryScriptLibrary_CreateNewAssetFunctions::CreateNewStaticMesh
 
 	GEditor->EndTransaction();
 
+	// publish new asset so that asset editor updates
+	FAssetRegistryModule::AssetCreated(NewStaticMesh);
+
 	Outcome = EGeometryScriptOutcomePins::Success;
 	return NewStaticMesh;
 }
diff --git a/Engine/Plugins/Experimental/GeometryScripting/Source/GeometryScriptingEditor/Public/GeometryScript/CreateNewAssetUtilityFunctions.h b/Engine/Plugins/Experimental/GeometryScripting/Source/GeometryScriptingEditor/Public/GeometryScript/CreateNewAssetUtilityFunctions.h
--- a/Engine/Plugins/Experimental/GeometryScripting/Source/GeometryScriptingEditor/Public/GeometryScript/CreateNewAssetUtilityFunctions.h
+++ b/Engine/Plugins/Experimental/GeometryScripting/Source/GeometryScriptingEditor/Public/GeometryScript/CreateNewAssetUtilityFunctions.h
@@ -14,6 +14,15 @@ class UDynamicMesh;
 class AVolume;
 
 
+USTRUCT(BlueprintType)
+struct GEOMETRYSCRIPTINGEDITOR_API FGeometryScriptUniqueAssetNameOptions
+{
+	GENERATED_BODY()
+public:
+	UPROPERTY(BlueprintReadWrite, Category = Options)
+	int32 UniqueIDDigits = 6;
+};
+
 
 USTRUCT(BlueprintType)
 struct GEOMETRYSCRIPTINGEDITOR_API FGeometryScriptCreateNewVolumeFromMeshOptions
@@ -64,6 +73,18 @@ class GEOMETRYSCRIPTINGEDITOR_API UGeometryScriptLibrary_CreateNewAssetFunctions
 	GENERATED_BODY()
 public:
 
+	UFUNCTION(BlueprintCallable, Category = "GeometryScript|AssetManagement", meta = (ExpandEnumAsExecs = "Outcome"))
+	static void
+	CreateUniqueNewAssetPathName(
+		FString AssetFolderPath,
+		FString BaseAssetName,
+		FString& UniqueAssetPathAndName,
+		FString& UniqueAssetName,
+		FGeometryScriptUniqueAssetNameOptions Options,
+		TEnumAsByte<EGeometryScriptOutcomePins>& Outcome,
+		UGeometryScriptDebug* Debug = nullptr);
+
+
 	UFUNCTION(BlueprintCallable, Category = "GeometryScript|AssetManagement", meta = (ExpandEnumAsExecs = "Outcome"))
 	static UPARAM(DisplayName = "Volume Actor") AVolume* 
 	CreateNewVolumeFromMesh(
