GeometryScript: add linear and curved stair generation functions
#rnx

diff --git a/Engine/Plugins/Experimental/GeometryScripting/Source/GeometryScriptingCore/Private/MeshPrimitiveFunctions.cpp b/Engine/Plugins/Experimental/GeometryScripting/Source/GeometryScriptingCore/Private/MeshPrimitiveFunctions.cpp
--- a/Engine/Plugins/Experimental/GeometryScripting/Source/GeometryScriptingCore/Private/MeshPrimitiveFunctions.cpp
+++ b/Engine/Plugins/Experimental/GeometryScripting/Source/GeometryScriptingCore/Private/MeshPrimitiveFunctions.cpp
@@ -17,6 +17,7 @@
 #include "Generators/SweepGenerator.h"
 #include "Generators/FlatTriangulationMeshGenerator.h"
 #include "Generators/RevolveGenerator.h"
+#include "Generators/StairGenerator.h"
 #include "ConstrainedDelaunay2.h"
 #include "Arrangement2d.h"
 
@@ -827,4 +828,76 @@ UDynamicMesh* UGeometryScriptLibrary_MeshPrimitiveFunctions::AppendTriangulatedP
 
 
 
+UDynamicMesh* UGeometryScriptLibrary_MeshPrimitiveFunctions::AppendLinearStairs(
+	UDynamicMesh* TargetMesh,
+	FGeometryScriptPrimitiveOptions PrimitiveOptions,
+	FTransform Transform,
+	float StepWidth,
+	float StepHeight,
+	float StepDepth,
+	int NumSteps,
+	bool bFloating,
+	UGeometryScriptDebug* Debug)
+{
+	if (TargetMesh == nullptr)
+	{
+		UE::Geometry::AppendError(Debug, EGeometryScriptErrorType::InvalidInputs, LOCTEXT("PrimitiveFunctions_AppendLinearStairs", "AppendLinearStairs: TargetMesh is Null"));
+		return TargetMesh;
+	}
+
+	TUniquePtr<FLinearStairGenerator> StairGenerator = (bFloating) ?
+		MakeUnique<FFloatingStairGenerator>() : MakeUnique<FLinearStairGenerator>();
+
+	StairGenerator->StepDepth = StepDepth;
+	StairGenerator->StepWidth = StepWidth;
+	StairGenerator->StepHeight = StepHeight;
+	StairGenerator->NumSteps = FMath::Max(1, NumSteps);
+	StairGenerator->bPolygroupPerQuad = (PrimitiveOptions.PolygroupMode == EGeometryScriptPrimitivePolygroupMode::PerQuad);
+	StairGenerator->bScaleUVByAspectRatio = (PrimitiveOptions.UVMode == EGeometryScriptPrimitiveUVMode::Uniform) ? true : false;
+	StairGenerator->Generate();
+	AppendPrimitive(TargetMesh, StairGenerator.Get(), Transform, PrimitiveOptions);
+
+	return TargetMesh;
+}
+
+
+
+
+UDynamicMesh* UGeometryScriptLibrary_MeshPrimitiveFunctions::AppendCurvedStairs(
+	UDynamicMesh* TargetMesh,
+	FGeometryScriptPrimitiveOptions PrimitiveOptions,
+	FTransform Transform,
+	float StepWidth,
+	float StepHeight,
+	float InnerRadius,
+	float CurveAngle,
+	int NumSteps,
+	bool bFloating,
+	UGeometryScriptDebug* Debug)
+{
+	if (TargetMesh == nullptr)
+	{
+		UE::Geometry::AppendError(Debug, EGeometryScriptErrorType::InvalidInputs, LOCTEXT("PrimitiveFunctions_AppendCurvedStairs", "AppendCurvedStairs: TargetMesh is Null"));
+		return TargetMesh;
+	}
+
+	TUniquePtr<FCurvedStairGenerator> StairGenerator = (bFloating) ?
+		MakeUnique<FSpiralStairGenerator>() : MakeUnique<FCurvedStairGenerator>();
+
+	StairGenerator->StepWidth = StepWidth;
+	StairGenerator->StepHeight = StepHeight;
+	StairGenerator->NumSteps = FMath::Max(1, NumSteps);
+	StairGenerator->InnerRadius = InnerRadius;
+	StairGenerator->CurveAngle = CurveAngle;
+	StairGenerator->bPolygroupPerQuad = (PrimitiveOptions.PolygroupMode == EGeometryScriptPrimitivePolygroupMode::PerQuad);
+	StairGenerator->bScaleUVByAspectRatio = (PrimitiveOptions.UVMode == EGeometryScriptPrimitiveUVMode::Uniform) ? true : false;
+	StairGenerator->Generate();
+	AppendPrimitive(TargetMesh, StairGenerator.Get(), Transform, PrimitiveOptions);
+
+	return TargetMesh;
+}
+
+
+
+
 #undef LOCTEXT_NAMESPACE
\ No newline at end of file
diff --git a/Engine/Plugins/Experimental/GeometryScripting/Source/GeometryScriptingCore/Public/GeometryScript/MeshPrimitiveFunctions.h b/Engine/Plugins/Experimental/GeometryScripting/Source/GeometryScriptingCore/Public/GeometryScript/MeshPrimitiveFunctions.h
--- a/Engine/Plugins/Experimental/GeometryScripting/Source/GeometryScriptingCore/Public/GeometryScript/MeshPrimitiveFunctions.h
+++ b/Engine/Plugins/Experimental/GeometryScripting/Source/GeometryScriptingCore/Public/GeometryScript/MeshPrimitiveFunctions.h
@@ -23,6 +23,13 @@ enum class EGeometryScriptPrimitiveOriginMode : uint8
 	Base = 1
 };
 
+UENUM(BlueprintType)
+enum class EGeometryScriptPrimitiveUVMode : uint8
+{
+	Uniform = 0,
+	ScaleToFill = 1
+};
+
 
 USTRUCT(BlueprintType)
 struct GEOMETRYSCRIPTINGCORE_API FGeometryScriptPrimitiveOptions
@@ -34,6 +41,9 @@ public:
 
 	UPROPERTY(BlueprintReadWrite, Category = Options)
 	bool bFlipOrientation = false;
+
+	UPROPERTY(BlueprintReadWrite, Category = Options)
+	EGeometryScriptPrimitiveUVMode UVMode = EGeometryScriptPrimitiveUVMode::Uniform;
 };
 
 
@@ -306,4 +316,34 @@ public:
 		bool bAllowSelfIntersections = true,
 		UGeometryScriptDebug* Debug = nullptr);
 
+
+
+	UFUNCTION(BlueprintCallable, Category = "GeometryScript|Primitives", meta=(ScriptMethod))
+	static UPARAM(DisplayName = "Target Mesh") UDynamicMesh* 
+	AppendLinearStairs( 
+		UDynamicMesh* TargetMesh, 
+		FGeometryScriptPrimitiveOptions PrimitiveOptions,
+		FTransform Transform,
+		float StepWidth = 100.0f,
+		float StepHeight = 20.0f,
+		float StepDepth = 30.0f,
+		int NumSteps = 8,
+		bool bFloating = false,
+		UGeometryScriptDebug* Debug = nullptr);
+
+
+	UFUNCTION(BlueprintCallable, Category = "GeometryScript|Primitives", meta=(ScriptMethod))
+	static UPARAM(DisplayName = "Target Mesh") UDynamicMesh* 
+	AppendCurvedStairs( 
+		UDynamicMesh* TargetMesh, 
+		FGeometryScriptPrimitiveOptions PrimitiveOptions,
+		FTransform Transform,
+		float StepWidth = 100.0f,
+		float StepHeight = 20.0f,
+		float InnerRadius = 150.0f,
+		float CurveAngle = 90.0f,
+		int NumSteps = 8,
+		bool bFloating = false,
+		UGeometryScriptDebug* Debug = nullptr);
+
 };
\ No newline at end of file
