GeometryScript: change plugin label and text blurb
#rnx

diff --git a/Engine/Plugins/Experimental/GeometryScripting/GeometryScripting.uplugin b/Engine/Plugins/Experimental/GeometryScripting/GeometryScripting.uplugin
--- a/Engine/Plugins/Experimental/GeometryScripting/GeometryScripting.uplugin
+++ b/Engine/Plugins/Experimental/GeometryScripting/GeometryScripting.uplugin
@@ -2,8 +2,8 @@
 	"FileVersion": 3,
 	"Version": 1,
 	"VersionName": "0.1",
-	"FriendlyName": "Geometry Scripting",
-	"Description": "Geometry Scripting Modules enable procedural generation and manipulation of meshes/geometry via Blueprints",
+	"FriendlyName": "Geometry Script",
+	"Description": "Geometry Script provides a library of functions for creating and editing Meshes in Blueprints and Python",
 	"Category": "Geometry",
 	"CreatedBy": "Epic Games, Inc.",
 	"CreatedByURL": "http://epicgames.com",
