GeometryScript: add ConvertComponentsToPolygroups, GetTrianglePolygroupID, DeleteTrianglesInPolygroup, SetPolygroupMaterialID functions
#rnx

diff --git a/Engine/Plugins/Experimental/GeometryScripting/Source/GeometryScriptingCore/Private/MeshMaterialFunctions.cpp b/Engine/Plugins/Experimental/GeometryScripting/Source/GeometryScriptingCore/Private/MeshMaterialFunctions.cpp
--- a/Engine/Plugins/Experimental/GeometryScripting/Source/GeometryScriptingCore/Private/MeshMaterialFunctions.cpp
+++ b/Engine/Plugins/Experimental/GeometryScripting/Source/GeometryScriptingCore/Private/MeshMaterialFunctions.cpp
@@ -5,6 +5,7 @@
 #include "DynamicMesh/DynamicMesh3.h"
 #include "DynamicMesh/DynamicMeshAttributeSet.h"
 #include "UDynamicMesh.h"
+#include "Polygroups/PolygroupSet.h"
 
 #include "ExplicitUseGeometryMathTypes.h"		// using UE::Geometry::(math types)
 using namespace UE::Geometry;
@@ -270,4 +271,50 @@ UDynamicMesh* UGeometryScriptLibrary_MeshMaterialFunctions::SetAllTriangleMateri
 
 
 
+
+UDynamicMesh* UGeometryScriptLibrary_MeshMaterialFunctions::SetPolygroupMaterialID( 
+	UDynamicMesh* TargetMesh, 
+	FGeometryScriptGroupLayer GroupLayer,
+	int PolygroupID, 
+	int MaterialID,
+	bool& bIsValidPolygroupID,
+	bool bDeferChangeNotifications,
+	UGeometryScriptDebug* Debug)
+{
+	bIsValidPolygroupID = false;
+	if (TargetMesh)
+	{
+		TargetMesh->EditMesh([&](FDynamicMesh3& EditMesh)
+		{
+			FPolygroupLayer InputGroupLayer{ GroupLayer.bDefaultLayer, GroupLayer.ExtendedLayerIndex };
+			if (InputGroupLayer.CheckExists(&EditMesh) == false)
+			{
+				UE::Geometry::AppendError(Debug, EGeometryScriptErrorType::InvalidInputs, LOCTEXT("SetPolygroupMaterialID_MissingGroups", "SetPolygroupMaterialID: Specified Polygroup Layer does not exist"));
+				return;
+			}
+			FDynamicMeshMaterialAttribute* MaterialIDs = (EditMesh.HasAttributes() && EditMesh.Attributes()->HasMaterialID()) ? EditMesh.Attributes()->GetMaterialID() : nullptr;
+			if (MaterialIDs == nullptr)
+			{
+				UE::Geometry::AppendError(Debug, EGeometryScriptErrorType::InvalidInputs, LOCTEXT("SetPolygroupMaterialID_NoMaterialID", "SetPolygroupMaterialID: MaterialID Attribute is not enabled"));
+				return;
+			}
+
+			FPolygroupSet Groups(&EditMesh, InputGroupLayer);
+			for (int32 tid : EditMesh.TriangleIndicesItr())
+			{
+				if (Groups.GetGroup(tid) == PolygroupID)
+				{
+					MaterialIDs->SetValue(tid, MaterialID);
+					bIsValidPolygroupID = true;
+				}
+			}
+
+		}, EDynamicMeshChangeType::GeneralEdit, EDynamicMeshAttributeChangeFlags::Unknown, bDeferChangeNotifications);
+	}
+	return TargetMesh;	
+}
+
+
+
+
 #undef LOCTEXT_NAMESPACE
\ No newline at end of file
diff --git a/Engine/Plugins/Experimental/GeometryScripting/Source/GeometryScriptingCore/Private/MeshPolygroupFunctions.cpp b/Engine/Plugins/Experimental/GeometryScripting/Source/GeometryScriptingCore/Private/MeshPolygroupFunctions.cpp
--- a/Engine/Plugins/Experimental/GeometryScripting/Source/GeometryScriptingCore/Private/MeshPolygroupFunctions.cpp
+++ b/Engine/Plugins/Experimental/GeometryScripting/Source/GeometryScriptingCore/Private/MeshPolygroupFunctions.cpp
@@ -15,6 +15,31 @@ using namespace UE::Geometry;
 
 
 
+template<typename ReturnType> 
+ReturnType SimpleMeshPolygroupQuery(UDynamicMesh* Mesh, FGeometryScriptGroupLayer ScriptGroupLayer, bool& bHasGroups, ReturnType DefaultValue, 
+	TFunctionRef<ReturnType(const FDynamicMesh3& Mesh, const FPolygroupSet& Polygroups)> QueryFunc)
+{
+	bHasGroups = false;
+	ReturnType RetVal = DefaultValue;
+	if (Mesh)
+	{
+		Mesh->ProcessMesh([&](const FDynamicMesh3& ReadMesh)
+		{
+			FPolygroupLayer GroupLayer{ ScriptGroupLayer.bDefaultLayer, ScriptGroupLayer.ExtendedLayerIndex };
+			if (GroupLayer.CheckExists(&ReadMesh))
+			{
+				FPolygroupSet Groups(&ReadMesh, GroupLayer);
+				bHasGroups = true;
+				RetVal = QueryFunc(ReadMesh, Groups);
+			}
+		});
+	}
+	return RetVal;
+}
+
+
+
+
 UDynamicMesh* UGeometryScriptLibrary_MeshPolygroupFunctions::EnablePolygroups( 
 	UDynamicMesh* TargetMesh, 
 	UGeometryScriptDebug* Debug)
@@ -189,6 +214,41 @@ UDynamicMesh* UGeometryScriptLibrary_MeshPolygroupFunctions::ConvertUVIslandsToP
 
 
 
+UDynamicMesh* UGeometryScriptLibrary_MeshPolygroupFunctions::ConvertComponentsToPolygroups( 
+	UDynamicMesh* TargetMesh, 
+	FGeometryScriptGroupLayer GroupLayer,
+	UGeometryScriptDebug* Debug)
+{
+	if (TargetMesh == nullptr)
+	{
+		UE::Geometry::AppendError(Debug, EGeometryScriptErrorType::InvalidInputs, LOCTEXT("ConvertComponentsToPolygroups_InvalidInput", "ConvertComponentsToPolygroups: TargetMesh is Null"));
+		return TargetMesh;
+	}
+
+	TargetMesh->EditMesh([&](FDynamicMesh3& EditMesh) 
+	{
+		TUniquePtr<FPolygroupSet> OutputGroups;
+		FPolygroupLayer InputGroupLayer{ GroupLayer.bDefaultLayer, GroupLayer.ExtendedLayerIndex };
+		if (InputGroupLayer.CheckExists(&EditMesh) == false)
+		{
+			UE::Geometry::AppendError(Debug, EGeometryScriptErrorType::InvalidInputs, LOCTEXT("ConvertUVIslandsToPolygroups_MissingGroups", "ConvertUVIslandsToPolygroups: Target Polygroup Layer does not exist"));
+			return;
+		}
+		OutputGroups = MakeUnique<FPolygroupSet>(&EditMesh, InputGroupLayer);
+
+		FPolygroupsGenerator Generator(&EditMesh);
+		Generator.bApplyPostProcessing = false;
+		Generator.bCopyToMesh = false;
+		Generator.FindPolygroupsFromConnectedTris();
+		Generator.CopyPolygroupsToPolygroupSet(*OutputGroups, EditMesh);
+
+	}, EDynamicMeshChangeType::GeneralEdit, EDynamicMeshAttributeChangeFlags::Unknown, false);
+
+	return TargetMesh;
+}
+
+
+
 UDynamicMesh* UGeometryScriptLibrary_MeshPolygroupFunctions::ComputePolygroupsFromAngleThreshold( 
 	UDynamicMesh* TargetMesh, 
 	FGeometryScriptGroupLayer GroupLayer,
@@ -228,4 +288,67 @@ UDynamicMesh* UGeometryScriptLibrary_MeshPolygroupFunctions::ComputePolygroupsFr
 
 
 
+
+int32 UGeometryScriptLibrary_MeshPolygroupFunctions::GetTrianglePolygroupID( 
+	UDynamicMesh* TargetMesh, 
+	FGeometryScriptGroupLayer GroupLayer, 
+	int TriangleID, 
+	bool& bIsValidTriangle)
+{
+	bIsValidTriangle = false;
+	bool bIsValidPolygroupLayer = false;
+	return SimpleMeshPolygroupQuery<int32>(TargetMesh, GroupLayer, bIsValidPolygroupLayer, 0, 
+	[&](const FDynamicMesh3& Mesh, const FPolygroupSet& PolyGroups) {
+		bIsValidTriangle = Mesh.IsTriangle(TriangleID);
+		return (bIsValidTriangle) ? PolyGroups.GetGroup(TriangleID) : 0;
+	});
+}
+
+
+
+UDynamicMesh* UGeometryScriptLibrary_MeshPolygroupFunctions::DeleteTrianglesInPolygroup(
+	UDynamicMesh* TargetMesh,
+	FGeometryScriptGroupLayer GroupLayer,
+	int PolygroupID,
+	int& NumDeleted,
+	bool bDeferChangeNotifications,
+	UGeometryScriptDebug* Debug)
+{
+	NumDeleted = 0;
+	if (TargetMesh)
+	{
+		TargetMesh->EditMesh([&](FDynamicMesh3& EditMesh)
+		{
+			FPolygroupLayer InputGroupLayer{ GroupLayer.bDefaultLayer, GroupLayer.ExtendedLayerIndex };
+			if (InputGroupLayer.CheckExists(&EditMesh) == false)
+			{
+				UE::Geometry::AppendError(Debug, EGeometryScriptErrorType::InvalidInputs, LOCTEXT("DeleteTrianglesInPolygroup_MissingGroups", "DeleteTrianglesInPolygroup: Specified Polygroup Layer does not exist"));
+				return;
+			}
+
+			FPolygroupSet Groups(&EditMesh, InputGroupLayer);
+			TArray<int32> TriangleList;
+			for (int32 TriangleID : EditMesh.TriangleIndicesItr())
+			{
+				if (Groups.GetGroup(TriangleID) == PolygroupID)
+				{
+					TriangleList.Add(TriangleID);
+				}
+			}
+
+			for (int32 TriangleID : TriangleList)
+			{
+				EMeshResult Result = EditMesh.RemoveTriangle(TriangleID);
+				if (Result == EMeshResult::Ok)
+				{
+					NumDeleted++;
+				}
+			}
+		}, EDynamicMeshChangeType::GeneralEdit, EDynamicMeshAttributeChangeFlags::Unknown, bDeferChangeNotifications);
+	}
+	return TargetMesh;
+}
+
+
+
 #undef LOCTEXT_NAMESPACE
\ No newline at end of file
diff --git a/Engine/Plugins/Experimental/GeometryScripting/Source/GeometryScriptingCore/Public/GeometryScript/MeshMaterialFunctions.h b/Engine/Plugins/Experimental/GeometryScripting/Source/GeometryScriptingCore/Public/GeometryScript/MeshMaterialFunctions.h
--- a/Engine/Plugins/Experimental/GeometryScripting/Source/GeometryScriptingCore/Public/GeometryScript/MeshMaterialFunctions.h
+++ b/Engine/Plugins/Experimental/GeometryScripting/Source/GeometryScriptingCore/Public/GeometryScript/MeshMaterialFunctions.h
@@ -40,7 +40,7 @@ public:
 	static UPARAM(DisplayName = "Material ID") int32
 	GetTriangleMaterialID( UDynamicMesh* TargetMesh, int TriangleID, bool& bIsValidTriangle );
 
-	UFUNCTION(BlueprintPure, Category = "GeometryScript|MeshQueries", meta=(ScriptMethod))
+	UFUNCTION(BlueprintPure, Category = "GeometryScript|Materials", meta=(ScriptMethod))
 	static UPARAM(DisplayName = "Target Mesh") UDynamicMesh* 
 	GetAllTriangleMaterialIDs( UDynamicMesh* TargetMesh, UPARAM(ref) TArray<int>& MaterialIDs, bool& bHasMaterialIDs );
 
@@ -62,4 +62,15 @@ public:
 		bool bDeferChangeNotifications = false,
 		UGeometryScriptDebug* Debug = nullptr);
 
+
+	UFUNCTION(BlueprintCallable, Category = "GeometryScript|Materials", meta=(ScriptMethod))
+	static UPARAM(DisplayName = "Target Mesh") UDynamicMesh* 
+	SetPolygroupMaterialID( 
+		UDynamicMesh* TargetMesh, 
+		FGeometryScriptGroupLayer GroupLayer,
+		int PolygroupID, 
+		int MaterialID,
+		bool& bIsValidPolygroupID,
+		bool bDeferChangeNotifications = false,
+		UGeometryScriptDebug* Debug = nullptr);
 };
\ No newline at end of file
diff --git a/Engine/Plugins/Experimental/GeometryScripting/Source/GeometryScriptingCore/Public/GeometryScript/MeshPolygroupFunctions.h b/Engine/Plugins/Experimental/GeometryScripting/Source/GeometryScriptingCore/Public/GeometryScript/MeshPolygroupFunctions.h
--- a/Engine/Plugins/Experimental/GeometryScripting/Source/GeometryScriptingCore/Public/GeometryScript/MeshPolygroupFunctions.h
+++ b/Engine/Plugins/Experimental/GeometryScripting/Source/GeometryScriptingCore/Public/GeometryScript/MeshPolygroupFunctions.h
@@ -52,6 +52,13 @@ public:
 		int UVLayer = 0,
 		UGeometryScriptDebug* Debug = nullptr );
 
+	UFUNCTION(BlueprintCallable, Category = "GeometryScript|Polygroups", meta=(ScriptMethod))
+	static UPARAM(DisplayName = "Target Mesh") UDynamicMesh* 
+	ConvertComponentsToPolygroups( 
+		UDynamicMesh* TargetMesh, 
+		FGeometryScriptGroupLayer GroupLayer,
+		UGeometryScriptDebug* Debug = nullptr );
+
 	UFUNCTION(BlueprintCallable, Category = "GeometryScript|Polygroups", meta=(ScriptMethod))
 	static UPARAM(DisplayName = "Target Mesh") UDynamicMesh* 
 	ComputePolygroupsFromAngleThreshold( 
@@ -61,4 +68,23 @@ public:
 		int MinGroupSize = 2,
 		UGeometryScriptDebug* Debug = nullptr );
 
+
+	UFUNCTION(BlueprintPure, Category = "GeometryScript|Polygroups", meta=(ScriptMethod))
+	static UPARAM(DisplayName = "Material ID") int32
+	GetTrianglePolygroupID( 
+		UDynamicMesh* TargetMesh, 
+		FGeometryScriptGroupLayer GroupLayer, 
+		int TriangleID, 
+		bool& bIsValidTriangle );
+
+	UFUNCTION(BlueprintCallable, Category = "GeometryScript|Polygroups", meta=(ScriptMethod))
+	static UPARAM(DisplayName = "Target Mesh") UDynamicMesh* 
+	DeleteTrianglesInPolygroup( 
+		UDynamicMesh* TargetMesh, 
+		FGeometryScriptGroupLayer GroupLayer,
+		int PolygroupID,
+		int& NumDeleted,
+		bool bDeferChangeNotifications = false,
+		UGeometryScriptDebug* Debug = nullptr);
+
 };
\ No newline at end of file
