GeometryScript: fix bug where wrong size was being used to pre-allocate array
#rnx

diff --git a/Engine/Plugins/Experimental/GeometryScripting/Source/GeometryScriptingCore/Private/CollisionFunctions.cpp b/Engine/Plugins/Experimental/GeometryScripting/Source/GeometryScriptingCore/Private/CollisionFunctions.cpp
--- a/Engine/Plugins/Experimental/GeometryScripting/Source/GeometryScriptingCore/Private/CollisionFunctions.cpp
+++ b/Engine/Plugins/Experimental/GeometryScripting/Source/GeometryScriptingCore/Private/CollisionFunctions.cpp
@@ -46,19 +46,20 @@ void ComputeCollisionFromMesh(
 
 	FMeshConnectedComponents Components(&Mesh);
 	Components.FindConnectedTriangles();
+	int32 NumComponents = Components.Num();
 
 	TArray<FDynamicMesh3> Submeshes;
 	TArray<const FDynamicMesh3*> SubmeshPointers;
-	SubmeshPointers.SetNum(Submeshes.Num());
+	SubmeshPointers.SetNum(NumComponents);
 
-	if (Components.Num() == 1)
+	if (NumComponents == 1)
 	{
 		SubmeshPointers.Add(&Mesh);
 	}
 	else
 	{
-		Submeshes.SetNum(Components.Num());
-		ParallelFor(Submeshes.Num(), [&](int32 k)
+		Submeshes.SetNum(NumComponents);
+		ParallelFor(NumComponents, [&](int32 k)
 		{
 			FDynamicSubmesh3 Submesh(&Mesh, Components[k].Indices, (int32)EMeshComponents::None, false);
 			Submeshes[k] = MoveTemp(Submesh.GetSubmesh());
