GeometryScript: replace usage of TArray input/output params for index/vector/color/triangle lists with custom list UStructs that contain a TSharedPtr to a TArray (FGeometryScriptIndexList, FGeometryScriptTriangleList, FGeometryScriptVectorList, FGeometryScriptUVList, FGeometryScriptColorList). This makes it safe(r) to work with large lists and will allow for future changes to list representation (like using DynamicVector / etc).
- add ListUtilityFunctions with length/lastindex/item and convert to/from array
- Updated various existing functions that were using a TArray as input/output
- Added GetAllTriangleIDs, GetAllTriangleIndices, GetAllVertexIDs, updated GetAllVertexPositions to support optionally skipping invalid vertices
#rnx

diff --git a/Engine/Plugins/Experimental/GeometryScripting/Source/GeometryScriptingCore/Private/ListUtilityFunctions.cpp b/Engine/Plugins/Experimental/GeometryScripting/Source/GeometryScriptingCore/Private/ListUtilityFunctions.cpp
new file mode 100644
--- /dev/null
+++ b/Engine/Plugins/Experimental/GeometryScripting/Source/GeometryScriptingCore/Private/ListUtilityFunctions.cpp
@@ -0,0 +1,208 @@
+// Copyright Epic Games, Inc. All Rights Reserved.
+
+#include "GeometryScript/ListUtilityFunctions.h"
+
+
+#define LOCTEXT_NAMESPACE "UGeometryScriptLibrary_ListUtilityFunctions"
+
+
+int UGeometryScriptLibrary_ListUtilityFunctions::GetIndexListLength(FGeometryScriptIndexList IndexList)
+{
+	return (IndexList.List.IsValid()) ? IndexList.List->Num() : 0;
+}
+
+int UGeometryScriptLibrary_ListUtilityFunctions::GetIndexListLastIndex(FGeometryScriptIndexList IndexList)
+{
+	return (IndexList.List.IsValid()) ? FMath::Max(IndexList.List->Num()-1,0) : 0;
+}
+
+int UGeometryScriptLibrary_ListUtilityFunctions::GetIndexListItem(FGeometryScriptIndexList IndexList, int Index, bool& bIsValidIndex)
+{
+	bIsValidIndex = false;
+	if (IndexList.List.IsValid() && Index >= 0 && Index < IndexList.List->Num())
+	{
+		bIsValidIndex = true;
+		return (*IndexList.List)[Index];
+	}
+	return -1;
+}
+
+void UGeometryScriptLibrary_ListUtilityFunctions::ConvertIndexListToArray(FGeometryScriptIndexList IndexList, TArray<int>& IndexArray)
+{
+	IndexArray.Reset();
+	if (IndexList.List.IsValid())
+	{
+		IndexArray.Append(*IndexList.List);
+	}
+}
+
+void UGeometryScriptLibrary_ListUtilityFunctions::ConvertArrayToIndexList(const TArray<int>& IndexArray, FGeometryScriptIndexList& IndexList)
+{
+	IndexList.Reset();
+	IndexList.List->Append(IndexArray);
+}
+
+
+
+
+
+int UGeometryScriptLibrary_ListUtilityFunctions::GetTriangleListLength(FGeometryScriptTriangleList TriangleList)
+{
+	return (TriangleList.List.IsValid()) ? TriangleList.List->Num() : 0;
+}
+
+int UGeometryScriptLibrary_ListUtilityFunctions::GetTriangleListLastTriangle(FGeometryScriptTriangleList TriangleList)
+{
+	return (TriangleList.List.IsValid()) ? FMath::Max(TriangleList.List->Num()-1,0) : 0;
+}
+
+FIntVector UGeometryScriptLibrary_ListUtilityFunctions::GetTriangleListItem(FGeometryScriptTriangleList TriangleList, int Triangle, bool& bIsValidTriangle)
+{
+	bIsValidTriangle = false;
+	if (TriangleList.List.IsValid() && Triangle >= 0 && Triangle < TriangleList.List->Num())
+	{
+		bIsValidTriangle = true;
+		return (*TriangleList.List)[Triangle];
+	}
+	return FIntVector::NoneValue;
+}
+
+void UGeometryScriptLibrary_ListUtilityFunctions::ConvertTriangleListToArray(FGeometryScriptTriangleList TriangleList, TArray<FIntVector>& TriangleArray)
+{
+	TriangleArray.Reset();
+	if (TriangleList.List.IsValid())
+	{
+		TriangleArray.Append(*TriangleList.List);
+	}
+}
+
+void UGeometryScriptLibrary_ListUtilityFunctions::ConvertArrayToTriangleList(const TArray<FIntVector>& TriangleArray, FGeometryScriptTriangleList& TriangleList)
+{
+	TriangleList.Reset();
+	TriangleList.List->Append(TriangleArray);
+}
+
+
+
+
+
+int UGeometryScriptLibrary_ListUtilityFunctions::GetVectorListLength(FGeometryScriptVectorList VectorList)
+{
+	return (VectorList.List.IsValid()) ? VectorList.List->Num() : 0;
+}
+
+int UGeometryScriptLibrary_ListUtilityFunctions::GetVectorListLastIndex(FGeometryScriptVectorList VectorList)
+{
+	return (VectorList.List.IsValid()) ? FMath::Max(VectorList.List->Num()-1,0) : 0;
+}
+
+FVector UGeometryScriptLibrary_ListUtilityFunctions::GetVectorListItem(FGeometryScriptVectorList VectorList, int Index, bool& bIsValidIndex)
+{
+	bIsValidIndex = false;
+	if (VectorList.List.IsValid() && Index >= 0 && Index < VectorList.List->Num())
+	{
+		bIsValidIndex = true;
+		return (*VectorList.List)[Index];
+	}
+	return FVector::Zero();
+}
+
+void UGeometryScriptLibrary_ListUtilityFunctions::ConvertVectorListToArray(FGeometryScriptVectorList VectorList, TArray<FVector>& VectorArray)
+{
+	VectorArray.Reset();
+	if (VectorList.List.IsValid())
+	{
+		VectorArray.Append(*VectorList.List);
+	}
+}
+
+void UGeometryScriptLibrary_ListUtilityFunctions::ConvertArrayToVectorList(const TArray<FVector>& VectorArray, FGeometryScriptVectorList& VectorList)
+{
+	VectorList.Reset();
+	VectorList.List->Append(VectorArray);
+}
+
+
+
+
+
+int UGeometryScriptLibrary_ListUtilityFunctions::GetUVListLength(FGeometryScriptUVList UVList)
+{
+	return (UVList.List.IsValid()) ? UVList.List->Num() : 0;
+}
+
+int UGeometryScriptLibrary_ListUtilityFunctions::GetUVListLastIndex(FGeometryScriptUVList UVList)
+{
+	return (UVList.List.IsValid()) ? FMath::Max(UVList.List->Num()-1,0) : 0;
+}
+
+FVector2D UGeometryScriptLibrary_ListUtilityFunctions::GetUVListItem(FGeometryScriptUVList UVList, int Index, bool& bIsValidIndex)
+{
+	bIsValidIndex = false;
+	if (UVList.List.IsValid() && Index >= 0 && Index < UVList.List->Num())
+	{
+		bIsValidIndex = true;
+		return (*UVList.List)[Index];
+	}
+	return FVector2D::ZeroVector;
+}
+
+void UGeometryScriptLibrary_ListUtilityFunctions::ConvertUVListToArray(FGeometryScriptUVList UVList, TArray<FVector2D>& UVArray)
+{
+	UVArray.Reset();
+	if (UVList.List.IsValid())
+	{
+		UVArray.Append(*UVList.List);
+	}
+}
+
+void UGeometryScriptLibrary_ListUtilityFunctions::ConvertArrayToUVList(const TArray<FVector2D>& UVArray, FGeometryScriptUVList& UVList)
+{
+	UVList.Reset();
+	UVList.List->Append(UVArray);
+}
+
+
+
+
+
+int UGeometryScriptLibrary_ListUtilityFunctions::GetColorListLength(FGeometryScriptColorList ColorList)
+{
+	return (ColorList.List.IsValid()) ? ColorList.List->Num() : 0;
+}
+
+int UGeometryScriptLibrary_ListUtilityFunctions::GetColorListLastIndex(FGeometryScriptColorList ColorList)
+{
+	return (ColorList.List.IsValid()) ? FMath::Max(ColorList.List->Num()-1,0) : 0;
+}
+
+FLinearColor UGeometryScriptLibrary_ListUtilityFunctions::GetColorListItem(FGeometryScriptColorList ColorList, int Index, bool& bIsValidIndex)
+{
+	bIsValidIndex = false;
+	if (ColorList.List.IsValid() && Index >= 0 && Index < ColorList.List->Num())
+	{
+		bIsValidIndex = true;
+		return (*ColorList.List)[Index];
+	}
+	return FLinearColor::White;
+}
+
+void UGeometryScriptLibrary_ListUtilityFunctions::ConvertColorListToArray(FGeometryScriptColorList ColorList, TArray<FLinearColor>& ColorArray)
+{
+	ColorArray.Reset();
+	if (ColorList.List.IsValid())
+	{
+		ColorArray.Append(*ColorList.List);
+	}
+}
+
+void UGeometryScriptLibrary_ListUtilityFunctions::ConvertArrayToColorList(const TArray<FLinearColor>& ColorArray, FGeometryScriptColorList& ColorList)
+{
+	ColorList.Reset();
+	ColorList.List->Append(ColorArray);
+}
+
+
+
+
+#undef LOCTEXT_NAMESPACE
\ No newline at end of file
diff --git a/Engine/Plugins/Experimental/GeometryScripting/Source/GeometryScriptingCore/Private/MeshBasicEditFunctions.cpp b/Engine/Plugins/Experimental/GeometryScripting/Source/GeometryScriptingCore/Private/MeshBasicEditFunctions.cpp
--- a/Engine/Plugins/Experimental/GeometryScripting/Source/GeometryScriptingCore/Private/MeshBasicEditFunctions.cpp
+++ b/Engine/Plugins/Experimental/GeometryScripting/Source/GeometryScriptingCore/Private/MeshBasicEditFunctions.cpp
@@ -71,19 +71,24 @@ UDynamicMesh* UGeometryScriptLibrary_MeshBasicEditFunctions::AddVertexToMesh(
 
 UDynamicMesh* UGeometryScriptLibrary_MeshBasicEditFunctions::AddVerticesToMesh(
 	UDynamicMesh* TargetMesh,
-	const TArray<FVector>& NewPositions, 
-	TArray<int>& NewIndices,
+	FGeometryScriptVectorList NewPositionsList, 
+	FGeometryScriptIndexList& NewIndicesList,
 	bool bDeferChangeNotifications)
 {
-	NewIndices.Reset();
+	if (NewPositionsList.List.IsValid() == false || NewPositionsList.List->Num() == 0)
+	{
+		return TargetMesh;
+	}
+
+	NewIndicesList.Reset();
 	if (TargetMesh)
 	{
 		TargetMesh->EditMesh([&](FDynamicMesh3& EditMesh)
 		{
-			for (FVector Position : NewPositions)
+			for (FVector Position : *NewPositionsList.List)
 			{
 				int32 NewVertexIndex = EditMesh.AppendVertex(Position);
-				NewIndices.Add(NewVertexIndex);
+				NewIndicesList.List->Add(NewVertexIndex);
 			}
 		
 		}, EDynamicMeshChangeType::GeneralEdit, EDynamicMeshAttributeChangeFlags::Unknown, bDeferChangeNotifications);
@@ -114,16 +119,21 @@ UDynamicMesh* UGeometryScriptLibrary_MeshBasicEditFunctions::DeleteVertexFromMes
 
 UDynamicMesh* UGeometryScriptLibrary_MeshBasicEditFunctions::DeleteVerticesFromMesh(
 	UDynamicMesh* TargetMesh,
-	const TArray<int32>& VertexList,
+	FGeometryScriptIndexList VertexList,
 	int& NumDeleted,
 	bool bDeferChangeNotifications)
 {
+	if (VertexList.List.IsValid() == false || VertexList.List->Num() == 0)
+	{
+		return TargetMesh;
+	}
+
 	NumDeleted = 0;
 	if (TargetMesh)
 	{
 		TargetMesh->EditMesh([&](FDynamicMesh3& EditMesh)
 		{
-			for (int32 VertexID : VertexList)
+			for (int32 VertexID : *VertexList.List)
 			{
 				EMeshResult Result = EditMesh.RemoveVertex(VertexID);
 				if (Result == EMeshResult::Ok)
@@ -186,26 +196,32 @@ UDynamicMesh* UGeometryScriptLibrary_MeshBasicEditFunctions::AddTriangleToMesh(
 
 UDynamicMesh* UGeometryScriptLibrary_MeshBasicEditFunctions::AddTrianglesToMesh(
 	UDynamicMesh* TargetMesh,
-	const TArray<FIntVector>& NewTriangles,
-	TArray<int>& NewIndices,
+	FGeometryScriptTriangleList NewTrianglesList,
+	FGeometryScriptIndexList& NewIndicesList,
 	int32 NewTriangleGroupID,
 	bool bDeferChangeNotifications,
 	UGeometryScriptDebug* Debug)
 {
-	NewIndices.Reset();
+	if (NewTrianglesList.List.IsValid() == false || NewTrianglesList.List->Num() == 0)
+	{
+		UE::Geometry::AppendError(Debug, EGeometryScriptErrorType::InvalidInputs, LOCTEXT("AddTrianglesToMesh_InvalidList", "AddTrianglesToMesh: NewTrianglesList is empty"));
+		return TargetMesh;
+	}
+
+	NewIndicesList.Reset();
 	if (TargetMesh)
 	{
 		TargetMesh->EditMesh([&](FDynamicMesh3& EditMesh)
 		{
-			for (FIntVector Triangle : NewTriangles)
+			for (FIntVector Triangle : *NewTrianglesList.List)
 			{
 				int32 NewTriangleIndex = EditMesh.AppendTriangle((FIndex3i)Triangle, NewTriangleGroupID);
-				NewIndices.Add(NewTriangleIndex);
+				NewIndicesList.List->Add(NewTriangleIndex);
 			}
 
 		}, EDynamicMeshChangeType::GeneralEdit, EDynamicMeshAttributeChangeFlags::Unknown, bDeferChangeNotifications);
 
-		for (int32& NewTriangleIndex : NewIndices)
+		for (int32& NewTriangleIndex : *NewIndicesList.List)
 		{
 			if (NewTriangleIndex < 0)
 			{
@@ -256,16 +272,21 @@ UDynamicMesh* UGeometryScriptLibrary_MeshBasicEditFunctions::DeleteTriangleFromM
 
 UDynamicMesh* UGeometryScriptLibrary_MeshBasicEditFunctions::DeleteTrianglesFromMesh(
 	UDynamicMesh* TargetMesh,
-	const TArray<int32>& TriangleList,
+	FGeometryScriptIndexList TriangleList,
 	int& NumDeleted,
 	bool bDeferChangeNotifications)
 {
+	if (TriangleList.List.IsValid() == false || TriangleList.List->Num() == 0)
+	{
+		return TargetMesh;
+	}
+
 	NumDeleted = 0;
 	if (TargetMesh)
 	{
 		TargetMesh->EditMesh([&](FDynamicMesh3& EditMesh)
 		{
-			for (int32 TriangleID : TriangleList)
+			for (int32 TriangleID : *TriangleList.List)
 			{
 				EMeshResult Result = EditMesh.RemoveTriangle(TriangleID);
 				if (Result == EMeshResult::Ok)
@@ -378,12 +399,12 @@ UDynamicMesh* UGeometryScriptLibrary_MeshBasicEditFunctions::AppendMeshRepeated(
 UDynamicMesh* UGeometryScriptLibrary_MeshBasicEditFunctions::AppendBuffersToMesh(
 	UDynamicMesh* TargetMesh,
 	const FGeometryScriptSimpleMeshBuffers& Buffers,
-	TArray<int>& NewTriangleIndices,
+	FGeometryScriptIndexList& NewTriangleIndicesList,
 	int MaterialID,
 	bool bDeferChangeNotifications,
 	UGeometryScriptDebug* Debug)
 {
-	NewTriangleIndices.Reset();
+	NewTriangleIndicesList.Reset();
 	if (TargetMesh == nullptr)
 	{
 		UE::Geometry::AppendError(Debug, EGeometryScriptErrorType::InvalidInputs, LOCTEXT("AppendMeshRepeated_InvalidInput1", "AppendMeshRepeated: TargetMesh is Null"));
@@ -421,6 +442,7 @@ UDynamicMesh* UGeometryScriptLibrary_MeshBasicEditFunctions::AppendBuffersToMesh
 			return FIndex3i(VertexIDMap[Triangle.X], VertexIDMap[Triangle.Y], VertexIDMap[Triangle.Z]);
 		};
 
+		TArray<int32>& NewTriangleIndices = *NewTriangleIndicesList.List;
 		int32 NumTriangles = Buffers.Triangles.Num();
 		bool bHaveGroups = Buffers.TriGroupIDs.Num() == NumTriangles;
 		int32 ConstantGroupID = EditMesh.AllocateTriangleGroup();
diff --git a/Engine/Plugins/Experimental/GeometryScripting/Source/GeometryScriptingCore/Private/MeshDecompositionFunctions.cpp b/Engine/Plugins/Experimental/GeometryScripting/Source/GeometryScriptingCore/Private/MeshDecompositionFunctions.cpp
--- a/Engine/Plugins/Experimental/GeometryScripting/Source/GeometryScriptingCore/Private/MeshDecompositionFunctions.cpp
+++ b/Engine/Plugins/Experimental/GeometryScripting/Source/GeometryScriptingCore/Private/MeshDecompositionFunctions.cpp
@@ -138,7 +138,7 @@ UDynamicMesh* UGeometryScriptLibrary_MeshDecompositionFunctions::SplitMeshByMate
 UDynamicMesh* UGeometryScriptLibrary_MeshDecompositionFunctions::GetSubMeshFromMesh(
 	UDynamicMesh* TargetMesh,
 	UDynamicMesh* StoreToSubmesh, 
-	const TArray<int>& TriangleList,
+	FGeometryScriptIndexList TriangleList,
 	UDynamicMesh*& StoreToSubmeshOut, 
 	UGeometryScriptDebug* Debug)
 {
@@ -152,6 +152,11 @@ UDynamicMesh* UGeometryScriptLibrary_MeshDecompositionFunctions::GetSubMeshFromM
 		UE::Geometry::AppendError(Debug, EGeometryScriptErrorType::InvalidInputs, LOCTEXT("GetSubMeshFromMesh_InvalidInput2", "GetSubMeshFromMesh: Submesh is Null"));
 		return TargetMesh;
 	}
+	if (TriangleList.List.IsValid() == false || TriangleList.List->Num() == 0)
+	{
+		UE::Geometry::AppendError(Debug, EGeometryScriptErrorType::InvalidInputs, LOCTEXT("SetMeshPerVertexColors_InvalidList", "GetSubMeshFromMesh: TriangleList is empty"));
+		return TargetMesh;
+	}
 
 	FDynamicMesh3 Submesh;
 	TargetMesh->ProcessMesh([&](const FDynamicMesh3& EditMesh)
@@ -165,7 +170,7 @@ UDynamicMesh* UGeometryScriptLibrary_MeshDecompositionFunctions::GetSubMeshFromM
 		FMeshIndexMappings Mappings;
 		FDynamicMeshEditResult EditResult;
 		FDynamicMeshEditor Editor(&Submesh);
-		Editor.AppendTriangles(&EditMesh, TriangleList, Mappings, EditResult);
+		Editor.AppendTriangles(&EditMesh, *TriangleList.List, Mappings, EditResult);
 	});
 
 	StoreToSubmesh->SetMesh(MoveTemp(Submesh));
diff --git a/Engine/Plugins/Experimental/GeometryScripting/Source/GeometryScriptingCore/Private/MeshMaterialFunctions.cpp b/Engine/Plugins/Experimental/GeometryScripting/Source/GeometryScriptingCore/Private/MeshMaterialFunctions.cpp
--- a/Engine/Plugins/Experimental/GeometryScripting/Source/GeometryScriptingCore/Private/MeshMaterialFunctions.cpp
+++ b/Engine/Plugins/Experimental/GeometryScripting/Source/GeometryScriptingCore/Private/MeshMaterialFunctions.cpp
@@ -102,6 +102,7 @@ UDynamicMesh* UGeometryScriptLibrary_MeshMaterialFunctions::EnableMaterialIDs(
 
 UDynamicMesh* UGeometryScriptLibrary_MeshMaterialFunctions::ClearMaterialIDs(
 	UDynamicMesh* TargetMesh,
+	int ClearValue,
 	UGeometryScriptDebug* Debug)
 {
 	if (TargetMesh == nullptr)
@@ -110,13 +111,15 @@ UDynamicMesh* UGeometryScriptLibrary_MeshMaterialFunctions::ClearMaterialIDs(
 		return TargetMesh;
 	}
 
+	ClearValue = FMath::Max(0, ClearValue);
+
 	bool bHasMaterialIDs;
 	SimpleMeshMaterialEdit(TargetMesh, true, bHasMaterialIDs, 
-		[](FDynamicMesh3& Mesh, FDynamicMeshMaterialAttribute& MaterialIDs) 
+		[ClearValue](FDynamicMesh3& Mesh, FDynamicMeshMaterialAttribute& MaterialIDs) 
 	{
 		for (int32 TriangleID : Mesh.TriangleIndicesItr())
 		{
-			MaterialIDs.SetValue(TriangleID, 0);
+			MaterialIDs.SetValue(TriangleID, ClearValue);
 		}
 	});
 
@@ -155,6 +158,7 @@ UDynamicMesh* UGeometryScriptLibrary_MeshMaterialFunctions::RemapMaterialIDs(
 }
 
 
+
 int UGeometryScriptLibrary_MeshMaterialFunctions::GetMaxMaterialID( UDynamicMesh* TargetMesh, bool& bHasMaterialIDs )
 {
 	return SimpleMeshMaterialQuery<int32>(TargetMesh, bHasMaterialIDs, 0, [&](const FDynamicMesh3& Mesh, const FDynamicMeshMaterialAttribute& MaterialIDs) {
@@ -182,9 +186,11 @@ int32 UGeometryScriptLibrary_MeshMaterialFunctions::GetTriangleMaterialID(
 }
 
 
-UDynamicMesh* UGeometryScriptLibrary_MeshMaterialFunctions::GetAllTriangleMaterialIDs(UDynamicMesh* TargetMesh, UPARAM(ref) TArray<int>& MaterialIDs, bool& bHasMaterialIDs)
+UDynamicMesh* UGeometryScriptLibrary_MeshMaterialFunctions::GetAllTriangleMaterialIDs(UDynamicMesh* TargetMesh, FGeometryScriptIndexList& MaterialIDList, bool& bHasMaterialIDs)
 {
-	MaterialIDs.Reset();
+	MaterialIDList.Reset();
+	MaterialIDList.IndexType = EGeometryScriptIndexType::MaterialID;
+	TArray<int32>& MaterialIDs = *MaterialIDList.List;
 	bHasMaterialIDs = false;
 	SimpleMeshMaterialQuery<int32>(TargetMesh, bHasMaterialIDs, 0, [&](const FDynamicMesh3& Mesh, const FDynamicMeshMaterialAttribute& MaterialIDAttrib) {
 		int32 MaxTriangleID = Mesh.MaxTriangleID();
@@ -230,21 +236,27 @@ UDynamicMesh* UGeometryScriptLibrary_MeshMaterialFunctions::SetTriangleMaterialI
 
 UDynamicMesh* UGeometryScriptLibrary_MeshMaterialFunctions::SetAllTriangleMaterialIDs(
 	UDynamicMesh* TargetMesh,
-	const TArray<int32>& TriangleMaterialIDs,
+	FGeometryScriptIndexList& TriangleMaterialIDList,
 	bool bDeferChangeNotifications,
 	UGeometryScriptDebug* Debug)
 {
 	if (TargetMesh == nullptr)
 	{
-		UE::Geometry::AppendError(Debug, EGeometryScriptErrorType::InvalidInputs, LOCTEXT("SetAllTriangleMaterialIDs_InvalidInput", "SetAllTriangleMaterialIDs: TargetMesh is Null"));
+		UE::Geometry::AppendError(Debug, EGeometryScriptErrorType::InvalidInputs, LOCTEXT("SetAllTriangleMaterialIDs_InvalidMesh", "SetAllTriangleMaterialIDs: TargetMesh is Null"));
+		return TargetMesh;
+	}
+	if (TriangleMaterialIDList.List.IsValid() == false || TriangleMaterialIDList.List->Num() == 0)
+	{
+		UE::Geometry::AppendError(Debug, EGeometryScriptErrorType::InvalidInputs, LOCTEXT("SetAllTriangleMaterialIDs_InvalidList", "SetAllTriangleMaterialIDs: List is empty"));
 		return TargetMesh;
 	}
 
 	TargetMesh->EditMesh([&](FDynamicMesh3& EditMesh) 
 	{
+		const TArray<int32>& TriangleMaterialIDs = *TriangleMaterialIDList.List;
 		if (TriangleMaterialIDs.Num() < EditMesh.MaxTriangleID())
 		{
-			UE::Geometry::AppendError(Debug, EGeometryScriptErrorType::InvalidInputs, LOCTEXT("SetAllTriangleMaterialIDs_IncorrectCount", "SetAllTriangleMaterialIDs: size of provided TriangleMaterialIDs is smaller than MaxTriangleID of Mesh"));
+			UE::Geometry::AppendError(Debug, EGeometryScriptErrorType::InvalidInputs, LOCTEXT("SetAllTriangleMaterialIDs_IncorrectCount", "SetAllTriangleMaterialIDs: size of provided TriangleMaterialIDList is smaller than MaxTriangleID of Mesh"));
 		}
 		else
 		{
diff --git a/Engine/Plugins/Experimental/GeometryScripting/Source/GeometryScriptingCore/Private/MeshQueryFunctions.cpp b/Engine/Plugins/Experimental/GeometryScripting/Source/GeometryScriptingCore/Private/MeshQueryFunctions.cpp
--- a/Engine/Plugins/Experimental/GeometryScripting/Source/GeometryScriptingCore/Private/MeshQueryFunctions.cpp
+++ b/Engine/Plugins/Experimental/GeometryScripting/Source/GeometryScriptingCore/Private/MeshQueryFunctions.cpp
@@ -142,6 +142,25 @@ bool UGeometryScriptLibrary_MeshQueryFunctions::IsValidTriangleID( UDynamicMesh*
 	return SimpleMeshQuery<bool>(TargetMesh, false, [&](const FDynamicMesh3& Mesh) { return Mesh.IsTriangle(TriangleID); });
 }
 
+UDynamicMesh* UGeometryScriptLibrary_MeshQueryFunctions::GetAllTriangleIDs(UDynamicMesh* TargetMesh, FGeometryScriptIndexList& TriangleIDList, bool& bHasTriangleIDGaps)
+{
+	TriangleIDList.Reset();
+	bHasTriangleIDGaps = false;
+	if (TargetMesh)
+	{
+		TargetMesh->ProcessMesh([&](const FDynamicMesh3& ReadMesh)
+		{
+			TriangleIDList.List->Reserve(ReadMesh.TriangleCount());
+			for (int32 tid : ReadMesh.TriangleIndicesItr())
+			{
+				TriangleIDList.List->Add(tid);
+			}
+			bHasTriangleIDGaps = ! ReadMesh.IsCompactT();
+		});
+	}
+	return TargetMesh;
+}
+
 FIntVector UGeometryScriptLibrary_MeshQueryFunctions::GetTriangleIndices(UDynamicMesh* TargetMesh, int32 TriangleID, bool& bIsValidTriangle)
 {
 	return SimpleMeshQuery<FIntVector>(TargetMesh, FIntVector::NoneValue, [&](const FDynamicMesh3& Mesh) {
@@ -150,6 +169,40 @@ FIntVector UGeometryScriptLibrary_MeshQueryFunctions::GetTriangleIndices(UDynami
 	});
 }
 
+UDynamicMesh* UGeometryScriptLibrary_MeshQueryFunctions::GetAllTriangleIndices(UDynamicMesh* TargetMesh, FGeometryScriptTriangleList& TriangleList, bool bSkipGaps, bool& bHasTriangleIDGaps)
+{
+	TriangleList.Reset();
+	TArray<FIntVector>& Triangles = *TriangleList.List;
+	bHasTriangleIDGaps = false;
+	if (TargetMesh)
+	{
+		TargetMesh->ProcessMesh([&](const FDynamicMesh3& ReadMesh)
+		{
+			if (bSkipGaps)
+			{
+				Triangles.Reserve(ReadMesh.TriangleCount());
+				for (int32 tid : ReadMesh.TriangleIndicesItr())
+				{
+					Triangles.Add( (FIntVector)ReadMesh.GetTriangle(tid) );
+				}
+				bHasTriangleIDGaps = false;
+			}
+			else
+			{
+				Triangles.Init(FIntVector::NoneValue, ReadMesh.MaxTriangleID());
+				for (int32 tid : ReadMesh.TriangleIndicesItr())
+				{
+					Triangles[tid] = (FIntVector)ReadMesh.GetTriangle(tid);
+				}
+				bHasTriangleIDGaps = ! ReadMesh.IsCompactT();
+			}
+		});
+	}
+
+	return TargetMesh;
+}
+
+
 void UGeometryScriptLibrary_MeshQueryFunctions::GetTrianglePositions(UDynamicMesh* TargetMesh, int32 TriangleID, bool& bIsValidTriangle, FVector& Vertex1, FVector& Vertex2, FVector& Vertex3)
 {
 	bIsValidTriangle = SimpleMeshQuery<bool>(TargetMesh, false, [&](const FDynamicMesh3& Mesh) 
@@ -197,6 +250,25 @@ bool UGeometryScriptLibrary_MeshQueryFunctions::IsValidVertexID( UDynamicMesh* T
 	return SimpleMeshQuery<bool>(TargetMesh, false, [&](const FDynamicMesh3& Mesh) { return Mesh.IsVertex(VertexID); });
 }
 
+UDynamicMesh* UGeometryScriptLibrary_MeshQueryFunctions::GetAllVertexIDs(UDynamicMesh* TargetMesh, FGeometryScriptIndexList& VertexIDList, bool& bHasVertexIDGaps)
+{
+	VertexIDList.Reset();
+	bHasVertexIDGaps = false;
+	if (TargetMesh)
+	{
+		TargetMesh->ProcessMesh([&](const FDynamicMesh3& ReadMesh)
+		{
+			VertexIDList.List->Reserve(ReadMesh.VertexCount());
+			for (int32 vid : ReadMesh.VertexIndicesItr())
+			{
+				VertexIDList.List->Add(vid);
+			}
+			bHasVertexIDGaps = ! ReadMesh.IsCompactV();
+		});
+	}
+	return TargetMesh;
+}
+
 FVector UGeometryScriptLibrary_MeshQueryFunctions::GetVertexPosition(UDynamicMesh* TargetMesh, int32 VertexID, bool& bIsValidVertex)
 {
 	return SimpleMeshQuery<FVector>(TargetMesh, FVector::ZeroVector, [&](const FDynamicMesh3& Mesh) {
@@ -205,20 +277,33 @@ FVector UGeometryScriptLibrary_MeshQueryFunctions::GetVertexPosition(UDynamicMes
 	});
 }
 
-UDynamicMesh* UGeometryScriptLibrary_MeshQueryFunctions::GetAllVertexPositions(UDynamicMesh* TargetMesh, TArray<FVector>& VertexPositions, bool& bHasVertexIDGaps)
+UDynamicMesh* UGeometryScriptLibrary_MeshQueryFunctions::GetAllVertexPositions(UDynamicMesh* TargetMesh, FGeometryScriptVectorList& PositionList, bool bSkipGaps, bool& bHasVertexIDGaps)
 {
-	VertexPositions.Reset();
+	PositionList.Reset();
+	TArray<FVector>& VertexPositions = *PositionList.List;
 	bHasVertexIDGaps = false;
 	if (TargetMesh)
 	{
 		TargetMesh->ProcessMesh([&](const FDynamicMesh3& ReadMesh)
 		{
-			VertexPositions.Init(FVector::ZeroVector, ReadMesh.MaxVertexID());
-			for (int32 vid : ReadMesh.VertexIndicesItr())
+			if (bSkipGaps)
 			{
-				VertexPositions[vid] = (FVector)ReadMesh.GetVertex(vid);
+				VertexPositions.Reserve(ReadMesh.VertexCount());
+				for (int32 vid : ReadMesh.VertexIndicesItr())
+				{
+					VertexPositions.Add( (FVector)ReadMesh.GetVertex(vid) );
+				}
+				bHasVertexIDGaps = false;
+			}
+			else
+			{
+				VertexPositions.Init(FVector::ZeroVector, ReadMesh.MaxVertexID());
+				for (int32 vid : ReadMesh.VertexIndicesItr())
+				{
+					VertexPositions[vid] = (FVector)ReadMesh.GetVertex(vid);
+				}
+				bHasVertexIDGaps = ! ReadMesh.IsCompactV();
 			}
-			bHasVertexIDGaps = ! ReadMesh.IsCompactV();
 		});
 	}
 	
diff --git a/Engine/Plugins/Experimental/GeometryScripting/Source/GeometryScriptingCore/Private/MeshVertexColorFunctions.cpp b/Engine/Plugins/Experimental/GeometryScripting/Source/GeometryScriptingCore/Private/MeshVertexColorFunctions.cpp
--- a/Engine/Plugins/Experimental/GeometryScripting/Source/GeometryScriptingCore/Private/MeshVertexColorFunctions.cpp
+++ b/Engine/Plugins/Experimental/GeometryScripting/Source/GeometryScriptingCore/Private/MeshVertexColorFunctions.cpp
@@ -84,20 +84,26 @@ UDynamicMesh* UGeometryScriptLibrary_MeshVertexColorFunctions::SetMeshConstantVe
 
 UDynamicMesh* UGeometryScriptLibrary_MeshVertexColorFunctions::SetMeshPerVertexColors(
 	UDynamicMesh* TargetMesh,
-	const TArray<FLinearColor>& VertexColors,
+	FGeometryScriptColorList VertexColorList,
 	UGeometryScriptDebug* Debug)
 {
 	if (TargetMesh == nullptr)
 	{
-		UE::Geometry::AppendError(Debug, EGeometryScriptErrorType::InvalidInputs, LOCTEXT("SetMeshPerVertexColors_InvalidInput", "SetMeshPerVertexColors: TargetMesh is Null"));
+		UE::Geometry::AppendError(Debug, EGeometryScriptErrorType::InvalidInputs, LOCTEXT("SetMeshPerVertexColors_InvalidMesh", "SetMeshPerVertexColors: TargetMesh is Null"));
+		return TargetMesh;
+	}
+	if (VertexColorList.List.IsValid() == false || VertexColorList.List->Num() == 0)
+	{
+		UE::Geometry::AppendError(Debug, EGeometryScriptErrorType::InvalidInputs, LOCTEXT("SetMeshPerVertexColors_InvalidList", "SetMeshPerVertexColors: List is empty"));
 		return TargetMesh;
 	}
 
 	TargetMesh->EditMesh([&](FDynamicMesh3& EditMesh) 
 	{
+		const TArray<FLinearColor>& VertexColors = *VertexColorList.List;
 		if (VertexColors.Num() < EditMesh.MaxVertexID())
 		{
-			UE::Geometry::AppendError(Debug, EGeometryScriptErrorType::InvalidInputs, LOCTEXT("SetMeshPerVertexColors_IncorrectCount", "SetMeshPerVertexColors: size of provided VertexColors is smaller than MaxVertexID of Mesh"));
+			UE::Geometry::AppendError(Debug, EGeometryScriptErrorType::InvalidInputs, LOCTEXT("SetMeshPerVertexColors_IncorrectCount", "SetMeshPerVertexColors: size of provided VertexColorList is smaller than MaxVertexID of Mesh"));
 		}
 		else
 		{
diff --git a/Engine/Plugins/Experimental/GeometryScripting/Source/GeometryScriptingCore/Public/GeometryScript/GeometryScriptTypes.h b/Engine/Plugins/Experimental/GeometryScripting/Source/GeometryScriptingCore/Public/GeometryScript/GeometryScriptTypes.h
--- a/Engine/Plugins/Experimental/GeometryScripting/Source/GeometryScriptingCore/Public/GeometryScript/GeometryScriptTypes.h
+++ b/Engine/Plugins/Experimental/GeometryScripting/Source/GeometryScriptingCore/Public/GeometryScript/GeometryScriptTypes.h
@@ -173,6 +173,118 @@ public:
 };
 
 
+//
+// List Types
+//
+
+UENUM(BlueprintType)
+enum class EGeometryScriptIndexType : uint8
+{
+	Unknown,
+	Triangle,
+	Vertex,
+	MaterialID,
+	PolygroupID
+};
+
+USTRUCT(BlueprintType, meta = (DisplayName = "Index List"))
+struct GEOMETRYSCRIPTINGCORE_API FGeometryScriptIndexList
+{
+	GENERATED_BODY()
+
+	UPROPERTY(BlueprintReadWrite, Category = IndexList)
+	EGeometryScriptIndexType IndexType = EGeometryScriptIndexType::Unknown;
+public:
+	TSharedPtr<TArray<int>> List;
+
+	void Reset()
+	{
+		if (List.IsValid() == false)
+		{
+			List = MakeShared<TArray<int>>();
+		}
+		List->Reset();
+	}
+};
+
+
+USTRUCT(BlueprintType, meta = (DisplayName = "Triangle List"))
+struct GEOMETRYSCRIPTINGCORE_API FGeometryScriptTriangleList
+{
+	GENERATED_BODY()
+public:
+	TSharedPtr<TArray<FIntVector>> List;
+
+	void Reset()
+	{
+		if (List.IsValid() == false)
+		{
+			List = MakeShared<TArray<FIntVector>>();
+		}
+		List->Reset();
+	}
+};
+
+
+
+USTRUCT(BlueprintType, meta = (DisplayName = "Vector List"))
+struct GEOMETRYSCRIPTINGCORE_API FGeometryScriptVectorList
+{
+	GENERATED_BODY()
+public:
+	TSharedPtr<TArray<FVector>> List;
+
+	void Reset()
+	{
+		if (List.IsValid() == false)
+		{
+			List = MakeShared<TArray<FVector>>();
+		}
+		List->Reset();
+	}
+};
+
+
+
+USTRUCT(BlueprintType, meta = (DisplayName = "UV List"))
+struct GEOMETRYSCRIPTINGCORE_API FGeometryScriptUVList
+{
+	GENERATED_BODY()
+public:
+	TSharedPtr<TArray<FVector2D>> List;
+
+	void Reset()
+	{
+		if (List.IsValid() == false)
+		{
+			List = MakeShared<TArray<FVector2D>>();
+		}
+		List->Reset();
+	}
+};
+
+
+
+USTRUCT(BlueprintType, meta = (DisplayName = "Color List"))
+struct GEOMETRYSCRIPTINGCORE_API FGeometryScriptColorList
+{
+	GENERATED_BODY()
+public:
+	TSharedPtr<TArray<FLinearColor>> List;
+
+	void Reset()
+	{
+		if (List.IsValid() == false)
+		{
+			List = MakeShared<TArray<FLinearColor>>();
+		}
+		List->Reset();
+	}
+};
+
+
+
+
 //
 // Spatial data structures
 //
diff --git a/Engine/Plugins/Experimental/GeometryScripting/Source/GeometryScriptingCore/Public/GeometryScript/ListUtilityFunctions.h b/Engine/Plugins/Experimental/GeometryScripting/Source/GeometryScriptingCore/Public/GeometryScript/ListUtilityFunctions.h
new file mode 100644
--- /dev/null
+++ b/Engine/Plugins/Experimental/GeometryScripting/Source/GeometryScriptingCore/Public/GeometryScript/ListUtilityFunctions.h
@@ -0,0 +1,101 @@
+// Copyright Epic Games, Inc. All Rights Reserved.
+
+#pragma once
+
+#include "CoreMinimal.h"
+#include "Kismet/BlueprintFunctionLibrary.h"
+#include "GeometryScript/GeometryScriptTypes.h"
+#include "ListUtilityFunctions.generated.h"
+
+class UDynamicMesh;
+
+
+UCLASS(meta = (ScriptName = "GeometryScript_List"))
+class GEOMETRYSCRIPTINGCORE_API UGeometryScriptLibrary_ListUtilityFunctions : public UBlueprintFunctionLibrary
+{
+	GENERATED_BODY()
+public:
+
+	UFUNCTION(BlueprintPure, Category = "GeometryScript|ListUtils", meta=(ScriptMethod))
+	static int GetIndexListLength(FGeometryScriptIndexList IndexList);
+
+	UFUNCTION(BlueprintPure, Category = "GeometryScript|ListUtils", meta=(ScriptMethod))
+	static int GetIndexListLastIndex(FGeometryScriptIndexList IndexList);
+
+	UFUNCTION(BlueprintPure, Category = "GeometryScript|ListUtils", meta=(ScriptMethod))
+	static int GetIndexListItem(FGeometryScriptIndexList IndexList, int Index, bool& bIsValidIndex);
+
+	UFUNCTION(BlueprintCallable, Category = "GeometryScript|ListUtils", meta=(ScriptMethod))
+	static void ConvertIndexListToArray(FGeometryScriptIndexList IndexList, TArray<int>& IndexArray);
+
+	UFUNCTION(BlueprintCallable, Category = "GeometryScript|ListUtils")
+	static void ConvertArrayToIndexList(const TArray<int>& IndexArray, FGeometryScriptIndexList& IndexList);
+
+
+
+	UFUNCTION(BlueprintPure, Category = "GeometryScript|ListUtils", meta=(ScriptMethod))
+	static int GetTriangleListLength(FGeometryScriptTriangleList TriangleList);
+
+	UFUNCTION(BlueprintPure, Category = "GeometryScript|ListUtils", meta=(ScriptMethod))
+	static int GetTriangleListLastTriangle(FGeometryScriptTriangleList TriangleList);
+
+	UFUNCTION(BlueprintPure, Category = "GeometryScript|ListUtils", meta=(ScriptMethod))
+	static FIntVector GetTriangleListItem(FGeometryScriptTriangleList TriangleList, int Triangle, bool& bIsValidTriangle);
+
+	UFUNCTION(BlueprintCallable, Category = "GeometryScript|ListUtils", meta=(ScriptMethod))
+	static void ConvertTriangleListToArray(FGeometryScriptTriangleList TriangleList, TArray<FIntVector>& TriangleArray);
+
+	UFUNCTION(BlueprintCallable, Category = "GeometryScript|ListUtils")
+	static void ConvertArrayToTriangleList(const TArray<FIntVector>& TriangleArray, FGeometryScriptTriangleList& TriangleList);
+
+
+
+	UFUNCTION(BlueprintPure, Category = "GeometryScript|ListUtils", meta=(ScriptMethod))
+	static int GetVectorListLength(FGeometryScriptVectorList VectorList);
+
+	UFUNCTION(BlueprintPure, Category = "GeometryScript|ListUtils", meta=(ScriptMethod))
+	static int GetVectorListLastIndex(FGeometryScriptVectorList VectorList);
+
+	UFUNCTION(BlueprintPure, Category = "GeometryScript|ListUtils", meta=(ScriptMethod))
+	static FVector GetVectorListItem(FGeometryScriptVectorList VectorList, int Index, bool& bIsValidIndex);
+
+	UFUNCTION(BlueprintCallable, Category = "GeometryScript|ListUtils", meta=(ScriptMethod))
+	static void ConvertVectorListToArray(FGeometryScriptVectorList VectorList, TArray<FVector>& VectorArray);
+
+	UFUNCTION(BlueprintCallable, Category = "GeometryScript|ListUtils")
+	static void ConvertArrayToVectorList(const TArray<FVector>& VectorArray, FGeometryScriptVectorList& VectorList);
+
+
+
+	UFUNCTION(BlueprintPure, Category = "GeometryScript|ListUtils", meta=(ScriptMethod))
+	static int GetUVListLength(FGeometryScriptUVList UVList);
+
+	UFUNCTION(BlueprintPure, Category = "GeometryScript|ListUtils", meta=(ScriptMethod))
+	static int GetUVListLastIndex(FGeometryScriptUVList UVList);
+
+	UFUNCTION(BlueprintPure, Category = "GeometryScript|ListUtils", meta=(ScriptMethod))
+	static FVector2D GetUVListItem(FGeometryScriptUVList UVList, int Index, bool& bIsValidIndex);
+
+	UFUNCTION(BlueprintCallable, Category = "GeometryScript|ListUtils", meta=(ScriptMethod))
+	static void ConvertUVListToArray(FGeometryScriptUVList UVList, TArray<FVector2D>& UVArray);
+
+	UFUNCTION(BlueprintCallable, Category = "GeometryScript|ListUtils")
+	static void ConvertArrayToUVList(const TArray<FVector2D>& UVArray, FGeometryScriptUVList& UVList);
+
+
+
+	UFUNCTION(BlueprintPure, Category = "GeometryScript|ListUtils", meta=(ScriptMethod))
+	static int GetColorListLength(FGeometryScriptColorList ColorList);
+
+	UFUNCTION(BlueprintPure, Category = "GeometryScript|ListUtils", meta=(ScriptMethod))
+	static int GetColorListLastIndex(FGeometryScriptColorList ColorList);
+
+	UFUNCTION(BlueprintPure, Category = "GeometryScript|ListUtils", meta=(ScriptMethod))
+	static FLinearColor GetColorListItem(FGeometryScriptColorList ColorList, int Index, bool& bIsValidIndex);
+
+	UFUNCTION(BlueprintCallable, Category = "GeometryScript|ListUtils", meta=(ScriptMethod))
+	static void ConvertColorListToArray(FGeometryScriptColorList ColorList, TArray<FLinearColor>& ColorArray);
+
+	UFUNCTION(BlueprintCallable, Category = "GeometryScript|ListUtils")
+	static void ConvertArrayToColorList(const TArray<FLinearColor>& ColorArray, FGeometryScriptColorList& ColorList);
+};
\ No newline at end of file
diff --git a/Engine/Plugins/Experimental/GeometryScripting/Source/GeometryScriptingCore/Public/GeometryScript/MeshBasicEditFunctions.h b/Engine/Plugins/Experimental/GeometryScripting/Source/GeometryScriptingCore/Public/GeometryScript/MeshBasicEditFunctions.h
--- a/Engine/Plugins/Experimental/GeometryScripting/Source/GeometryScriptingCore/Public/GeometryScript/MeshBasicEditFunctions.h
+++ b/Engine/Plugins/Experimental/GeometryScripting/Source/GeometryScriptingCore/Public/GeometryScript/MeshBasicEditFunctions.h
@@ -86,8 +86,8 @@ public:
 	static UPARAM(DisplayName = "Target Mesh") UDynamicMesh* 
 	AddVerticesToMesh( 
 		UDynamicMesh* TargetMesh, 
-		const TArray<FVector>& NewPositions, 
-		TArray<int>& NewIndices,
+		FGeometryScriptVectorList NewPositionsList, 
+		FGeometryScriptIndexList& NewIndicesList,
 		bool bDeferChangeNotifications = false );
 
 	UFUNCTION(BlueprintCallable, Category = "GeometryScript|MeshEdits", meta=(ScriptMethod))
@@ -102,7 +102,7 @@ public:
 	static UPARAM(DisplayName = "Target Mesh") UDynamicMesh* 
 	DeleteVerticesFromMesh( 
 		UDynamicMesh* TargetMesh, 
-		const TArray<int32>& VertexList,
+		FGeometryScriptIndexList VertexList,
 		int& NumDeleted,
 		bool bDeferChangeNotifications = false );
 
@@ -122,8 +122,8 @@ public:
 	static UPARAM(DisplayName = "Target Mesh") UDynamicMesh* 
 	AddTrianglesToMesh( 
 		UDynamicMesh* TargetMesh, 
-		const TArray<FIntVector>& NewTriangles,
-		TArray<int>& NewIndices,
+		FGeometryScriptTriangleList NewTrianglesList,
+		FGeometryScriptIndexList& NewIndicesList,
 		int NewTriangleGroupID = 0,
 		bool bDeferChangeNotifications = false,
 		UGeometryScriptDebug* Debug = nullptr );
@@ -140,7 +140,7 @@ public:
 	static UPARAM(DisplayName = "Target Mesh") UDynamicMesh* 
 	DeleteTrianglesFromMesh( 
 		UDynamicMesh* TargetMesh, 
-		const TArray<int32>& TriangleList,
+		FGeometryScriptIndexList TriangleList,
 		int& NumDeleted,
 		bool bDeferChangeNotifications = false );
 
@@ -172,7 +172,7 @@ public:
 	AppendBuffersToMesh( 
 		UDynamicMesh* TargetMesh, 
 		const FGeometryScriptSimpleMeshBuffers& Buffers,
-		TArray<int>& NewTriangleIndices,
+		FGeometryScriptIndexList& NewTriangleIndicesList,
 		int MaterialID = 0,
 		bool bDeferChangeNotifications = false,
 		UGeometryScriptDebug* Debug = nullptr );
diff --git a/Engine/Plugins/Experimental/GeometryScripting/Source/GeometryScriptingCore/Public/GeometryScript/MeshDecompositionFunctions.h b/Engine/Plugins/Experimental/GeometryScripting/Source/GeometryScriptingCore/Public/GeometryScript/MeshDecompositionFunctions.h
--- a/Engine/Plugins/Experimental/GeometryScripting/Source/GeometryScriptingCore/Public/GeometryScript/MeshDecompositionFunctions.h
+++ b/Engine/Plugins/Experimental/GeometryScripting/Source/GeometryScriptingCore/Public/GeometryScript/MeshDecompositionFunctions.h
@@ -38,7 +38,7 @@ public:
 	GetSubMeshFromMesh(  
 		UDynamicMesh* TargetMesh, 
 		UPARAM(DisplayName = "Copy To Submesh", ref) UDynamicMesh* StoreToSubmesh, 
-		const TArray<int>& TriangleList,
+		FGeometryScriptIndexList TriangleList,
 		UPARAM(DisplayName = "Copy To Submesh") UDynamicMesh*& StoreToSubmeshOut, 
 		UGeometryScriptDebug* Debug = nullptr);
 
diff --git a/Engine/Plugins/Experimental/GeometryScripting/Source/GeometryScriptingCore/Public/GeometryScript/MeshMaterialFunctions.h b/Engine/Plugins/Experimental/GeometryScripting/Source/GeometryScriptingCore/Public/GeometryScript/MeshMaterialFunctions.h
--- a/Engine/Plugins/Experimental/GeometryScripting/Source/GeometryScriptingCore/Public/GeometryScript/MeshMaterialFunctions.h
+++ b/Engine/Plugins/Experimental/GeometryScripting/Source/GeometryScriptingCore/Public/GeometryScript/MeshMaterialFunctions.h
@@ -25,7 +25,7 @@ public:
 
 	UFUNCTION(BlueprintCallable, Category = "GeometryScript|Materials", meta=(ScriptMethod))
 	static UPARAM(DisplayName = "Target Mesh") UDynamicMesh* 
-	ClearMaterialIDs( UDynamicMesh* TargetMesh, UGeometryScriptDebug* Debug = nullptr );
+	ClearMaterialIDs( UDynamicMesh* TargetMesh, int ClearValue = 0, UGeometryScriptDebug* Debug = nullptr );
 
 	UFUNCTION(BlueprintCallable, Category = "GeometryScript|Materials", meta=(ScriptMethod))
 	static UPARAM(DisplayName = "Target Mesh") UDynamicMesh* 
@@ -35,14 +35,13 @@ public:
 		int ToMaterialID,
 		UGeometryScriptDebug* Debug = nullptr );
 
-
 	UFUNCTION(BlueprintPure, Category = "GeometryScript|Materials", meta=(ScriptMethod))
 	static UPARAM(DisplayName = "Material ID") int32
 	GetTriangleMaterialID( UDynamicMesh* TargetMesh, int TriangleID, bool& bIsValidTriangle );
 
 	UFUNCTION(BlueprintPure, Category = "GeometryScript|Materials", meta=(ScriptMethod))
 	static UPARAM(DisplayName = "Target Mesh") UDynamicMesh* 
-	GetAllTriangleMaterialIDs( UDynamicMesh* TargetMesh, UPARAM(ref) TArray<int>& MaterialIDs, bool& bHasMaterialIDs );
+	GetAllTriangleMaterialIDs( UDynamicMesh* TargetMesh, FGeometryScriptIndexList& MaterialIDList, bool& bHasMaterialIDs );
 
 
 	UFUNCTION(BlueprintCallable, Category = "GeometryScript|Materials", meta=(ScriptMethod))
@@ -58,7 +57,7 @@ public:
 	static UPARAM(DisplayName = "Target Mesh") UDynamicMesh*
 	SetAllTriangleMaterialIDs(
 		UDynamicMesh* TargetMesh,
-		UPARAM(ref) const TArray<int32>& TriangleMaterialIDs,
+		FGeometryScriptIndexList& TriangleMaterialIDList,
 		bool bDeferChangeNotifications = false,
 		UGeometryScriptDebug* Debug = nullptr);
 
diff --git a/Engine/Plugins/Experimental/GeometryScripting/Source/GeometryScriptingCore/Public/GeometryScript/MeshQueryFunctions.h b/Engine/Plugins/Experimental/GeometryScripting/Source/GeometryScriptingCore/Public/GeometryScript/MeshQueryFunctions.h
--- a/Engine/Plugins/Experimental/GeometryScripting/Source/GeometryScriptingCore/Public/GeometryScript/MeshQueryFunctions.h
+++ b/Engine/Plugins/Experimental/GeometryScripting/Source/GeometryScriptingCore/Public/GeometryScript/MeshQueryFunctions.h
@@ -53,9 +53,18 @@ public:
 	UFUNCTION(BlueprintPure, Category = "GeometryScript|MeshQueries", meta=(ScriptMethod))
 	static bool IsValidTriangleID( UDynamicMesh* TargetMesh, int32 TriangleID );
 
+	UFUNCTION(BlueprintCallable, Category = "GeometryScript|MeshQueries", meta=(ScriptMethod))
+	static UPARAM(DisplayName = "Target Mesh") UDynamicMesh* 
+	GetAllTriangleIDs( UDynamicMesh* TargetMesh, FGeometryScriptIndexList& TriangleIDList, bool& bHasTriangleIDGaps );
+
+
 	UFUNCTION(BlueprintPure, Category = "GeometryScript|MeshQueries", meta=(ScriptMethod))
 	static UPARAM(DisplayName = "Indices") FIntVector GetTriangleIndices( UDynamicMesh* TargetMesh, int32 TriangleID, bool& bIsValidTriangle );
 
+	UFUNCTION(BlueprintCallable, Category = "GeometryScript|MeshQueries", meta=(ScriptMethod))
+	static UPARAM(DisplayName = "Target Mesh") UDynamicMesh* 
+	GetAllTriangleIndices( UDynamicMesh* TargetMesh, FGeometryScriptTriangleList& TriangleList, bool bSkipGaps, bool& bHasTriangleIDGaps);
+
 	UFUNCTION(BlueprintPure, Category = "GeometryScript|MeshQueries", meta=(ScriptMethod))
 	static void GetTrianglePositions( UDynamicMesh* TargetMesh, int32 TriangleID, bool& bIsValidTriangle, FVector& Vertex1, FVector& Vertex2, FVector& Vertex3 );
 
@@ -64,7 +73,6 @@ public:
 
 
 
-
 	UFUNCTION(BlueprintPure, Category = "GeometryScript|MeshQueries", meta=(ScriptMethod))
 	static UPARAM(DisplayName = "Vertex Count") int32 GetVertexCount( UDynamicMesh* TargetMesh );
 
@@ -77,12 +85,16 @@ public:
 	UFUNCTION(BlueprintPure, Category = "GeometryScript|MeshQueries", meta=(ScriptMethod))
 	static bool IsValidVertexID( UDynamicMesh* TargetMesh, int32 VertexID );
 
+	UFUNCTION(BlueprintCallable, Category = "GeometryScript|MeshQueries", meta=(ScriptMethod))
+	static UPARAM(DisplayName = "Target Mesh") UDynamicMesh* 
+	GetAllVertexIDs( UDynamicMesh* TargetMesh, FGeometryScriptIndexList& VertexIDList, bool& bHasVertexIDGaps );
+
 	UFUNCTION(BlueprintPure, Category = "GeometryScript|MeshQueries", meta=(ScriptMethod))
 	static UPARAM(DisplayName = "Position") FVector GetVertexPosition( UDynamicMesh* TargetMesh, int32 VertexID, bool& bIsValidVertex );
 
-	UFUNCTION(BlueprintPure, Category = "GeometryScript|MeshQueries", meta=(ScriptMethod))
+	UFUNCTION(BlueprintCallable, Category = "GeometryScript|MeshQueries", meta=(ScriptMethod))
 	static UPARAM(DisplayName = "Target Mesh") UDynamicMesh* 
-	GetAllVertexPositions( UDynamicMesh* TargetMesh,  UPARAM(ref) TArray<FVector>& VertexPositions, bool& bHasVertexIDGaps );
+	GetAllVertexPositions( UDynamicMesh* TargetMesh, FGeometryScriptVectorList& PositionList, bool bSkipGaps, bool& bHasVertexIDGaps );
 
 
 
diff --git a/Engine/Plugins/Experimental/GeometryScripting/Source/GeometryScriptingCore/Public/GeometryScript/MeshVertexColorFunctions.h b/Engine/Plugins/Experimental/GeometryScripting/Source/GeometryScriptingCore/Public/GeometryScript/MeshVertexColorFunctions.h
--- a/Engine/Plugins/Experimental/GeometryScripting/Source/GeometryScriptingCore/Public/GeometryScript/MeshVertexColorFunctions.h
+++ b/Engine/Plugins/Experimental/GeometryScripting/Source/GeometryScriptingCore/Public/GeometryScript/MeshVertexColorFunctions.h
@@ -30,7 +30,7 @@ public:
 	static UPARAM(DisplayName = "Target Mesh") UDynamicMesh*
 	SetMeshPerVertexColors(
 		UDynamicMesh* TargetMesh,
-		UPARAM(ref) const TArray<FLinearColor>& VertexColors,
+		FGeometryScriptColorList VertexColorList,
 		UGeometryScriptDebug* Debug = nullptr);
 
 
