fix call to IsPendingKillOrUnreachable that was not replaced in mass conversion
#rnx

diff --git a/Engine/Plugins/Experimental/GeometryScripting/Source/GeometryScriptingEditor/Private/GeometryActors/EditorGeometryGenerationSubsystem.cpp b/Engine/Plugins/Experimental/GeometryScripting/Source/GeometryScriptingEditor/Private/GeometryActors/EditorGeometryGenerationSubsystem.cpp
--- a/Engine/Plugins/Experimental/GeometryScripting/Source/GeometryScriptingEditor/Private/GeometryActors/EditorGeometryGenerationSubsystem.cpp
+++ b/Engine/Plugins/Experimental/GeometryScripting/Source/GeometryScriptingEditor/Private/GeometryActors/EditorGeometryGenerationSubsystem.cpp
@@ -10,7 +10,7 @@ void UEditorGeometryGenerationManager::Tick(float DeltaTime)
 	// rebuild all valid generated actors, if necessary
 	for (AGeneratedDynamicMeshActor* Actor : ActiveGeneratedActors)
 	{
-		if (Actor->IsValidLowLevel() && !Actor->IsPendingKillOrUnreachable() &&  Actor->GetLevel() != nullptr)
+		if (Actor->IsValidLowLevel() && IsValid(Actor) && !Actor->IsUnreachable() && Actor->GetLevel() != nullptr)
 		{
 			Actor->ExecuteRebuildGeneratedMeshIfPending();
 		}
