GeometryFramework: add BP-exposed functions to GeneratedDynamicMeshActor to copy properties to/from a StaticMeshActor. This allows for preserving Actor settings like data layer/etc when swapping between Generated and Static meshes. Also can copy Material assignment.
#rnx

diff --git a/Engine/Plugins/Experimental/GeometryScripting/Source/GeometryScriptingEditor/Private/GeometryActors/GeneratedDynamicMeshActor.cpp b/Engine/Plugins/Experimental/GeometryScripting/Source/GeometryScriptingEditor/Private/GeometryActors/GeneratedDynamicMeshActor.cpp
--- a/Engine/Plugins/Experimental/GeometryScripting/Source/GeometryScriptingEditor/Private/GeometryActors/GeneratedDynamicMeshActor.cpp
+++ b/Engine/Plugins/Experimental/GeometryScripting/Source/GeometryScriptingEditor/Private/GeometryActors/GeneratedDynamicMeshActor.cpp
@@ -3,6 +3,9 @@
 #include "GeometryActors/GeneratedDynamicMeshActor.h"
 #include "GeometryActors/EditorGeometryGenerationSubsystem.h"
 
+#include "Editor/EditorEngine.h" // for CopyPropertiesForUnrelatedObjects
+#include "Engine/StaticMeshActor.h"
+
 #define LOCTEXT_NAMESPACE "AGeneratedDynamicMeshActor"
 
 
@@ -94,4 +97,53 @@ void AGeneratedDynamicMeshActor::ExecuteRebuildGeneratedMeshIfPending()
 }
 
 
+
+
+void AGeneratedDynamicMeshActor::CopyPropertiesToStaticMesh(AStaticMeshActor* StaticMeshActor, bool bCopyComponentMaterials)
+{
+	StaticMeshActor->Modify();
+	StaticMeshActor->UnregisterAllComponents();
+	UEditorEngine::CopyPropertiesForUnrelatedObjects(this, StaticMeshActor);
+
+	if (bCopyComponentMaterials)
+	{
+		if (UStaticMeshComponent* SMComponent = StaticMeshActor->GetStaticMeshComponent())
+		{
+			if (UDynamicMeshComponent* DMComponent = this->GetDynamicMeshComponent())
+			{
+				TArray<UMaterialInterface*> Materials = DMComponent->GetMaterials();
+				for (int32 k = 0; k < Materials.Num(); ++k)
+				{
+					SMComponent->SetMaterial(k, Materials[k]);
+				}
+			}
+		}
+	}
+
+	StaticMeshActor->ReregisterAllComponents();
+}
+
+
+void AGeneratedDynamicMeshActor::CopyPropertiesFromStaticMesh(AStaticMeshActor* StaticMeshActor, bool bCopyComponentMaterials)
+{
+	this->Modify();
+	this->UnregisterAllComponents();
+	UEditorEngine::CopyPropertiesForUnrelatedObjects(StaticMeshActor, this);
+
+	if (bCopyComponentMaterials)
+	{
+		if (UStaticMeshComponent* SMComponent = StaticMeshActor->GetStaticMeshComponent())
+		{
+			if (UDynamicMeshComponent* DMComponent = this->GetDynamicMeshComponent())
+			{
+				TArray<UMaterialInterface*> Materials = SMComponent->GetMaterials();
+				DMComponent->ConfigureMaterialSet(Materials);
+			}
+		}
+	}
+
+	this->ReregisterAllComponents();
+}
+
+
 #undef LOCTEXT_NAMESPACE
\ No newline at end of file
diff --git a/Engine/Plugins/Experimental/GeometryScripting/Source/GeometryScriptingEditor/Public/GeometryActors/GeneratedDynamicMeshActor.h b/Engine/Plugins/Experimental/GeometryScripting/Source/GeometryScriptingEditor/Public/GeometryActors/GeneratedDynamicMeshActor.h
--- a/Engine/Plugins/Experimental/GeometryScripting/Source/GeometryScriptingEditor/Public/GeometryActors/GeneratedDynamicMeshActor.h
+++ b/Engine/Plugins/Experimental/GeometryScripting/Source/GeometryScriptingEditor/Public/GeometryActors/GeneratedDynamicMeshActor.h
@@ -6,6 +6,8 @@
 
 #include "GeneratedDynamicMeshActor.generated.h"
 
+class AStaticMeshActor;
+
 
 /**
  * AGeneratedDynamicMeshActor is an Editor-only subclass of ADynamicMeshActor that provides 
@@ -49,6 +51,26 @@ public:
 	virtual void ExecuteRebuildGeneratedMeshIfPending();
 
 
+public:
+
+	/** 
+	 * Attempt to copy Actor Properties to a StaticMeshActor. Optionally copy DynamicMeshComponent material list to the StaticMeshComponent.
+	 * This function is useful when (eg) swapping from a DynamicMeshActor to a StaticMeshActor as it will allow
+	 * many configured Actor settings to be preserved (like assigned DataLayers, etc)
+	 */
+	UFUNCTION(BlueprintCallable, Category = DynamicMeshActor)
+	void CopyPropertiesToStaticMesh(AStaticMeshActor* StaticMeshActor, bool bCopyComponentMaterials = false);
+
+	/**
+	 * Attempt to copy Actor Properties from a StaticMeshActor. Optionally copy DynamicMeshComponent material list to the StaticMeshComponent.
+	 * This function is useful when (eg) swapping from a StaticMeshActor to a DynamicMeshActor as it will allow
+	 * many configured Actor settings to be preserved (like assigned DataLayers, etc) 
+	 */
+	UFUNCTION(BlueprintCallable, Category = DynamicMeshActor)
+	void CopyPropertiesFromStaticMesh(AStaticMeshActor* StaticMeshActor, bool bCopyComponentMaterials = false);
+
+
+
 public:
 
 	// AActor overrides
