Mark the ApplyRecursivePNTriangles as deprecated and add the replacement function ApplyPNTriangles.
Make sure ApplyRecursivePNTriangles produces the same result as before the changes to the PNTriangles tessellation were made.
Add a geometry script function for the UniformTessellate operator.

#rnx

diff --git a/Engine/Plugins/Experimental/GeometryScripting/Source/GeometryScriptingCore/Private/MeshSubdivideFunctions.cpp b/Engine/Plugins/Experimental/GeometryScripting/Source/GeometryScriptingCore/Private/MeshSubdivideFunctions.cpp
--- a/Engine/Plugins/Experimental/GeometryScripting/Source/GeometryScriptingCore/Private/MeshSubdivideFunctions.cpp
+++ b/Engine/Plugins/Experimental/GeometryScripting/Source/GeometryScriptingCore/Private/MeshSubdivideFunctions.cpp
@@ -6,7 +6,9 @@
 #include "DynamicMesh/DynamicMeshAttributeSet.h"
 #include "DynamicMesh/MeshNormals.h"
 #include "Operations/PNTriangles.h"
+#include "Operations/UniformTesselate.h"
 #include "UDynamicMesh.h"
+#include "Math/UnrealMathUtility.h"
 
 using namespace UE::Geometry;
 
@@ -32,7 +34,11 @@ UDynamicMesh* UGeometryScriptLibrary_MeshSubdivideFunctions::ApplyRecursivePNTes
 	TargetMesh->EditMesh([&](FDynamicMesh3& EditMesh) 
 	{
 		FPNTriangles Tessellator(&EditMesh);
-		Tessellator.TesselationLevel = NumIterations;
+		
+		// Convert the number of recursive subdivisions to the equivalent tessellation level to make sure that we 
+		// produce the same topology. 
+		int TessellationLevel = static_cast<int>(FMath::RoundHalfFromZero(FMath::Exp2((double)NumIterations))) - 1;
+		Tessellator.TesselationLevel = TessellationLevel;
 		Tessellator.Compute();
 
 		if (Options.bRecomputeNormals && EditMesh.HasAttributes() && EditMesh.Attributes()->PrimaryNormals() != nullptr)
@@ -47,7 +53,84 @@ UDynamicMesh* UGeometryScriptLibrary_MeshSubdivideFunctions::ApplyRecursivePNTes
 	return TargetMesh;
 }
 
+UDynamicMesh* UGeometryScriptLibrary_MeshSubdivideFunctions::ApplyPNTessellation(
+	UDynamicMesh* TargetMesh,
+	FGeometryScriptPNTessellateOptions Options,
+	int TessellationLevel,
+	UGeometryScriptDebug* Debug)
+{
+	if (TargetMesh == nullptr)
+	{
+		UE::Geometry::AppendError(Debug, EGeometryScriptErrorType::InvalidInputs, LOCTEXT("ApplyPNTessellation_InvalidInput", "ApplyPNTessellation: TargetMesh is Null"));
+		return TargetMesh;
+	}
+	if (TessellationLevel <= 0)
+	{
+		return TargetMesh;
+	}
 
+	TargetMesh->EditMesh([&](FDynamicMesh3& EditMesh) 
+	{
+		FPNTriangles Tessellator(&EditMesh);
+		Tessellator.TesselationLevel = TessellationLevel;
+		
+		if (Tessellator.Validate() != EOperationValidationResult::Ok)
+		{
+			UE::Geometry::AppendError(Debug, EGeometryScriptErrorType::InvalidInputs, LOCTEXT("ApplyPNTessellation_InvalidInput", "ApplyPNTessellation: TargetMesh is Null"));
+			return;
+		}
+
+		if (Tessellator.Compute() == false) 
+		{
+			UE::Geometry::AppendError(Debug, EGeometryScriptErrorType::OperationFailed, LOCTEXT("ApplyPNTessellation_Failed", "ApplyPNTessellation: Tessellation failed."));
+			return;
+		}
+
+		if (Options.bRecomputeNormals && EditMesh.HasAttributes() && EditMesh.Attributes()->PrimaryNormals() != nullptr)
+		{
+			FMeshNormals MeshNormals(&EditMesh);
+			MeshNormals.RecomputeOverlayNormals(EditMesh.Attributes()->PrimaryNormals(), true, true);
+			MeshNormals.CopyToOverlay(EditMesh.Attributes()->PrimaryNormals(), false);
+		}
 
+	}, EDynamicMeshChangeType::GeneralEdit, EDynamicMeshAttributeChangeFlags::Unknown, false);
+
+	return TargetMesh;
+}
+
+UDynamicMesh* UGeometryScriptLibrary_MeshSubdivideFunctions::ApplyUniformTessellation(
+	UDynamicMesh* TargetMesh,
+	int TessellationLevel,
+	UGeometryScriptDebug* Debug)
+{
+	if (TargetMesh == nullptr)
+	{
+		UE::Geometry::AppendError(Debug, EGeometryScriptErrorType::InvalidInputs, LOCTEXT("ApplyUniformTessellation_InvalidInput", "ApplyUniformTessellation: TargetMesh is Null"));
+		return TargetMesh;
+	}
+	if (TessellationLevel <= 0)
+	{
+		return TargetMesh;
+	}
+
+	TargetMesh->EditMesh([&](FDynamicMesh3& EditMesh) 
+	{
+		FUniformTesselate Tessellator(&EditMesh);
+		Tessellator.TesselationNum = TessellationLevel;
+		
+		if (Tessellator.Validate() != EOperationValidationResult::Ok)
+		{
+			UE::Geometry::AppendError(Debug, EGeometryScriptErrorType::InvalidInputs, LOCTEXT("ApplyUniformTessellation_Error", "ApplyUniformTessellation: The inputs are invalid."));
+			return;
+		} 
+
+		if (Tessellator.Compute() == false) 
+		{
+			UE::Geometry::AppendError(Debug, EGeometryScriptErrorType::OperationFailed, LOCTEXT("ApplyUniformTessellation_Failed", "ApplyUniformTessellation: Tessellation failed.")); 
+		}
+	}, EDynamicMeshChangeType::GeneralEdit, EDynamicMeshAttributeChangeFlags::Unknown, false);
+
+	return TargetMesh;
+}
 
 #undef LOCTEXT_NAMESPACE
diff --git a/Engine/Plugins/Experimental/GeometryScripting/Source/GeometryScriptingCore/Public/GeometryScript/MeshSubdivideFunctions.h b/Engine/Plugins/Experimental/GeometryScripting/Source/GeometryScriptingCore/Public/GeometryScript/MeshSubdivideFunctions.h
--- a/Engine/Plugins/Experimental/GeometryScripting/Source/GeometryScriptingCore/Public/GeometryScript/MeshSubdivideFunctions.h
+++ b/Engine/Plugins/Experimental/GeometryScripting/Source/GeometryScriptingCore/Public/GeometryScript/MeshSubdivideFunctions.h
@@ -20,15 +20,14 @@ public:
 	bool bRecomputeNormals = true;
 };
 
-
-
 UCLASS(meta = (ScriptName = "GeometryScript_MeshSubdivide"))
 class GEOMETRYSCRIPTINGCORE_API UGeometryScriptLibrary_MeshSubdivideFunctions : public UBlueprintFunctionLibrary
 {
 	GENERATED_BODY()
 public:
 
-	UFUNCTION(BlueprintCallable, Category = "GeometryScript|Subdivide", meta=(ScriptMethod))
+	UE_DEPRECATED(5.2, "Use 'Apply PN Tessellation' instead; this deprecated version recursively subdivided the mesh NumIterations times while the new version splits every triangle into (TessellationLevel+1)^2 smaller triangles.")
+	UFUNCTION(BlueprintCallable, Category = "GeometryScript|Subdivide", meta = (ScriptMethod, DeprecatedFunction, DeprecationMessage = "Use 'Apply PN Tessellation' instead; this deprecated version recursively subdivided the mesh NumIterations times while the new version splits every triangle into (TessellationLevel+1)^2 smaller triangles."))
 	static UPARAM(DisplayName = "Target Mesh") UDynamicMesh* 
 	ApplyRecursivePNTessellation(
 		UDynamicMesh* TargetMesh,
@@ -36,4 +35,18 @@ public:
 		int NumIterations = 3,
 		UGeometryScriptDebug* Debug = nullptr );
 
+	UFUNCTION(BlueprintCallable, Category = "GeometryScript|Subdivide", meta=(ScriptMethod))
+	static UPARAM(DisplayName = "Target Mesh") UDynamicMesh* 
+	ApplyPNTessellation(
+		UDynamicMesh* TargetMesh,
+		FGeometryScriptPNTessellateOptions Options,
+		int TessellationLevel = 3,
+		UGeometryScriptDebug* Debug = nullptr );
+	
+	UFUNCTION(BlueprintCallable, Category = "GeometryScript|Subdivide", meta=(ScriptMethod))
+	static UPARAM(DisplayName = "Target Mesh") UDynamicMesh* 
+	ApplyUniformTessellation(
+		UDynamicMesh* TargetMesh,
+		int TessellationLevel = 3,
+		UGeometryScriptDebug* Debug = nullptr );
 };
\ No newline at end of file
