GeometryScripting: add SetStaticMeshCollisionFromComponent function that copies simple collision from any UPrimitiveComponent to a UStaticMesh. Fix some minor issues in CollisionFunctions.
#rnx

diff --git a/Engine/Plugins/Experimental/GeometryScripting/Source/GeometryScriptingCore/Private/CollisionFunctions.cpp b/Engine/Plugins/Experimental/GeometryScripting/Source/GeometryScriptingCore/Private/CollisionFunctions.cpp
--- a/Engine/Plugins/Experimental/GeometryScripting/Source/GeometryScriptingCore/Private/CollisionFunctions.cpp
+++ b/Engine/Plugins/Experimental/GeometryScripting/Source/GeometryScriptingCore/Private/CollisionFunctions.cpp
@@ -49,7 +49,6 @@ void ComputeCollisionFromMesh(
 
 	TArray<FDynamicMesh3> Submeshes;
 	TArray<const FDynamicMesh3*> SubmeshPointers;
-	SubmeshPointers.SetNum(NumComponents);
 
 	if (NumComponents == 1)
 	{
@@ -119,47 +118,25 @@ void ComputeCollisionFromMesh(
 	GeneratedCollision = NewCollision.AggGeom;
 }
 
-}		// end namespace UELocal
 
 
-UDynamicMesh* UGeometryScriptLibrary_CollisionFunctions::SetStaticMeshCollisionFromMesh(
-	UDynamicMesh* FromDynamicMesh,
-	UStaticMesh* ToStaticMeshAsset,
-	FGeometryScriptCollisionFromMeshOptions Options,
-	UGeometryScriptDebug* Debug)
+static void SetStaticMeshSimpleCollision(UStaticMesh* StaticMeshAsset, const FKAggregateGeom& NewSimpleCollision, bool bEmitTransaction)
 {
-	if (FromDynamicMesh == nullptr)
-	{
-		UE::Geometry::AppendError(Debug, EGeometryScriptErrorType::InvalidInputs, LOCTEXT("SetStaticMeshCollisionFromMesh_InvalidInput1", "SetStaticMeshCollisionFromMesh: FromDynamicMesh is Null"));
-		return FromDynamicMesh;
-	}
-	if (ToStaticMeshAsset == nullptr)
-	{
-		UE::Geometry::AppendError(Debug, EGeometryScriptErrorType::InvalidInputs, LOCTEXT("SetStaticMeshCollisionFromMesh_InvalidInput2", "SetStaticMeshCollisionFromMesh: ToStaticMeshAsset is Null"));
-		return FromDynamicMesh;
-	}
-
-	FKAggregateGeom NewCollision;
-	FromDynamicMesh->ProcessMesh([&](const FDynamicMesh3& ReadMesh)
-	{
-		UELocal::ComputeCollisionFromMesh(ReadMesh, NewCollision, Options);
-	});
-
 #if WITH_EDITOR
-	if (Options.bEmitTransaction)
+	if (bEmitTransaction)
 	{
 		GEditor->BeginTransaction(LOCTEXT("UpdateStaticMesh", "Set Simple Collision"));
 	}
 
-	ToStaticMeshAsset->Modify();
+	StaticMeshAsset->Modify();
 #endif
 
-	UBodySetup* BodySetup = ToStaticMeshAsset->GetBodySetup();
+	UBodySetup* BodySetup = StaticMeshAsset->GetBodySetup();
 	if (BodySetup != nullptr)
 	{
 		// mark the BodySetup for modification. Do we need to modify the UStaticMesh??
 #if WITH_EDITOR
-		if (Options.bEmitTransaction)
+		if (bEmitTransaction)
 		{
 			BodySetup->Modify();
 		}
@@ -169,7 +146,7 @@ UDynamicMesh* UGeometryScriptLibrary_CollisionFunctions::SetStaticMeshCollisionF
 		BodySetup->RemoveSimpleCollision();
 
 		// set new collision geometry
-		BodySetup->AggGeom = NewCollision;
+		BodySetup->AggGeom = NewSimpleCollision;
 
 		// update collision type
 		//BodySetup->CollisionTraceFlag = (ECollisionTraceFlag)(int32)Settings->SetCollisionType;
@@ -178,13 +155,13 @@ UDynamicMesh* UGeometryScriptLibrary_CollisionFunctions::SetStaticMeshCollisionF
 		BodySetup->CreatePhysicsMeshes();
 
 		// rebuild nav collision (? StaticMeshEditor does this)
-		ToStaticMeshAsset->CreateNavCollision(/*bIsUpdate=*/true);
+		StaticMeshAsset->CreateNavCollision(/*bIsUpdate=*/true);
 
 		// update physics state on all components using this StaticMesh
 		for (FThreadSafeObjectIterator Iter(UStaticMeshComponent::StaticClass()); Iter; ++Iter)
 		{
 			UStaticMeshComponent* SMComponent = Cast<UStaticMeshComponent>(*Iter);
-			if (SMComponent->GetStaticMesh() == ToStaticMeshAsset)
+			if (SMComponent->GetStaticMesh() == StaticMeshAsset)
 			{
 				if (SMComponent->IsPhysicsStateCreated())
 				{
@@ -196,21 +173,52 @@ UDynamicMesh* UGeometryScriptLibrary_CollisionFunctions::SetStaticMeshCollisionF
 		// do we need to do a post edit change here??
 
 		// mark static mesh as dirty so it gets resaved?
-		ToStaticMeshAsset->MarkPackageDirty();
+		StaticMeshAsset->MarkPackageDirty();
 
 #if WITH_EDITORONLY_DATA
 		// mark the static mesh as having customized collision so it is not regenerated on reimport
-		ToStaticMeshAsset->bCustomizedCollision = true;
+		StaticMeshAsset->bCustomizedCollision = true;
 #endif // WITH_EDITORONLY_DATA
 	}
 
 #if WITH_EDITOR
-	if (Options.bEmitTransaction)
+	if (bEmitTransaction)
 	{
 		GEditor->EndTransaction();
 	}
 #endif
 
+}
+
+
+
+}		// end namespace UELocal
+
+
+UDynamicMesh* UGeometryScriptLibrary_CollisionFunctions::SetStaticMeshCollisionFromMesh(
+	UDynamicMesh* FromDynamicMesh,
+	UStaticMesh* ToStaticMeshAsset,
+	FGeometryScriptCollisionFromMeshOptions Options,
+	UGeometryScriptDebug* Debug)
+{
+	if (FromDynamicMesh == nullptr)
+	{
+		UE::Geometry::AppendError(Debug, EGeometryScriptErrorType::InvalidInputs, LOCTEXT("SetStaticMeshCollisionFromMesh_InvalidInput1", "SetStaticMeshCollisionFromMesh: FromDynamicMesh is Null"));
+		return FromDynamicMesh;
+	}
+	if (ToStaticMeshAsset == nullptr)
+	{
+		UE::Geometry::AppendError(Debug, EGeometryScriptErrorType::InvalidInputs, LOCTEXT("SetStaticMeshCollisionFromMesh_InvalidInput2", "SetStaticMeshCollisionFromMesh: ToStaticMeshAsset is Null"));
+		return FromDynamicMesh;
+	}
+
+	FKAggregateGeom NewCollision;
+	FromDynamicMesh->ProcessMesh([&](const FDynamicMesh3& ReadMesh)
+	{
+		UELocal::ComputeCollisionFromMesh(ReadMesh, NewCollision, Options);
+	});
+
+	UELocal::SetStaticMeshSimpleCollision(ToStaticMeshAsset, NewCollision, Options.bEmitTransaction);
 
 	return FromDynamicMesh;
 }
@@ -218,6 +226,34 @@ UDynamicMesh* UGeometryScriptLibrary_CollisionFunctions::SetStaticMeshCollisionF
 
 
 
+void UGeometryScriptLibrary_CollisionFunctions::SetStaticMeshCollisionFromComponent(
+	UStaticMesh* UpdateStaticMeshAsset, 
+	UPrimitiveComponent* SourceComponent,
+	FGeometryScriptSetSimpleCollisionOptions Options,
+	UGeometryScriptDebug* Debug)
+{
+	if (UpdateStaticMeshAsset == nullptr)
+	{
+		UE::Geometry::AppendError(Debug, EGeometryScriptErrorType::InvalidInputs, LOCTEXT("SetStaticMeshCollisionFromComponent_InvalidStaticMesh", "SetStaticMeshCollisionFromComponent: UpdateStaticMeshAsset is Null"));
+		return;
+	}
+	if (SourceComponent == nullptr)
+	{
+		UE::Geometry::AppendError(Debug, EGeometryScriptErrorType::InvalidInputs, LOCTEXT("SetStaticMeshCollisionFromComponent_InvalidSourceComponent", "SetStaticMeshCollisionFromComponent: SourceComponent is Null"));
+		return;
+	}
+
+	const UBodySetup* BodySetup = SourceComponent->GetBodySetup();
+	if (BodySetup == nullptr)
+	{
+		UE::Geometry::AppendError(Debug, EGeometryScriptErrorType::InvalidInputs, LOCTEXT("SetStaticMeshCollisionFromComponent_InvalidBodySetup", "SetStaticMeshCollisionFromComponent: SourceComponent BodySetup is Null"));
+		return;
+	}
+
+	UELocal::SetStaticMeshSimpleCollision(UpdateStaticMeshAsset, BodySetup->AggGeom, Options.bEmitTransaction);
+}
+
+
 
 
 
diff --git a/Engine/Plugins/Experimental/GeometryScripting/Source/GeometryScriptingCore/Public/GeometryScript/CollisionFunctions.h b/Engine/Plugins/Experimental/GeometryScripting/Source/GeometryScriptingCore/Public/GeometryScript/CollisionFunctions.h
--- a/Engine/Plugins/Experimental/GeometryScripting/Source/GeometryScriptingCore/Public/GeometryScript/CollisionFunctions.h
+++ b/Engine/Plugins/Experimental/GeometryScripting/Source/GeometryScriptingCore/Public/GeometryScript/CollisionFunctions.h
@@ -81,6 +81,15 @@ public:
 };
 
 
+USTRUCT(BlueprintType)
+struct GEOMETRYSCRIPTINGCORE_API FGeometryScriptSetSimpleCollisionOptions
+{
+	GENERATED_BODY()
+public:
+	UPROPERTY(BlueprintReadWrite, Category = Options)
+	bool bEmitTransaction = true;
+};
+
 
 
 UCLASS(meta = (ScriptName = "GeometryScript_Collision"))
@@ -98,6 +107,18 @@ public:
 		FGeometryScriptCollisionFromMeshOptions Options,
 		UGeometryScriptDebug* Debug = nullptr);
 
+	/**
+	 * Copy the Simple Collision Geometry from the SourceComponent to the StaticMeshAsset
+	 */
+	UFUNCTION(BlueprintCallable, Category = "GeometryScript|Collision")
+	static void 
+	SetStaticMeshCollisionFromComponent(
+		UStaticMesh* StaticMeshAsset, 
+		UPrimitiveComponent* SourceComponent,
+		FGeometryScriptSetSimpleCollisionOptions Options = FGeometryScriptSetSimpleCollisionOptions(),
+		UGeometryScriptDebug* Debug = nullptr);
+
+
 	UFUNCTION(BlueprintCallable, Category = "GeometryScript|Collision")
 	static UPARAM(DisplayName = "Dynamic Mesh") UDynamicMesh* 
 	SetDynamicMeshCollisionFromMesh(
