GeometryScript: don't mark Components dirty when setting collision unless a transaction is open
#rnx

diff --git a/Engine/Plugins/Experimental/GeometryScripting/Source/GeometryScriptingCore/Private/CollisionFunctions.cpp b/Engine/Plugins/Experimental/GeometryScripting/Source/GeometryScriptingCore/Private/CollisionFunctions.cpp
--- a/Engine/Plugins/Experimental/GeometryScripting/Source/GeometryScriptingCore/Private/CollisionFunctions.cpp
+++ b/Engine/Plugins/Experimental/GeometryScripting/Source/GeometryScriptingCore/Private/CollisionFunctions.cpp
@@ -126,9 +126,9 @@ static void SetStaticMeshSimpleCollision(UStaticMesh* StaticMeshAsset, const FKA
 	if (bEmitTransaction)
 	{
 		GEditor->BeginTransaction(LOCTEXT("UpdateStaticMesh", "Set Simple Collision"));
-	}
 
-	StaticMeshAsset->Modify();
+		StaticMeshAsset->Modify();
+	}
 #endif
 
 	UBodySetup* BodySetup = StaticMeshAsset->GetBodySetup();
@@ -284,9 +284,9 @@ UDynamicMesh* UGeometryScriptLibrary_CollisionFunctions::SetDynamicMeshCollision
 	if (Options.bEmitTransaction)
 	{
 		GEditor->BeginTransaction(LOCTEXT("UpdateDynamicMesh", "Set Simple Collision"));
-	}
 
-	DynamicMeshComponent->Modify();
+		DynamicMeshComponent->Modify();
+	}
 #endif
 
 	UBodySetup* BodySetup = DynamicMeshComponent->GetBodySetup();
