GeometryScripting: additions to support working with StaticMesh asset materials
- added GetSectionMaterialListFromStaticMesh() function that constructs the per-section material list (ie corresponds to DynamicMesh MaterialIDs after conversion)
- added SetComponentMaterialList() utility function to assign new set of Materials to a Component (eg to apply per-section Material list to a DynamicMeshComponent initialized from a StaticMesh)
- added DeleteTrianglesByMaterialID() to simplify BPs that want to removing triangles by assigned Material or known MaterialID
- added CompactMaterialIDs() function that cleans up MaterialID assignments and also can simultaneously compact the corresponding Material list
- CopyMeshToStaticMesh() can now also update the StaticMesh with a new Material set via new fields in FGeometryScriptCopyMeshToAssetOptions
#rnx

diff --git a/Engine/Plugins/Experimental/GeometryScripting/Source/GeometryScriptingCore/Private/MeshAssetFunctions.cpp b/Engine/Plugins/Experimental/GeometryScripting/Source/GeometryScriptingCore/Private/MeshAssetFunctions.cpp
--- a/Engine/Plugins/Experimental/GeometryScripting/Source/GeometryScriptingCore/Private/MeshAssetFunctions.cpp
+++ b/Engine/Plugins/Experimental/GeometryScripting/Source/GeometryScriptingCore/Private/MeshAssetFunctions.cpp
@@ -11,13 +11,16 @@
 #include "StaticMeshAttributes.h"
 #include "StaticMeshOperations.h"
 #include "RenderingThread.h"
+#include "PhysicsEngine/BodySetup.h"
 
 #include "MeshDescriptionToDynamicMesh.h"
 #include "DynamicMeshToMeshDescription.h"
+#include "AssetUtils/StaticMeshMaterialUtil.h"
 
 
 #if WITH_EDITOR
 #include "Editor.h"
+#include "ScopedTransaction.h"
 #endif
 
 using namespace UE::Geometry;
@@ -141,6 +144,12 @@ UDynamicMesh*  UGeometryScriptLibrary_StaticMeshFunctions::CopyMeshToStaticMesh(
 
 	int32 UseLODIndex = FMath::Clamp(TargetLOD.LODIndex, 0, 32);
 
+	if (Options.bReplaceMaterials && UseLODIndex != 0)
+	{
+		UE::Geometry::AppendError(Debug, EGeometryScriptErrorType::InvalidInputs, LOCTEXT("CopyMeshToStaticMesh_InvalidOptions1", "CopyMeshToStaticMesh: Can only Replace Materials when updating LOD0"));
+		return FromDynamicMesh;
+	}
+
 	// don't allow build-in engine assets to be modified
 	if (ToStaticMeshAsset->GetPathName().StartsWith(TEXT("/Engine/")))
 	{
@@ -196,13 +205,42 @@ UDynamicMesh*  UGeometryScriptLibrary_StaticMeshFunctions::CopyMeshToStaticMesh(
 		Converter.Convert(&ReadMesh, *MeshDescription, !BuildSettings.bRecomputeTangents);
 	});
 
-	ToStaticMeshAsset->CommitMeshDescription(UseLODIndex);
-
 	// Setting to prevent the standard static mesh reduction from running and replacing the render LOD.
 	FStaticMeshSourceModel& ThisSourceModel = ToStaticMeshAsset->GetSourceModel(UseLODIndex);
 	ThisSourceModel.ReductionSettings.PercentTriangles = 1.f;
 	ThisSourceModel.ReductionSettings.PercentVertices = 1.f;
 
+	if (Options.bReplaceMaterials)
+	{
+		bool bHaveSlotNames = (Options.NewMaterialSlotNames.Num() == Options.NewMaterials.Num());
+
+		TArray<FStaticMaterial> NewMaterials;
+		for (int32 k = 0; k < Options.NewMaterials.Num(); ++k)
+		{
+			FStaticMaterial NewMaterial;
+			NewMaterial.MaterialInterface = Options.NewMaterials[k];
+			FName UseSlotName = (bHaveSlotNames && Options.NewMaterialSlotNames[k] != NAME_None) ? Options.NewMaterialSlotNames[k] :
+				UE::AssetUtils::GenerateNewMaterialSlotName(NewMaterials, NewMaterial.MaterialInterface, k);
+
+			NewMaterial.MaterialSlotName = UseSlotName;
+			NewMaterial.ImportedMaterialSlotName = UseSlotName;
+			NewMaterial.UVChannelData = FMeshUVChannelInfo(1.f);		// this avoids an ensure in  UStaticMesh::GetUVChannelData
+			NewMaterials.Add(NewMaterial);
+		}
+
+		ToStaticMeshAsset->SetStaticMaterials(NewMaterials);
+
+		// Reset the section info map
+		int32 NumSections = ToStaticMeshAsset->GetSectionInfoMap().GetSectionNumber(UseLODIndex);
+		for (int32 SectionIndex = 0; SectionIndex < NumSections; ++SectionIndex)
+		{
+			ToStaticMeshAsset->GetSectionInfoMap().Remove(UseLODIndex, SectionIndex);
+		}
+	}
+
+	ToStaticMeshAsset->CommitMeshDescription(UseLODIndex);
+
+
 	if (Options.bDeferMeshPostEditChange == false)
 	{
 		ToStaticMeshAsset->PostEditChange();
@@ -225,4 +263,46 @@ UDynamicMesh*  UGeometryScriptLibrary_StaticMeshFunctions::CopyMeshToStaticMesh(
 
 
 
+
+
+void UGeometryScriptLibrary_StaticMeshFunctions::GetSectionMaterialListFromStaticMesh(
+	UStaticMesh* FromStaticMeshAsset, 
+	FGeometryScriptMeshReadLOD RequestedLOD,
+	TArray<UMaterialInterface*>& MaterialList,
+	TArray<int32>& MaterialIndex,
+	TEnumAsByte<EGeometryScriptOutcomePins>& Outcome,
+	UGeometryScriptDebug* Debug)
+{
+	Outcome = EGeometryScriptOutcomePins::Failure;
+
+	if (FromStaticMeshAsset == nullptr)
+	{
+		UE::Geometry::AppendError(Debug, EGeometryScriptErrorType::InvalidInputs, LOCTEXT("GetSectionMaterialListFromStaticMesh_InvalidInput1", "GetSectionMaterialListFromStaticMesh: FromStaticMeshAsset is Null"));
+		return;
+	}
+	if (RequestedLOD.LODType != EGeometryScriptLODType::MaxAvailable && RequestedLOD.LODType != EGeometryScriptLODType::SourceModel)
+	{
+		UE::Geometry::AppendError(Debug, EGeometryScriptErrorType::InvalidInputs, LOCTEXT("GetSectionMaterialListFromStaticMesh_LODNotAvailable", "GetSectionMaterialListFromStaticMesh: Requested LOD is not available"));
+		return;
+	}
+
+	int32 UseLODIndex = FMath::Clamp(RequestedLOD.LODIndex, 0, FromStaticMeshAsset->GetNumLODs() - 1);
+
+	MaterialList.Reset();
+	MaterialIndex.Reset();
+	if (UE::AssetUtils::GetStaticMeshLODMaterialListBySection(FromStaticMeshAsset, UseLODIndex, MaterialList, MaterialIndex) == false)
+	{
+		UE::Geometry::AppendError(Debug, EGeometryScriptErrorType::InvalidInputs, LOCTEXT("GetSectionMaterialListFromStaticMesh_QueryFailed", "GetSectionMaterialListFromStaticMesh: Could not fetch Material Set from Asset"));
+		return;
+	}
+
+	Outcome = EGeometryScriptOutcomePins::Success;
+}
+
+
+
+
+
+
+
 #undef LOCTEXT_NAMESPACE
\ No newline at end of file
diff --git a/Engine/Plugins/Experimental/GeometryScripting/Source/GeometryScriptingCore/Private/MeshMaterialFunctions.cpp b/Engine/Plugins/Experimental/GeometryScripting/Source/GeometryScriptingCore/Private/MeshMaterialFunctions.cpp
--- a/Engine/Plugins/Experimental/GeometryScripting/Source/GeometryScriptingCore/Private/MeshMaterialFunctions.cpp
+++ b/Engine/Plugins/Experimental/GeometryScripting/Source/GeometryScriptingCore/Private/MeshMaterialFunctions.cpp
@@ -4,6 +4,7 @@
 
 #include "DynamicMesh/DynamicMesh3.h"
 #include "DynamicMesh/DynamicMeshAttributeSet.h"
+#include "DynamicMesh/MeshAttributeUtil.h"
 #include "UDynamicMesh.h"
 #include "Polygroups/PolygroupSet.h"
 
@@ -90,6 +91,7 @@ UDynamicMesh* UGeometryScriptLibrary_MeshMaterialFunctions::EnableMaterialIDs(
 		return TargetMesh;
 	}
 
+	// this will enable the material IDs even though the lambda doesn't do anything
 	bool bHasMaterialIDs;
 	SimpleMeshMaterialEdit(TargetMesh, true, bHasMaterialIDs, [](FDynamicMesh3& Mesh, FDynamicMeshMaterialAttribute& MaterialIDs) {});
 
@@ -328,4 +330,149 @@ UDynamicMesh* UGeometryScriptLibrary_MeshMaterialFunctions::SetPolygroupMaterial
 
 
 
+UDynamicMesh* UGeometryScriptLibrary_MeshMaterialFunctions::DeleteTrianglesByMaterialID(
+	UDynamicMesh* TargetMesh,
+	int MaterialID,
+	int& NumDeleted,
+	bool bDeferChangeNotifications,
+	UGeometryScriptDebug* Debug)
+{
+	NumDeleted = 0;
+	if (TargetMesh)
+	{
+		TargetMesh->EditMesh([&](FDynamicMesh3& EditMesh)
+		{
+			FDynamicMeshMaterialAttribute* MaterialIDs = (EditMesh.HasAttributes() && EditMesh.Attributes()->HasMaterialID()) ? EditMesh.Attributes()->GetMaterialID() : nullptr;
+			if (MaterialIDs == nullptr)
+			{
+				UE::Geometry::AppendError(Debug, EGeometryScriptErrorType::InvalidInputs, LOCTEXT("DeleteTrianglesByMaterialID_NoMaterialID", "DeleteTrianglesByMaterialID: MaterialID Attribute is not enabled"));
+				return;
+			}
+
+			TArray<int32> TriangleList;
+			for (int32 TriangleID : EditMesh.TriangleIndicesItr())
+			{
+				if (MaterialIDs->GetValue(TriangleID) == MaterialID)
+				{
+					TriangleList.Add(TriangleID);
+				}
+			}
+
+			for (int32 TriangleID : TriangleList)
+			{
+				EMeshResult Result = EditMesh.RemoveTriangle(TriangleID);
+				if (Result == EMeshResult::Ok)
+				{
+					NumDeleted++;
+				}
+			}
+		}, EDynamicMeshChangeType::GeneralEdit, EDynamicMeshAttributeChangeFlags::Unknown, bDeferChangeNotifications);
+	}
+	return TargetMesh;
+}
+
+
+
+UDynamicMesh* UGeometryScriptLibrary_MeshMaterialFunctions::CompactMaterialIDs(
+	UDynamicMesh* TargetMesh,
+	TArray<UMaterialInterface*> SourceMaterialList,
+	TArray<UMaterialInterface*>& CompactedMaterialList,
+	UGeometryScriptDebug* Debug)
+{
+	CompactedMaterialList = SourceMaterialList;
+
+	if (TargetMesh)
+	{
+		TargetMesh->EditMesh([&](FDynamicMesh3& EditMesh)
+		{
+			FDynamicMeshMaterialAttribute* MaterialIDs = (EditMesh.HasAttributes() && EditMesh.Attributes()->HasMaterialID()) ? EditMesh.Attributes()->GetMaterialID() : nullptr;
+			if (MaterialIDs == nullptr)
+			{
+				UE::Geometry::AppendError(Debug, EGeometryScriptErrorType::InvalidInputs, LOCTEXT("CompactMaterialIDs_NoMaterialID", "CompactMaterialIDs: MaterialID Attribute is not enabled"));
+				return;
+			}
+
+			bool bWasCompact = false;
+			TArray<int32> OldToNewMap, NewToOldMap;
+			bool bOK = UE::Geometry::CompactAttributeValues(EditMesh, *MaterialIDs, OldToNewMap, NewToOldMap, bWasCompact);
+			if (bOK == false)
+			{
+				UE::Geometry::AppendError(Debug, EGeometryScriptErrorType::InvalidInputs, LOCTEXT("CompactMaterialIDs_InvalidMaterialID", "CompactMaterialIDs: Invalid MaterialIDs found, unsafe to Compact"));
+				return;
+			}
+			if (bWasCompact)
+			{
+				return;
+			}
+
+			int32 NewMaterialIDCount = NewToOldMap.Num();
+			CompactedMaterialList.Reset();
+			CompactedMaterialList.SetNum(NewMaterialIDCount);
+			for (int32 k = 0; k < NewMaterialIDCount; ++k)
+			{
+				int32 OldIndex = NewToOldMap[k];
+				CompactedMaterialList[k] = (OldIndex < SourceMaterialList.Num()) ?
+					SourceMaterialList[OldIndex] : nullptr;
+			}
+
+			//TArray<int32> KnownMaterialIDs;
+			//int32 MaxMaterialID = -1, MinMaterialID = TNumericLimits<int32>::Max();
+			//for (int32 TriangleID : EditMesh.TriangleIndicesItr())
+			//{
+			//	int32 MaterialID = MaterialIDs->GetValue(TriangleID);
+			//	MaxMaterialID = FMath::Max(MaterialID, MaxMaterialID);
+			//	MinMaterialID = FMath::Min(MaterialID, MinMaterialID);
+			//	KnownMaterialIDs.AddUnique(MaterialID);
+			//}
+			//if (MinMaterialID < 0)
+			//{
+			//	UE::Geometry::AppendError(Debug, EGeometryScriptErrorType::InvalidInputs, LOCTEXT("CompactMaterialIDs_InvalidMaterialID", "CompactMaterialIDs: Negative MaterialID found, unsafe to Compact"));
+			//	return;
+			//}
+			//if ( MaxMaterialID == (KnownMaterialIDs.Num()-1) )
+			//{
+			//	// set is already compact, nothing to do
+			//	return;
+			//}
+
+			//TArray<int32> MaterialIDMap;
+			//MaterialIDMap.Init(-1, MaxMaterialID+1);
+
+			//int32 NewMaterialIDCounter = 0;
+			//int32 NewMaxID = 0;
+			//for (int32 TriangleID : EditMesh.TriangleIndicesItr())
+			//{
+			//	int32 MaterialID = MaterialIDs->GetValue(TriangleID);
+			//	if (MaterialIDMap[MaterialID] == -1)
+			//	{
+			//		MaterialIDMap[MaterialID] = NewMaterialIDCounter;
+			//		NewMaxID = FMath::Max(NewMaxID, NewMaterialIDCounter);
+			//		NewMaterialIDCounter++;
+			//	}
+			//	int32 NewMaterialID = MaterialIDMap[MaterialID];
+			//	MaterialIDs->SetValue(TriangleID, NewMaterialID);
+			//}
+
+			//CompactedMaterialList.Reset();
+			//if ( SourceMaterialList.Num() == (MaxMaterialID+1) )
+			//{
+			//	CompactedMaterialList.SetNum(NewMaterialIDCounter);
+			//	for (int32 k = 0; k <= MaxMaterialID; ++k)
+			//	{
+			//		if (MaterialIDMap[k] >= 0)
+			//		{
+			//			CompactedMaterialList[MaterialIDMap[k]] = SourceMaterialList[k];
+			//		}
+			//		
+			//	}
+			//}
+
+		}, EDynamicMeshChangeType::GeneralEdit, EDynamicMeshAttributeChangeFlags::Unknown, false);
+	}
+	return TargetMesh;
+}
+
+
+
+
 #undef LOCTEXT_NAMESPACE
\ No newline at end of file
diff --git a/Engine/Plugins/Experimental/GeometryScripting/Source/GeometryScriptingCore/Private/SceneUtilityFunctions.cpp b/Engine/Plugins/Experimental/GeometryScripting/Source/GeometryScriptingCore/Private/SceneUtilityFunctions.cpp
--- a/Engine/Plugins/Experimental/GeometryScripting/Source/GeometryScriptingCore/Private/SceneUtilityFunctions.cpp
+++ b/Engine/Plugins/Experimental/GeometryScripting/Source/GeometryScriptingCore/Private/SceneUtilityFunctions.cpp
@@ -115,4 +115,25 @@ UDynamicMesh* UGeometryScriptLibrary_SceneUtilityFunctions::CopyMeshFromComponen
 }
 
 
+
+void UGeometryScriptLibrary_SceneUtilityFunctions::SetComponentMaterialList(
+	UPrimitiveComponent* Component,
+	const TArray<UMaterialInterface*>& MaterialList,
+	UGeometryScriptDebug* Debug)
+{
+	if (Component == nullptr)
+	{
+		UE::Geometry::AppendError(Debug, EGeometryScriptErrorType::InvalidInputs, LOCTEXT("SetComponentMaterialList_InvalidInput1", "SetComponentMaterialList: FromStaticMeshAsset is Null"));
+		return;
+	}
+
+	for (int32 k = 0; k < MaterialList.Num(); ++k)
+	{
+		Component->SetMaterial(k, MaterialList[k]);
+	}
+}
+
+
+
+
 #undef LOCTEXT_NAMESPACE
\ No newline at end of file
diff --git a/Engine/Plugins/Experimental/GeometryScripting/Source/GeometryScriptingCore/Public/GeometryScript/MeshAssetFunctions.h b/Engine/Plugins/Experimental/GeometryScripting/Source/GeometryScriptingCore/Public/GeometryScript/MeshAssetFunctions.h
--- a/Engine/Plugins/Experimental/GeometryScripting/Source/GeometryScriptingCore/Public/GeometryScript/MeshAssetFunctions.h
+++ b/Engine/Plugins/Experimental/GeometryScripting/Source/GeometryScriptingCore/Public/GeometryScript/MeshAssetFunctions.h
@@ -9,7 +9,7 @@
 
 class UStaticMesh;
 class UDynamicMesh;
-
+class UMaterialInterface;
 
 
 
@@ -44,6 +44,17 @@ public:
 	UPROPERTY(BlueprintReadWrite, Category = Options)
 	bool bEnableRemoveDegenerates = false;
 
+	
+	UPROPERTY(BlueprintReadWrite, Category = Options)
+	bool bReplaceMaterials = false;
+
+	UPROPERTY(BlueprintReadWrite, Category = Options)
+	TArray<UMaterialInterface*> NewMaterials;
+
+	UPROPERTY(BlueprintReadWrite, Category = Options)
+	TArray<FName> NewMaterialSlotNames;
+
+
 	UPROPERTY(BlueprintReadWrite, Category = Options)
 	bool bEmitTransaction = true;
 
@@ -54,7 +65,6 @@ public:
 
 
 
-
 UCLASS(meta = (ScriptName = "GeometryScript_AssetUtils"))
 class GEOMETRYSCRIPTINGCORE_API UGeometryScriptLibrary_StaticMeshFunctions : public UBlueprintFunctionLibrary
 {
@@ -82,6 +92,17 @@ public:
 		TEnumAsByte<EGeometryScriptOutcomePins>& Outcome,
 		UGeometryScriptDebug* Debug = nullptr);
 
+
+
+	UFUNCTION(BlueprintCallable, Category = "GeometryScript|StaticMesh", meta = (ExpandEnumAsExecs = "Outcome"))
+	static void
+	GetSectionMaterialListFromStaticMesh(
+		UStaticMesh* FromStaticMeshAsset, 
+		FGeometryScriptMeshReadLOD RequestedLOD,
+		TArray<UMaterialInterface*>& MaterialList,
+		TArray<int32>& MaterialIndex,
+		TEnumAsByte<EGeometryScriptOutcomePins>& Outcome,
+		UGeometryScriptDebug* Debug = nullptr);
 };
 
 
diff --git a/Engine/Plugins/Experimental/GeometryScripting/Source/GeometryScriptingCore/Public/GeometryScript/MeshMaterialFunctions.h b/Engine/Plugins/Experimental/GeometryScripting/Source/GeometryScriptingCore/Public/GeometryScript/MeshMaterialFunctions.h
--- a/Engine/Plugins/Experimental/GeometryScripting/Source/GeometryScriptingCore/Public/GeometryScript/MeshMaterialFunctions.h
+++ b/Engine/Plugins/Experimental/GeometryScripting/Source/GeometryScriptingCore/Public/GeometryScript/MeshMaterialFunctions.h
@@ -72,4 +72,28 @@ public:
 		bool& bIsValidPolygroupID,
 		bool bDeferChangeNotifications = false,
 		UGeometryScriptDebug* Debug = nullptr);
+
+	UFUNCTION(BlueprintCallable, Category = "GeometryScript|Materials", meta=(ScriptMethod))
+	static UPARAM(DisplayName = "Target Mesh") UDynamicMesh* 
+	DeleteTrianglesByMaterialID( 
+		UDynamicMesh* TargetMesh, 
+		int MaterialID,
+		int& NumDeleted,
+		bool bDeferChangeNotifications = false,
+		UGeometryScriptDebug* Debug = nullptr);
+
+	/**
+	 * Compact the MaterialIDs of the TargetMesh, ie remove any un-used MaterialIDs and remap the remaining
+	 * N in-use MaterialIDs to the range [0,N-1]. Optionally compute a Compacted list of Materials.
+	 * @param SourceMaterialList Input Material list, assumption is that SourceMaterialList.Num() == number of MaterialIDs on mesh at input
+	 * @param CompactedMaterialList new Compacted Material list, one-to-one with new compacted MaterialIDs
+	 */
+	UFUNCTION(BlueprintCallable, Category = "GeometryScript|Materials", meta=(ScriptMethod))
+	static UPARAM(DisplayName = "Target Mesh") UDynamicMesh* 
+	CompactMaterialIDs( 
+		UDynamicMesh* TargetMesh, 
+		TArray<UMaterialInterface*> SourceMaterialList,
+		TArray<UMaterialInterface*>& CompactedMaterialList,
+		UGeometryScriptDebug* Debug = nullptr);
+
 };
\ No newline at end of file
diff --git a/Engine/Plugins/Experimental/GeometryScripting/Source/GeometryScriptingCore/Public/GeometryScript/SceneUtilityFunctions.h b/Engine/Plugins/Experimental/GeometryScripting/Source/GeometryScriptingCore/Public/GeometryScript/SceneUtilityFunctions.h
--- a/Engine/Plugins/Experimental/GeometryScripting/Source/GeometryScriptingCore/Public/GeometryScript/SceneUtilityFunctions.h
+++ b/Engine/Plugins/Experimental/GeometryScripting/Source/GeometryScriptingCore/Public/GeometryScript/SceneUtilityFunctions.h
@@ -8,6 +8,7 @@
 #include "SceneUtilityFunctions.generated.h"
 
 class UStaticMesh;
+class UMaterialInterface;
 class UDynamicMesh;
 class UDynamicMeshPool;
 
@@ -50,4 +51,12 @@ public:
 		TEnumAsByte<EGeometryScriptOutcomePins>& Outcome,
 		UGeometryScriptDebug* Debug = nullptr);
 
+
+	UFUNCTION(BlueprintCallable, Category = "GeometryScript|Scene")
+	static void
+	SetComponentMaterialList(
+		UPrimitiveComponent* Component, 
+		const TArray<UMaterialInterface*>& MaterialList,
+		UGeometryScriptDebug* Debug = nullptr);
+
 };
\ No newline at end of file
