GeometryScript: enable editor-only module on linux/osx platforms
#rnx

diff --git a/Engine/Plugins/Experimental/GeometryScripting/GeometryScripting.uplugin b/Engine/Plugins/Experimental/GeometryScripting/GeometryScripting.uplugin
--- a/Engine/Plugins/Experimental/GeometryScripting/GeometryScripting.uplugin
+++ b/Engine/Plugins/Experimental/GeometryScripting/GeometryScripting.uplugin
@@ -22,8 +22,7 @@
 		{
 			"Name": "GeometryScriptingEditor",
 			"Type": "Editor",
-			"LoadingPhase": "Default",
-			"PlatformAllowList": [ "Win64" ]
+			"LoadingPhase": "Default"
 		}
 	],
 	"Plugins": [
