GeometryScript: fix issue with CompactMaterialIDs where if material ID set is modified but still compact, compacted material list was incorrect
ModelingComponents: add some more return values to ::CompactAttributeValues() to help to identify above case
#rnx

diff --git a/Engine/Plugins/Experimental/GeometryScripting/Source/GeometryScriptingCore/Private/MeshMaterialFunctions.cpp b/Engine/Plugins/Experimental/GeometryScripting/Source/GeometryScriptingCore/Private/MeshMaterialFunctions.cpp
--- a/Engine/Plugins/Experimental/GeometryScripting/Source/GeometryScriptingCore/Private/MeshMaterialFunctions.cpp
+++ b/Engine/Plugins/Experimental/GeometryScripting/Source/GeometryScriptingCore/Private/MeshMaterialFunctions.cpp
@@ -393,8 +393,10 @@ UDynamicMesh* UGeometryScriptLibrary_MeshMaterialFunctions::CompactMaterialIDs(
 			}
 
 			bool bWasCompact = false;
+			FInterval1i OldValueRange;
+			int32 NewMaxValue;
 			TArray<int32> OldToNewMap, NewToOldMap;
-			bool bOK = UE::Geometry::CompactAttributeValues(EditMesh, *MaterialIDs, OldToNewMap, NewToOldMap, bWasCompact);
+			bool bOK = UE::Geometry::CompactAttributeValues(EditMesh, *MaterialIDs, OldValueRange, NewMaxValue, OldToNewMap, NewToOldMap, bWasCompact);
 			if (bOK == false)
 			{
 				UE::Geometry::AppendError(Debug, EGeometryScriptErrorType::InvalidInputs, LOCTEXT("CompactMaterialIDs_InvalidMaterialID", "CompactMaterialIDs: Invalid MaterialIDs found, unsafe to Compact"));
@@ -402,17 +404,26 @@ UDynamicMesh* UGeometryScriptLibrary_MeshMaterialFunctions::CompactMaterialIDs(
 			}
 			if (bWasCompact)
 			{
-				return;
+				// this case occurs if MaterialIDs were removed such that the attribute values are still sequential/compact  (eg removed "at the end")
+				int32 NewMaterialIDCount = NewMaxValue+1;
+				CompactedMaterialList.Reset();
+				CompactedMaterialList.SetNum(NewMaterialIDCount);
+				for (int32 k = 0; k < NewMaterialIDCount; ++k)
+				{
+					CompactedMaterialList[k] = SourceMaterialList.IsValidIndex(k) ? SourceMaterialList[k] : nullptr;
+				}		
+				
 			}
-
-			int32 NewMaterialIDCount = NewToOldMap.Num();
-			CompactedMaterialList.Reset();
-			CompactedMaterialList.SetNum(NewMaterialIDCount);
-			for (int32 k = 0; k < NewMaterialIDCount; ++k)
+			else
 			{
-				int32 OldIndex = NewToOldMap[k];
-				CompactedMaterialList[k] = (OldIndex < SourceMaterialList.Num()) ?
-					SourceMaterialList[OldIndex] : nullptr;
+				int32 NewMaterialIDCount = NewToOldMap.Num();
+				CompactedMaterialList.Reset();
+				CompactedMaterialList.SetNum(NewMaterialIDCount);
+				for (int32 k = 0; k < NewMaterialIDCount; ++k)
+				{
+					int32 OldIndex = NewToOldMap[k];
+					CompactedMaterialList[k] = SourceMaterialList.IsValidIndex(OldIndex) ? SourceMaterialList[OldIndex] : nullptr;
+				}
 			}
 
 		}, EDynamicMeshChangeType::GeneralEdit, EDynamicMeshAttributeChangeFlags::Unknown, false);
diff --git a/Engine/Source/Runtime/GeometryCore/Private/DynamicMesh/MeshAttributeUtil.cpp b/Engine/Source/Runtime/GeometryCore/Private/DynamicMesh/MeshAttributeUtil.cpp
--- a/Engine/Source/Runtime/GeometryCore/Private/DynamicMesh/MeshAttributeUtil.cpp
+++ b/Engine/Source/Runtime/GeometryCore/Private/DynamicMesh/MeshAttributeUtil.cpp
@@ -1,7 +1,6 @@
 // Copyright Epic Games, Inc. All Rights Reserved.
 
 #include "DynamicMesh/MeshAttributeUtil.h"
-#include "IntBoxTypes.h"
 
 using namespace UE::Geometry;
 
@@ -9,6 +8,8 @@ using namespace UE::Geometry;
 bool UE::Geometry::CompactAttributeValues(
 	const FDynamicMesh3& Mesh,
 	TDynamicMeshScalarTriangleAttribute<int32>& TriangleAttrib,
+	FInterval1i& OldMaxAttributeRangeOut,
+	int& NewMaxAttributeValueOut,
 	TArray<int32>& OldToNewMap,
 	TArray<int32>& NewToOldMap,
 	bool& bWasCompact)
@@ -21,6 +22,7 @@ bool UE::Geometry::CompactAttributeValues(
 	{
 		IndexRange.Contain(TriangleAttrib.GetValue(TriangleID));
 	}
+	OldMaxAttributeRangeOut = IndexRange;
 	if (IndexRange.Min < 0)
 	{
 		return false;
@@ -45,6 +47,7 @@ bool UE::Geometry::CompactAttributeValues(
 			TriangleAttrib.SetValue(TriangleID, NewValue);
 		}
 	}
+	NewMaxAttributeValueOut = NewValueCount - 1;
 
 	// construct inverse mapping
 	NewToOldMap.Init(IndexConstants::InvalidID, NewValueCount);
diff --git a/Engine/Source/Runtime/GeometryCore/Public/DynamicMesh/MeshAttributeUtil.h b/Engine/Source/Runtime/GeometryCore/Public/DynamicMesh/MeshAttributeUtil.h
--- a/Engine/Source/Runtime/GeometryCore/Public/DynamicMesh/MeshAttributeUtil.h
+++ b/Engine/Source/Runtime/GeometryCore/Public/DynamicMesh/MeshAttributeUtil.h
@@ -4,6 +4,7 @@
 
 #include "DynamicMesh/DynamicMesh3.h"
 #include "DynamicMesh/DynamicMeshAttributeSet.h"
+#include "IntBoxTypes.h"
 
 
 /**
@@ -17,6 +18,11 @@ namespace Geometry
 	/**
 	 * Compact the values of an integer Triangle Attribute, ie so that the attribute values are dense in range 0..N.
 	 * Useful for (eg) compacting MaterialIDs or Polygroups.
+	 * Note that this function cannot tell if values were removed from the "end" of the attribute value list. Some care must 
+	 * be taken in calling code to handle this case if using this function to determine whether compaction of other data structures is needed.
+	 * Currently OldToNewMapOut and NewToOldMapOut are always populated even if the input is compact. 
+	 * @param OldMaxAttributeRangeOut min and max input attribute values, ie pre-compaction
+	 * @param NewMaxAttributeValueOut new max attribute value after compacting.
 	 * @param OldToNewMapOut generated mapping from previous IDs to new IDs (ie size is maximum input attrib value (+1), will contain InvalidID for any compacted values )
 	 * @param NewToOldMapOut generated mapping from new IDs to previous IDs (ie size is maximum output attribute value (+1), no invalid values)
 	 * @param bWasCompactOut set to true if the attribute set was already compact, IE was not modified by the operation
@@ -25,6 +31,8 @@ namespace Geometry
 	GEOMETRYCORE_API bool CompactAttributeValues(
 		const FDynamicMesh3& Mesh,
 		TDynamicMeshScalarTriangleAttribute<int32>& TriangleAttrib,
+		FInterval1i& OldMaxAttributeRangeOut,
+		int& NewMaxAttributeValueOut,
 		TArray<int32>& OldToNewMapOut,
 		TArray<int32>& NewToOldMapOut,
 		bool& bWasCompactOut);
