GeometryScript: add new BP functions SplitMeshByPolygroups, GetAllTrianglePolygroupIDs, GetPolygroupIDsInMesh, GetTrianglesInPolygroup
#rnx

diff --git a/Engine/Plugins/Experimental/GeometryScripting/Source/GeometryScriptingCore/Private/MeshDecompositionFunctions.cpp b/Engine/Plugins/Experimental/GeometryScripting/Source/GeometryScriptingCore/Private/MeshDecompositionFunctions.cpp
--- a/Engine/Plugins/Experimental/GeometryScripting/Source/GeometryScriptingCore/Private/MeshDecompositionFunctions.cpp
+++ b/Engine/Plugins/Experimental/GeometryScripting/Source/GeometryScriptingCore/Private/MeshDecompositionFunctions.cpp
@@ -5,6 +5,7 @@
 #include "DynamicMesh/DynamicMesh3.h"
 #include "DynamicMesh/DynamicMeshAttributeSet.h"
 #include "Selections/MeshConnectedComponents.h"
+#include "Polygroups/PolygroupSet.h"
 #include "DynamicMeshEditor.h"
 #include "UDynamicMesh.h"
 
@@ -134,6 +135,64 @@ UDynamicMesh* UGeometryScriptLibrary_MeshDecompositionFunctions::SplitMeshByMate
 }
 
 
+
+
+UDynamicMesh* UGeometryScriptLibrary_MeshDecompositionFunctions::SplitMeshByPolygroups(  
+	UDynamicMesh* TargetMesh, 
+	FGeometryScriptGroupLayer GroupLayer,
+	TArray<UDynamicMesh*>& ComponentMeshes,
+	TArray<int>& ComponentPolygroups,
+	UDynamicMeshPool* MeshPool,
+	UGeometryScriptDebug* Debug)
+{
+	ComponentMeshes.Reset();
+	ComponentPolygroups.Reset();
+
+	if (TargetMesh == nullptr)
+	{
+		UE::Geometry::AppendError(Debug, EGeometryScriptErrorType::InvalidInputs, LOCTEXT("SplitMeshByPolygroups_InvalidInput", "SplitMeshByPolygroups: TargetMesh is Null"));
+		return TargetMesh;
+	}
+
+	TArray<FDynamicMesh3> SplitMeshes;
+	TargetMesh->ProcessMesh([&](const FDynamicMesh3& EditMesh)
+	{
+		TUniquePtr<FPolygroupSet> SplitGroups;
+		FPolygroupLayer InputGroupLayer{ GroupLayer.bDefaultLayer, GroupLayer.ExtendedLayerIndex };
+		if (InputGroupLayer.CheckExists(&EditMesh) == false)
+		{
+			UE::Geometry::AppendError(Debug, EGeometryScriptErrorType::InvalidInputs, LOCTEXT("SplitMeshByPolygroups_MissingGroups", "SplitMeshByPolygroups: Target Polygroup Layer does not exist"));
+			return;
+		}
+		SplitGroups = MakeUnique<FPolygroupSet>(&EditMesh, InputGroupLayer);
+
+		for (int32 tid : EditMesh.TriangleIndicesItr())
+		{
+			int32 GroupID = SplitGroups->GetGroup(tid);
+			ComponentPolygroups.AddUnique(GroupID);
+		}
+		ComponentPolygroups.Sort();
+
+		TMap<int32, int32> PolygroupMap;
+		for (int32 k = 0; k < ComponentPolygroups.Num(); ++k)
+		{
+			PolygroupMap.Add(ComponentPolygroups[k], k);
+		}
+
+		FDynamicMeshEditor::SplitMesh(&EditMesh, SplitMeshes, [&](int32 tid) {
+			int32 GroupID = SplitGroups->GetGroup(tid);
+			const int32* FoundIndex = PolygroupMap.Find(GroupID);
+			return (FoundIndex == nullptr) ? 0 : *FoundIndex;
+		});
+	});
+
+	BuildNewDynamicMeshes(TargetMesh, MeshPool, SplitMeshes, ComponentMeshes);
+
+	return TargetMesh;
+}
+
+
+
 UDynamicMesh* UGeometryScriptLibrary_MeshDecompositionFunctions::GetSubMeshFromMesh(
 	UDynamicMesh* TargetMesh,
 	UDynamicMesh* StoreToSubmesh, 
diff --git a/Engine/Plugins/Experimental/GeometryScripting/Source/GeometryScriptingCore/Private/MeshPolygroupFunctions.cpp b/Engine/Plugins/Experimental/GeometryScripting/Source/GeometryScriptingCore/Private/MeshPolygroupFunctions.cpp
--- a/Engine/Plugins/Experimental/GeometryScripting/Source/GeometryScriptingCore/Private/MeshPolygroupFunctions.cpp
+++ b/Engine/Plugins/Experimental/GeometryScripting/Source/GeometryScriptingCore/Private/MeshPolygroupFunctions.cpp
@@ -350,4 +350,86 @@ UDynamicMesh* UGeometryScriptLibrary_MeshPolygroupFunctions::DeleteTrianglesInPo
 
 
 
+UDynamicMesh* UGeometryScriptLibrary_MeshPolygroupFunctions::GetAllTrianglePolygroupIDs(
+	UDynamicMesh* TargetMesh, 
+	FGeometryScriptGroupLayer GroupLayer, 
+	FGeometryScriptIndexList& PolygroupIDsOut)
+{
+	PolygroupIDsOut.Reset();
+	PolygroupIDsOut.IndexType = EGeometryScriptIndexType::PolygroupID;
+	TArray<int32>& PolygroupIDs = *PolygroupIDsOut.List;
+
+	bool bIsValidPolygroupLayer = false;
+	SimpleMeshPolygroupQuery<int32>(TargetMesh, GroupLayer, bIsValidPolygroupLayer, 0, 
+		[&](const FDynamicMesh3& Mesh, const FPolygroupSet& PolyGroups) {
+		int32 MaxTriangleID = Mesh.MaxTriangleID();
+		for (int32 TriangleID = 0; TriangleID < MaxTriangleID; ++TriangleID)
+		{
+			int32 GroupID = Mesh.IsTriangle(TriangleID) ? PolyGroups.GetGroup(TriangleID) : -1;
+			PolygroupIDs.Add(GroupID);
+		}
+		return 0;
+	});
+	return TargetMesh;
+}
+
+
+UDynamicMesh* UGeometryScriptLibrary_MeshPolygroupFunctions::GetPolygroupIDsInMesh(
+	UDynamicMesh* TargetMesh, 
+	FGeometryScriptGroupLayer GroupLayer, 
+	FGeometryScriptIndexList& PolygroupIDsOut)
+{
+	PolygroupIDsOut.Reset();
+	PolygroupIDsOut.IndexType = EGeometryScriptIndexType::PolygroupID;
+	TArray<int32>& PolygroupIDs = *PolygroupIDsOut.List;
+
+	TSet<int32> UniqueGroupIDs;
+	bool bIsValidPolygroupLayer = false;
+	SimpleMeshPolygroupQuery<int32>(TargetMesh, GroupLayer, bIsValidPolygroupLayer, 0, 
+		[&](const FDynamicMesh3& Mesh, const FPolygroupSet& PolyGroups) {
+		int32 MaxTriangleID = Mesh.MaxTriangleID();
+		for (int32 tid : Mesh.TriangleIndicesItr())
+		{
+			int32 GroupID = PolyGroups.GetGroup(tid);
+			if (UniqueGroupIDs.Contains(GroupID) == false)
+			{
+				UniqueGroupIDs.Add(GroupID);
+				PolygroupIDs.Add(GroupID);
+			}
+		}
+		return 0;
+	});
+	return TargetMesh;
+}
+
+
+
+UDynamicMesh* UGeometryScriptLibrary_MeshPolygroupFunctions::GetTrianglesInPolygroup(
+	UDynamicMesh* TargetMesh,
+	FGeometryScriptGroupLayer GroupLayer,
+	int PolygroupID,
+	FGeometryScriptIndexList& TriangleIDsOut)
+{
+	TriangleIDsOut.Reset();
+	TriangleIDsOut.IndexType = EGeometryScriptIndexType::Triangle;
+	TArray<int32>& TriangleIDs = *TriangleIDsOut.List;
+
+	bool bIsValidPolygroupLayer = false;
+	SimpleMeshPolygroupQuery<int32>(TargetMesh, GroupLayer, bIsValidPolygroupLayer, 0, 
+		[&](const FDynamicMesh3& Mesh, const FPolygroupSet& PolyGroups) {
+		for (int32 TriangleID : Mesh.TriangleIndicesItr() )
+		{
+			if (PolyGroups.GetGroup(TriangleID) == PolygroupID)
+			{
+				TriangleIDs.Add(TriangleID);
+			}
+		}
+		return 0;
+	});
+	return TargetMesh;
+}
+
+
+
+
 #undef LOCTEXT_NAMESPACE
\ No newline at end of file
diff --git a/Engine/Plugins/Experimental/GeometryScripting/Source/GeometryScriptingCore/Public/GeometryScript/MeshDecompositionFunctions.h b/Engine/Plugins/Experimental/GeometryScripting/Source/GeometryScriptingCore/Public/GeometryScript/MeshDecompositionFunctions.h
--- a/Engine/Plugins/Experimental/GeometryScripting/Source/GeometryScriptingCore/Public/GeometryScript/MeshDecompositionFunctions.h
+++ b/Engine/Plugins/Experimental/GeometryScripting/Source/GeometryScriptingCore/Public/GeometryScript/MeshDecompositionFunctions.h
@@ -33,6 +33,16 @@ public:
 		UDynamicMeshPool* MeshPool,
 		UGeometryScriptDebug* Debug = nullptr);
 
+	UFUNCTION(BlueprintCallable, Category = "GeometryScript|Decomposition", meta=(ScriptMethod))
+	static UPARAM(DisplayName = "Target Mesh") UDynamicMesh* 
+	SplitMeshByPolygroups(  
+		UDynamicMesh* TargetMesh, 
+		FGeometryScriptGroupLayer GroupLayer,
+		TArray<UDynamicMesh*>& ComponentMeshes,
+		TArray<int>& ComponentPolygroups,
+		UDynamicMeshPool* MeshPool,
+		UGeometryScriptDebug* Debug = nullptr);
+
 	UFUNCTION(BlueprintCallable, Category = "GeometryScript|Decomposition", meta=(ScriptMethod))
 	static UPARAM(DisplayName = "Target Mesh") UDynamicMesh* 
 	GetSubMeshFromMesh(  
diff --git a/Engine/Plugins/Experimental/GeometryScripting/Source/GeometryScriptingCore/Public/GeometryScript/MeshPolygroupFunctions.h b/Engine/Plugins/Experimental/GeometryScripting/Source/GeometryScriptingCore/Public/GeometryScript/MeshPolygroupFunctions.h
--- a/Engine/Plugins/Experimental/GeometryScripting/Source/GeometryScriptingCore/Public/GeometryScript/MeshPolygroupFunctions.h
+++ b/Engine/Plugins/Experimental/GeometryScripting/Source/GeometryScriptingCore/Public/GeometryScript/MeshPolygroupFunctions.h
@@ -87,4 +87,36 @@ public:
 		bool bDeferChangeNotifications = false,
 		UGeometryScriptDebug* Debug = nullptr);
 
+	/**
+	* Create list of per-triangle Polygroup IDs for the Polygroup in the Mesh
+	* @warning if the mesh is not Triangle-Compact (eg GetHasTriangleIDGaps == false) then the returned list will also have the same gaps
+	*/
+	UFUNCTION(BlueprintCallable, Category = "GeometryScript|Polygroups", meta=(ScriptMethod))
+	static UPARAM(DisplayName = "Target Mesh") UDynamicMesh* 
+	GetAllTrianglePolygroupIDs( 
+		UDynamicMesh* TargetMesh, 
+		FGeometryScriptGroupLayer GroupLayer, 
+		UPARAM(ref, DisplayName="Polygroup IDs Out") FGeometryScriptIndexList& PolygroupIDsOut );
+
+	/**
+	* Create list of all unique Polygroup IDs that exist in the Polygroup Layer in the Mesh
+	*/
+	UFUNCTION(BlueprintCallable, Category = "GeometryScript|Polygroups", meta=(ScriptMethod))
+	static UPARAM(DisplayName = "Target Mesh") UDynamicMesh* 
+	GetPolygroupIDsInMesh( 
+		UDynamicMesh* TargetMesh, 
+		FGeometryScriptGroupLayer GroupLayer, 
+		UPARAM(ref, DisplayName="Polygroup IDs Out") FGeometryScriptIndexList& PolygroupIDsOut );
+
+	/**
+	 * Create list of all triangles with the given Polygroup ID in the given GroupLayer (not necessarily a single connected-component)
+	 */
+	UFUNCTION(BlueprintCallable, Category = "GeometryScript|Polygroups", meta=(ScriptMethod))
+	static UPARAM(DisplayName = "Target Mesh") UDynamicMesh* 
+	GetTrianglesInPolygroup( 
+		UDynamicMesh* TargetMesh, 
+		FGeometryScriptGroupLayer GroupLayer, 
+		int PolygroupID, 
+		UPARAM(ref, DisplayName="Triangle IDs Out") FGeometryScriptIndexList& TriangleIDsOut );
+
 };
\ No newline at end of file
