GeometryScript: Set IndexType member of FGeometryScriptIndexList and use it for error checking, and rename EGeometryScriptIndexType::Unknown to EGeometryScriptIndexType::Any

#rnx

diff --git a/Engine/Plugins/Experimental/GeometryScripting/Source/GeometryScriptingCore/Private/ListUtilityFunctions.cpp b/Engine/Plugins/Experimental/GeometryScripting/Source/GeometryScriptingCore/Private/ListUtilityFunctions.cpp
--- a/Engine/Plugins/Experimental/GeometryScripting/Source/GeometryScriptingCore/Private/ListUtilityFunctions.cpp
+++ b/Engine/Plugins/Experimental/GeometryScripting/Source/GeometryScriptingCore/Private/ListUtilityFunctions.cpp
@@ -36,9 +36,9 @@ void UGeometryScriptLibrary_ListUtilityFunctions::ConvertIndexListToArray(FGeome
 	}
 }
 
-void UGeometryScriptLibrary_ListUtilityFunctions::ConvertArrayToIndexList(const TArray<int>& IndexArray, FGeometryScriptIndexList& IndexList)
+void UGeometryScriptLibrary_ListUtilityFunctions::ConvertArrayToIndexList(const TArray<int>& IndexArray, FGeometryScriptIndexList& IndexList, EGeometryScriptIndexType IndexType)
 {
-	IndexList.Reset();
+	IndexList.Reset(IndexType);
 	IndexList.List->Append(IndexArray);
 }
 
diff --git a/Engine/Plugins/Experimental/GeometryScripting/Source/GeometryScriptingCore/Private/MeshBasicEditFunctions.cpp b/Engine/Plugins/Experimental/GeometryScripting/Source/GeometryScriptingCore/Private/MeshBasicEditFunctions.cpp
--- a/Engine/Plugins/Experimental/GeometryScripting/Source/GeometryScriptingCore/Private/MeshBasicEditFunctions.cpp
+++ b/Engine/Plugins/Experimental/GeometryScripting/Source/GeometryScriptingCore/Private/MeshBasicEditFunctions.cpp
@@ -79,7 +79,7 @@ UDynamicMesh* UGeometryScriptLibrary_MeshBasicEditFunctions::AddVerticesToMesh(
 		return TargetMesh;
 	}
 
-	NewIndicesList.Reset();
+	NewIndicesList.Reset(EGeometryScriptIndexType::Vertex);
 	if (TargetMesh)
 	{
 		TargetMesh->EditMesh([&](FDynamicMesh3& EditMesh)
@@ -122,7 +122,7 @@ UDynamicMesh* UGeometryScriptLibrary_MeshBasicEditFunctions::DeleteVerticesFromM
 	int& NumDeleted,
 	bool bDeferChangeNotifications)
 {
-	if (VertexList.List.IsValid() == false || VertexList.List->Num() == 0)
+	if (VertexList.List.IsValid() == false || VertexList.List->Num() == 0 || VertexList.IsCompatibleWith(EGeometryScriptIndexType::Vertex) == false)
 	{
 		return TargetMesh;
 	}
@@ -207,7 +207,7 @@ UDynamicMesh* UGeometryScriptLibrary_MeshBasicEditFunctions::AddTrianglesToMesh(
 		return TargetMesh;
 	}
 
-	NewIndicesList.Reset();
+	NewIndicesList.Reset(EGeometryScriptIndexType::Triangle);
 	if (TargetMesh)
 	{
 		TargetMesh->EditMesh([&](FDynamicMesh3& EditMesh)
@@ -403,7 +403,7 @@ UDynamicMesh* UGeometryScriptLibrary_MeshBasicEditFunctions::AppendBuffersToMesh
 	bool bDeferChangeNotifications,
 	UGeometryScriptDebug* Debug)
 {
-	NewTriangleIndicesList.Reset();
+	NewTriangleIndicesList.Reset(EGeometryScriptIndexType::Triangle);
 	if (TargetMesh == nullptr)
 	{
 		UE::Geometry::AppendError(Debug, EGeometryScriptErrorType::InvalidInputs, LOCTEXT("AppendMeshRepeated_InvalidInput1", "AppendMeshRepeated: TargetMesh is Null"));
diff --git a/Engine/Plugins/Experimental/GeometryScripting/Source/GeometryScriptingCore/Private/MeshDecompositionFunctions.cpp b/Engine/Plugins/Experimental/GeometryScripting/Source/GeometryScriptingCore/Private/MeshDecompositionFunctions.cpp
--- a/Engine/Plugins/Experimental/GeometryScripting/Source/GeometryScriptingCore/Private/MeshDecompositionFunctions.cpp
+++ b/Engine/Plugins/Experimental/GeometryScripting/Source/GeometryScriptingCore/Private/MeshDecompositionFunctions.cpp
@@ -215,6 +215,11 @@ UDynamicMesh* UGeometryScriptLibrary_MeshDecompositionFunctions::GetSubMeshFromM
 		UE::Geometry::AppendError(Debug, EGeometryScriptErrorType::InvalidInputs, LOCTEXT("SetMeshPerVertexColors_InvalidList", "GetSubMeshFromMesh: TriangleList is empty"));
 		return TargetMesh;
 	}
+	if (TriangleList.IsCompatibleWith(EGeometryScriptIndexType::Triangle) == false)
+	{
+		UE::Geometry::AppendError(Debug, EGeometryScriptErrorType::InvalidInputs, LOCTEXT("SetMeshPerVertexColors_InvalidList2", "GetSubMeshFromMesh: TriangleList has incompatible index type"));
+		return TargetMesh;
+	}
 
 	FDynamicMesh3 Submesh;
 	TargetMesh->ProcessMesh([&](const FDynamicMesh3& EditMesh)
diff --git a/Engine/Plugins/Experimental/GeometryScripting/Source/GeometryScriptingCore/Private/MeshMaterialFunctions.cpp b/Engine/Plugins/Experimental/GeometryScripting/Source/GeometryScriptingCore/Private/MeshMaterialFunctions.cpp
--- a/Engine/Plugins/Experimental/GeometryScripting/Source/GeometryScriptingCore/Private/MeshMaterialFunctions.cpp
+++ b/Engine/Plugins/Experimental/GeometryScripting/Source/GeometryScriptingCore/Private/MeshMaterialFunctions.cpp
@@ -189,8 +189,7 @@ int32 UGeometryScriptLibrary_MeshMaterialFunctions::GetTriangleMaterialID(
 
 UDynamicMesh* UGeometryScriptLibrary_MeshMaterialFunctions::GetAllTriangleMaterialIDs(UDynamicMesh* TargetMesh, FGeometryScriptIndexList& MaterialIDList, bool& bHasMaterialIDs)
 {
-	MaterialIDList.Reset();
-	MaterialIDList.IndexType = EGeometryScriptIndexType::MaterialID;
+	MaterialIDList.Reset(EGeometryScriptIndexType::MaterialID);
 	TArray<int32>& MaterialIDs = *MaterialIDList.List;
 	bHasMaterialIDs = false;
 	SimpleMeshMaterialQuery<int32>(TargetMesh, bHasMaterialIDs, 0, [&](const FDynamicMesh3& Mesh, const FDynamicMeshMaterialAttribute& MaterialIDAttrib) {
@@ -248,7 +247,12 @@ UDynamicMesh* UGeometryScriptLibrary_MeshMaterialFunctions::SetAllTriangleMateri
 	}
 	if (TriangleMaterialIDList.List.IsValid() == false || TriangleMaterialIDList.List->Num() == 0)
 	{
-		UE::Geometry::AppendError(Debug, EGeometryScriptErrorType::InvalidInputs, LOCTEXT("SetAllTriangleMaterialIDs_InvalidList", "SetAllTriangleMaterialIDs: List is empty"));
+		UE::Geometry::AppendError(Debug, EGeometryScriptErrorType::InvalidInputs, LOCTEXT("SetAllTriangleMaterialIDs_InvalidList", "SetAllTriangleMaterialIDs: TriangleMaterialIDList is empty"));
+		return TargetMesh;
+	}
+	if (TriangleMaterialIDList.IsCompatibleWith(EGeometryScriptIndexType::MaterialID) == false)
+	{
+		UE::Geometry::AppendError(Debug, EGeometryScriptErrorType::InvalidInputs, LOCTEXT("SetAllTriangleMaterialIDs_InvalidList2", "SetAllTriangleMaterialIDs: TriangleMaterialIDList has incompatible index type"));
 		return TargetMesh;
 	}
 
diff --git a/Engine/Plugins/Experimental/GeometryScripting/Source/GeometryScriptingCore/Private/MeshPolygroupFunctions.cpp b/Engine/Plugins/Experimental/GeometryScripting/Source/GeometryScriptingCore/Private/MeshPolygroupFunctions.cpp
--- a/Engine/Plugins/Experimental/GeometryScripting/Source/GeometryScriptingCore/Private/MeshPolygroupFunctions.cpp
+++ b/Engine/Plugins/Experimental/GeometryScripting/Source/GeometryScriptingCore/Private/MeshPolygroupFunctions.cpp
@@ -257,7 +257,7 @@ UDynamicMesh* UGeometryScriptLibrary_MeshPolygroupFunctions::ComputePolygroupsFr
 {
 	if (TargetMesh == nullptr)
 	{
-		UE::Geometry::AppendError(Debug, EGeometryScriptErrorType::InvalidInputs, LOCTEXT("ComputePolygroupsFromAngleThresholde_InvalidInput", "ComputePolygroupsFromAngleThreshold: TargetMesh is Null"));
+		UE::Geometry::AppendError(Debug, EGeometryScriptErrorType::InvalidInputs, LOCTEXT("ComputePolygroupsFromAngleThreshold_InvalidInput", "ComputePolygroupsFromAngleThreshold: TargetMesh is Null"));
 		return TargetMesh;
 	}
 
@@ -355,8 +355,7 @@ UDynamicMesh* UGeometryScriptLibrary_MeshPolygroupFunctions::GetAllTrianglePolyg
 	FGeometryScriptGroupLayer GroupLayer, 
 	FGeometryScriptIndexList& PolygroupIDsOut)
 {
-	PolygroupIDsOut.Reset();
-	PolygroupIDsOut.IndexType = EGeometryScriptIndexType::PolygroupID;
+	PolygroupIDsOut.Reset(EGeometryScriptIndexType::PolygroupID);
 	TArray<int32>& PolygroupIDs = *PolygroupIDsOut.List;
 
 	bool bIsValidPolygroupLayer = false;
@@ -379,8 +378,7 @@ UDynamicMesh* UGeometryScriptLibrary_MeshPolygroupFunctions::GetPolygroupIDsInMe
 	FGeometryScriptGroupLayer GroupLayer, 
 	FGeometryScriptIndexList& PolygroupIDsOut)
 {
-	PolygroupIDsOut.Reset();
-	PolygroupIDsOut.IndexType = EGeometryScriptIndexType::PolygroupID;
+	PolygroupIDsOut.Reset(EGeometryScriptIndexType::PolygroupID);
 	TArray<int32>& PolygroupIDs = *PolygroupIDsOut.List;
 
 	TSet<int32> UniqueGroupIDs;
@@ -410,8 +408,7 @@ UDynamicMesh* UGeometryScriptLibrary_MeshPolygroupFunctions::GetTrianglesInPolyg
 	int PolygroupID,
 	FGeometryScriptIndexList& TriangleIDsOut)
 {
-	TriangleIDsOut.Reset();
-	TriangleIDsOut.IndexType = EGeometryScriptIndexType::Triangle;
+	TriangleIDsOut.Reset(EGeometryScriptIndexType::Triangle);
 	TArray<int32>& TriangleIDs = *TriangleIDsOut.List;
 
 	bool bIsValidPolygroupLayer = false;
diff --git a/Engine/Plugins/Experimental/GeometryScripting/Source/GeometryScriptingCore/Private/MeshQueryFunctions.cpp b/Engine/Plugins/Experimental/GeometryScripting/Source/GeometryScriptingCore/Private/MeshQueryFunctions.cpp
--- a/Engine/Plugins/Experimental/GeometryScripting/Source/GeometryScriptingCore/Private/MeshQueryFunctions.cpp
+++ b/Engine/Plugins/Experimental/GeometryScripting/Source/GeometryScriptingCore/Private/MeshQueryFunctions.cpp
@@ -142,7 +142,7 @@ bool UGeometryScriptLibrary_MeshQueryFunctions::IsValidTriangleID( UDynamicMesh*
 
 UDynamicMesh* UGeometryScriptLibrary_MeshQueryFunctions::GetAllTriangleIDs(UDynamicMesh* TargetMesh, FGeometryScriptIndexList& TriangleIDList, bool& bHasTriangleIDGaps)
 {
-	TriangleIDList.Reset();
+	TriangleIDList.Reset(EGeometryScriptIndexType::Triangle);
 	bHasTriangleIDGaps = false;
 	if (TargetMesh)
 	{
@@ -250,7 +250,7 @@ bool UGeometryScriptLibrary_MeshQueryFunctions::IsValidVertexID( UDynamicMesh* T
 
 UDynamicMesh* UGeometryScriptLibrary_MeshQueryFunctions::GetAllVertexIDs(UDynamicMesh* TargetMesh, FGeometryScriptIndexList& VertexIDList, bool& bHasVertexIDGaps)
 {
-	VertexIDList.Reset();
+	VertexIDList.Reset(EGeometryScriptIndexType::Vertex);
 	bHasVertexIDGaps = false;
 	if (TargetMesh)
 	{
diff --git a/Engine/Plugins/Experimental/GeometryScripting/Source/GeometryScriptingCore/Public/GeometryScript/GeometryScriptTypes.h b/Engine/Plugins/Experimental/GeometryScripting/Source/GeometryScriptingCore/Public/GeometryScript/GeometryScriptTypes.h
--- a/Engine/Plugins/Experimental/GeometryScripting/Source/GeometryScriptingCore/Public/GeometryScript/GeometryScriptTypes.h
+++ b/Engine/Plugins/Experimental/GeometryScripting/Source/GeometryScriptingCore/Public/GeometryScript/GeometryScriptTypes.h
@@ -180,7 +180,8 @@ public:
 UENUM(BlueprintType)
 enum class EGeometryScriptIndexType : uint8
 {
-	Unknown,
+	// Index lists of Any type are compatible with any other index list type
+	Any,
 	Triangle,
 	Vertex,
 	MaterialID,
@@ -193,17 +194,23 @@ struct GEOMETRYSCRIPTINGCORE_API FGeometryScriptIndexList
 	GENERATED_BODY()
 
 	UPROPERTY(BlueprintReadWrite, Category = IndexList)
-	EGeometryScriptIndexType IndexType = EGeometryScriptIndexType::Unknown;
+	EGeometryScriptIndexType IndexType = EGeometryScriptIndexType::Any;
 public:
 	TSharedPtr<TArray<int>> List;
 
-	void Reset()
+	void Reset(EGeometryScriptIndexType TargetIndexType)
 	{
 		if (List.IsValid() == false)
 		{
 			List = MakeShared<TArray<int>>();
 		}
 		List->Reset();
+		IndexType = TargetIndexType;
+	}
+
+	bool IsCompatibleWith(EGeometryScriptIndexType OtherType) const
+	{
+		return IndexType == OtherType || IndexType == EGeometryScriptIndexType::Any;
 	}
 };
 
diff --git a/Engine/Plugins/Experimental/GeometryScripting/Source/GeometryScriptingCore/Public/GeometryScript/ListUtilityFunctions.h b/Engine/Plugins/Experimental/GeometryScripting/Source/GeometryScriptingCore/Public/GeometryScript/ListUtilityFunctions.h
--- a/Engine/Plugins/Experimental/GeometryScripting/Source/GeometryScriptingCore/Public/GeometryScript/ListUtilityFunctions.h
+++ b/Engine/Plugins/Experimental/GeometryScripting/Source/GeometryScriptingCore/Public/GeometryScript/ListUtilityFunctions.h
@@ -29,7 +29,7 @@ public:
 	static void ConvertIndexListToArray(FGeometryScriptIndexList IndexList, TArray<int>& IndexArray);
 
 	UFUNCTION(BlueprintCallable, Category = "GeometryScript|ListUtils")
-	static void ConvertArrayToIndexList(const TArray<int>& IndexArray, FGeometryScriptIndexList& IndexList);
+	static void ConvertArrayToIndexList(const TArray<int>& IndexArray, FGeometryScriptIndexList& IndexList, EGeometryScriptIndexType IndexType = EGeometryScriptIndexType::Any);
 
 
 
diff --git a/Engine/Plugins/Experimental/GeometryScripting/Source/GeometryScriptingCore/Public/GeometryScript/MeshPolygroupFunctions.h b/Engine/Plugins/Experimental/GeometryScripting/Source/GeometryScriptingCore/Public/GeometryScript/MeshPolygroupFunctions.h
--- a/Engine/Plugins/Experimental/GeometryScripting/Source/GeometryScriptingCore/Public/GeometryScript/MeshPolygroupFunctions.h
+++ b/Engine/Plugins/Experimental/GeometryScripting/Source/GeometryScriptingCore/Public/GeometryScript/MeshPolygroupFunctions.h
@@ -70,7 +70,7 @@ public:
 
 
 	UFUNCTION(BlueprintPure, Category = "GeometryScript|Polygroups", meta=(ScriptMethod))
-	static UPARAM(DisplayName = "Material ID") int32
+	static UPARAM(DisplayName = "Polygroup ID") int32
 	GetTrianglePolygroupID( 
 		UDynamicMesh* TargetMesh, 
 		FGeometryScriptGroupLayer GroupLayer, 
