GeometryScript: Add some BP functions related to selections based on materials

#rnx

diff --git a/Engine/Plugins/Experimental/GeometryScripting/Source/GeometryScriptingCore/Private/GeometryScriptTypes.cpp b/Engine/Plugins/Experimental/GeometryScripting/Source/GeometryScriptingCore/Private/GeometryScriptTypes.cpp
--- a/Engine/Plugins/Experimental/GeometryScripting/Source/GeometryScriptingCore/Private/GeometryScriptTypes.cpp
+++ b/Engine/Plugins/Experimental/GeometryScripting/Source/GeometryScriptingCore/Private/GeometryScriptTypes.cpp
@@ -12,6 +12,17 @@ FGeometryScriptDebugMessage UE::Geometry::MakeScriptError(EGeometryScriptErrorTy
 	return FGeometryScriptDebugMessage{ EGeometryScriptDebugMessageType::ErrorMessage, ErrorTypeIn, MessageIn };
 }
 
+FGeometryScriptDebugMessage UE::Geometry::MakeScriptWarning(EGeometryScriptErrorType WarningTypeIn, const FText& MessageIn)
+{
+	UE_LOG(LogGeometry, Warning, TEXT("GeometryScriptWarning: %s"), *MessageIn.ToString() );
+
+	return FGeometryScriptDebugMessage{ EGeometryScriptDebugMessageType::WarningMessage, WarningTypeIn, MessageIn };
+}
+
+
+
+
+
 
 void UE::Geometry::AppendError(UGeometryScriptDebug* Debug, EGeometryScriptErrorType ErrorTypeIn, const FText& MessageIn)
 {
@@ -20,4 +31,13 @@ void UE::Geometry::AppendError(UGeometryScriptDebug* Debug, EGeometryScriptError
 	{
 		Debug->Append(Result);
 	}
+}
+
+void UE::Geometry::AppendWarning(UGeometryScriptDebug* Debug, EGeometryScriptErrorType WarningTypeIn, const FText& MessageIn)
+{
+	FGeometryScriptDebugMessage Result = MakeScriptWarning(WarningTypeIn, MessageIn);
+	if (Debug != nullptr)
+	{
+		Debug->Append(Result);
+	}
 }
\ No newline at end of file
diff --git a/Engine/Plugins/Experimental/GeometryScripting/Source/GeometryScriptingCore/Private/MeshMaterialFunctions.cpp b/Engine/Plugins/Experimental/GeometryScripting/Source/GeometryScriptingCore/Private/MeshMaterialFunctions.cpp
--- a/Engine/Plugins/Experimental/GeometryScripting/Source/GeometryScriptingCore/Private/MeshMaterialFunctions.cpp
+++ b/Engine/Plugins/Experimental/GeometryScripting/Source/GeometryScriptingCore/Private/MeshMaterialFunctions.cpp
@@ -159,6 +159,67 @@ UDynamicMesh* UGeometryScriptLibrary_MeshMaterialFunctions::RemapMaterialIDs(
 }
 
 
+UDynamicMesh* UGeometryScriptLibrary_MeshMaterialFunctions::GetMaterialIDsOfTriangles( 
+	UDynamicMesh* TargetMesh, 
+	FGeometryScriptIndexList TriangleIDList,
+	FGeometryScriptIndexList& MaterialIDList,
+    UGeometryScriptDebug* Debug)
+{
+	MaterialIDList.Reset(EGeometryScriptIndexType::MaterialID);
+	MaterialIDList.List->SetNumZeroed(TriangleIDList.List->Num());
+	
+	if (TargetMesh == nullptr)
+	{
+		AppendError(Debug, EGeometryScriptErrorType::InvalidInputs, LOCTEXT("GetMaterialIDsOfTriangles_InvalidMesh", "GetMaterialIDsOfTriangles: TargetMesh is Null"));
+		return TargetMesh;
+	}
+	if (TriangleIDList.List.IsValid() == false)
+	{
+		AppendError(Debug, EGeometryScriptErrorType::InvalidInputs, LOCTEXT("GetMaterialIDsOfTriangles_InvalidList", "GetMaterialIDsOfTriangles: TriangleIDList is Null"));
+		return TargetMesh;
+	}
+	if (TriangleIDList.IsCompatibleWith(EGeometryScriptIndexType::Triangle) == false)
+	{
+		AppendError(Debug, EGeometryScriptErrorType::InvalidInputs, LOCTEXT("GetMaterialIDsOfTriangles_InvalidList2", "GetMaterialIDsOfTriangles: TriangleIDList has incompatible index type"));
+		return TargetMesh;
+	}
+	if (TriangleIDList.List->Num() == 0 && TargetMesh->IsEmpty() == false)
+	{
+		AppendWarning(Debug, EGeometryScriptErrorType::InvalidInputs, LOCTEXT("GetMaterialIDsOfTriangles_EmptyList", "GetMaterialIDsOfTriangles: TriangleIDList is empty"));
+	}
+	
+	bool bHasMaterials = false;
+	bool bAllValidTriangles = true;
+	SimpleMeshMaterialQuery<int32>(TargetMesh, bHasMaterials, 0, [&](const FDynamicMesh3& EditMesh, const FDynamicMeshMaterialAttribute& MaterialIDs) {
+		TArray<int>& MaterialIDArray = *MaterialIDList.List;
+		for (int32 i = 0; i < TriangleIDList.List->Num(); i++)
+		{
+			int32 TriangleID = (*TriangleIDList.List)[i];
+			if (EditMesh.IsTriangle(TriangleID) == false)
+			{
+				bAllValidTriangles = false;
+				return 0;
+			}
+			MaterialIDArray[i] = MaterialIDs.GetValue(TriangleID);
+		}
+		return 0;
+	});
+	
+	if (bHasMaterials == false)
+	{
+		AppendError(Debug, EGeometryScriptErrorType::InvalidInputs, LOCTEXT("GetMaterialIDsOfTriangles_MissingMaterials", "GetMaterialIDsOfTriangles: MaterialID Attribute is not enabled"));
+		return TargetMesh;
+	}
+	
+	if (bAllValidTriangles == false)
+	{
+		AppendError(Debug, EGeometryScriptErrorType::InvalidInputs, LOCTEXT("GetMaterialIDsOfTriangles_InvalidTriangles", "GetMaterialIDsOfTriangles: TriangleIDList has invalid triangles"));
+		return TargetMesh;
+	}
+	
+	return TargetMesh;
+}
+
 
 int UGeometryScriptLibrary_MeshMaterialFunctions::GetMaxMaterialID( UDynamicMesh* TargetMesh, bool& bHasMaterialIDs )
 {
@@ -205,6 +266,54 @@ UDynamicMesh* UGeometryScriptLibrary_MeshMaterialFunctions::GetAllTriangleMateri
 }
 
 
+UDynamicMesh* UGeometryScriptLibrary_MeshMaterialFunctions::GetTrianglesByMaterialID( 
+		UDynamicMesh* TargetMesh, 
+		int MaterialID,
+		FGeometryScriptIndexList& TriangleIDList,
+		UGeometryScriptDebug* Debug)
+{
+	TriangleIDList.Reset(EGeometryScriptIndexType::Triangle);
+
+	if (TargetMesh == nullptr)
+	{
+		AppendError(Debug, EGeometryScriptErrorType::InvalidInputs, LOCTEXT("GetTrianglesByMaterialID_InvalidInput", "GetTrianglesByMaterialID: TargetMesh is Null"));
+		return TargetMesh;
+	}
+
+	bool bHasMaterialIDs = false;
+	bool bAllValidTriangles = true;
+	SimpleMeshMaterialQuery<int32>(TargetMesh, bHasMaterialIDs, 0, [&](const FDynamicMesh3& Mesh, const FDynamicMeshMaterialAttribute& MaterialIDAttrib) {
+		for (int32 TriangleID = 0; TriangleID < Mesh.MaxTriangleID(); ++TriangleID)
+		{
+			if (Mesh.IsTriangle(TriangleID) == false)
+			{
+				bAllValidTriangles = false;
+				continue;
+			}
+
+			if (MaterialIDAttrib.GetValue(TriangleID) == MaterialID)
+			{
+				TriangleIDList.List->Add(TriangleID);
+			}
+		}
+		return 0;
+	});
+
+	if (bHasMaterialIDs == false)
+	{
+		AppendError(Debug, EGeometryScriptErrorType::InvalidInputs, LOCTEXT("GetTrianglesByMaterialID_MissingMaterialID", "GetTrianglesByMaterialID: MaterialID Attribute is not enabled"));
+		return TargetMesh;
+	}
+
+	if (bAllValidTriangles == false)
+	{
+		AppendError(Debug, EGeometryScriptErrorType::InvalidInputs, LOCTEXT("GetTrianglesByMaterialID_InvalidTriangles", "GetTrianglesByMaterialID: TriangleIDList has invalid triangles"));
+		return TargetMesh;
+	}
+
+	return TargetMesh;
+}
+
 
 UDynamicMesh* UGeometryScriptLibrary_MeshMaterialFunctions::SetTriangleMaterialID( 
 	UDynamicMesh* TargetMesh, 
@@ -236,7 +345,7 @@ UDynamicMesh* UGeometryScriptLibrary_MeshMaterialFunctions::SetTriangleMaterialI
 
 UDynamicMesh* UGeometryScriptLibrary_MeshMaterialFunctions::SetAllTriangleMaterialIDs(
 	UDynamicMesh* TargetMesh,
-	FGeometryScriptIndexList& TriangleMaterialIDList,
+	FGeometryScriptIndexList TriangleMaterialIDList,
 	bool bDeferChangeNotifications,
 	UGeometryScriptDebug* Debug)
 {
@@ -245,9 +354,9 @@ UDynamicMesh* UGeometryScriptLibrary_MeshMaterialFunctions::SetAllTriangleMateri
 		UE::Geometry::AppendError(Debug, EGeometryScriptErrorType::InvalidInputs, LOCTEXT("SetAllTriangleMaterialIDs_InvalidMesh", "SetAllTriangleMaterialIDs: TargetMesh is Null"));
 		return TargetMesh;
 	}
-	if (TriangleMaterialIDList.List.IsValid() == false || TriangleMaterialIDList.List->Num() == 0)
+	if (TriangleMaterialIDList.List.IsValid() == false)
 	{
-		UE::Geometry::AppendError(Debug, EGeometryScriptErrorType::InvalidInputs, LOCTEXT("SetAllTriangleMaterialIDs_InvalidList", "SetAllTriangleMaterialIDs: TriangleMaterialIDList is empty"));
+		UE::Geometry::AppendError(Debug, EGeometryScriptErrorType::InvalidInputs, LOCTEXT("SetAllTriangleMaterialIDs_InvalidList", "SetAllTriangleMaterialIDs: TriangleMaterialIDList is Null"));
 		return TargetMesh;
 	}
 	if (TriangleMaterialIDList.IsCompatibleWith(EGeometryScriptIndexType::MaterialID) == false)
@@ -255,6 +364,10 @@ UDynamicMesh* UGeometryScriptLibrary_MeshMaterialFunctions::SetAllTriangleMateri
 		UE::Geometry::AppendError(Debug, EGeometryScriptErrorType::InvalidInputs, LOCTEXT("SetAllTriangleMaterialIDs_InvalidList2", "SetAllTriangleMaterialIDs: TriangleMaterialIDList has incompatible index type"));
 		return TargetMesh;
 	}
+	if (TriangleMaterialIDList.List->Num() == 0 && TargetMesh->IsEmpty() == false)
+	{
+		UE::Geometry::AppendWarning(Debug, EGeometryScriptErrorType::InvalidInputs, LOCTEXT("SetAllTriangleMaterialIDs_EmptyList", "SetAllTriangleMaterialIDs: TriangleMaterialIDList is empty"));
+	}
 
 	TargetMesh->EditMesh([&](FDynamicMesh3& EditMesh) 
 	{
@@ -287,6 +400,65 @@ UDynamicMesh* UGeometryScriptLibrary_MeshMaterialFunctions::SetAllTriangleMateri
 
 
 
+UDynamicMesh* UGeometryScriptLibrary_MeshMaterialFunctions::SetMaterialIDOnTriangles( 
+	UDynamicMesh* TargetMesh, 
+	FGeometryScriptIndexList TriangleIDList,
+	int MaterialID,
+	bool bDeferChangeNotifications,
+	UGeometryScriptDebug* Debug)
+{	
+	if (TargetMesh == nullptr)
+	{
+		AppendError(Debug, EGeometryScriptErrorType::InvalidInputs, LOCTEXT("SetMaterialIDOnTriangles_InvalidMesh", "SetMaterialIDOnTriangles: TargetMesh is Null"));
+		return TargetMesh;
+	}
+	if (TriangleIDList.List.IsValid() == false)
+	{
+		AppendError(Debug, EGeometryScriptErrorType::InvalidInputs, LOCTEXT("SetMaterialIDOnTriangles_InvalidList", "SetMaterialIDOnTriangles: TriangleIDList is Null"));
+		return TargetMesh;
+	}
+	if (TriangleIDList.IsCompatibleWith(EGeometryScriptIndexType::Triangle) == false)
+	{
+		AppendError(Debug, EGeometryScriptErrorType::InvalidInputs, LOCTEXT("SetMaterialIDOnTriangles_InvalidList2", "SetMaterialIDOnTriangles: TriangleIDList has incompatible index type"));
+		return TargetMesh;
+	}
+	if (TriangleIDList.List->Num() == 0 && TargetMesh->IsEmpty() == false)
+	{
+		UE::Geometry::AppendWarning(Debug, EGeometryScriptErrorType::InvalidInputs, LOCTEXT("SetMaterialIDOnTriangles_EmptyList", "SetMaterialIDOnTriangles: TriangleIDList is empty"));
+	}
+
+	bool bHasInvalidTriangles = false;
+	TargetMesh->EditMesh([&](FDynamicMesh3& EditMesh) 
+		{
+			if (EditMesh.HasAttributes() == false)
+			{
+				EditMesh.EnableAttributes();
+			}
+			if (EditMesh.Attributes()->HasMaterialID() == false)
+			{
+				EditMesh.Attributes()->EnableMaterialID();
+			}
+			FDynamicMeshMaterialAttribute* MaterialIDs = EditMesh.Attributes()->GetMaterialID();
+			for (int32 TriangleID : *TriangleIDList.List)
+			{
+				if (EditMesh.IsTriangle(TriangleID) == false)
+				{
+					bHasInvalidTriangles = true;
+					return;
+				}
+				MaterialIDs->SetValue(TriangleID, MaterialID);
+			}
+		}, EDynamicMeshChangeType::GeneralEdit, EDynamicMeshAttributeChangeFlags::Unknown, bDeferChangeNotifications);
+
+	if (bHasInvalidTriangles == true)
+	{
+		AppendError(Debug, EGeometryScriptErrorType::InvalidInputs, LOCTEXT("SetMaterialIDOnTriangles_InvalidTriangles", "SetMaterialIDOnTriangles: TriangleIDList has invalid triangles"));
+		return TargetMesh;
+	}
+	
+	return TargetMesh;
+}
+
 
 
 UDynamicMesh* UGeometryScriptLibrary_MeshMaterialFunctions::SetPolygroupMaterialID( 
diff --git a/Engine/Plugins/Experimental/GeometryScripting/Source/GeometryScriptingCore/Public/GeometryScript/GeometryScriptTypes.h b/Engine/Plugins/Experimental/GeometryScripting/Source/GeometryScriptingCore/Public/GeometryScript/GeometryScriptTypes.h
--- a/Engine/Plugins/Experimental/GeometryScripting/Source/GeometryScriptingCore/Public/GeometryScript/GeometryScriptTypes.h
+++ b/Engine/Plugins/Experimental/GeometryScripting/Source/GeometryScriptingCore/Public/GeometryScript/GeometryScriptTypes.h
@@ -373,7 +373,10 @@ namespace UE
 namespace Geometry
 {
 	GEOMETRYSCRIPTINGCORE_API FGeometryScriptDebugMessage MakeScriptError(EGeometryScriptErrorType ErrorTypeIn, const FText& MessageIn);
+	GEOMETRYSCRIPTINGCORE_API FGeometryScriptDebugMessage MakeScriptWarning(EGeometryScriptErrorType WarningTypeIn, const FText& MessageIn);
+
 
 	GEOMETRYSCRIPTINGCORE_API void AppendError(UGeometryScriptDebug* Debug, EGeometryScriptErrorType ErrorTypeIn, const FText& MessageIn);
+	GEOMETRYSCRIPTINGCORE_API void AppendWarning(UGeometryScriptDebug* Debug, EGeometryScriptErrorType WarningTypeIn, const FText& MessageIn);
 }
 }
\ No newline at end of file
diff --git a/Engine/Plugins/Experimental/GeometryScripting/Source/GeometryScriptingCore/Public/GeometryScript/MeshMaterialFunctions.h b/Engine/Plugins/Experimental/GeometryScripting/Source/GeometryScriptingCore/Public/GeometryScript/MeshMaterialFunctions.h
--- a/Engine/Plugins/Experimental/GeometryScripting/Source/GeometryScriptingCore/Public/GeometryScript/MeshMaterialFunctions.h
+++ b/Engine/Plugins/Experimental/GeometryScripting/Source/GeometryScriptingCore/Public/GeometryScript/MeshMaterialFunctions.h
@@ -35,16 +35,31 @@ public:
 		int ToMaterialID,
 		UGeometryScriptDebug* Debug = nullptr );
 
-	UFUNCTION(BlueprintPure, Category = "GeometryScript|Materials", meta=(ScriptMethod))
+	UFUNCTION(BlueprintPure, Category = "GeometryScript|Materials", meta=(ScriptMethod, DisplayName="GetMaterialIDOfTriangle"))
 	static UPARAM(DisplayName = "Material ID") int32
 	GetTriangleMaterialID( UDynamicMesh* TargetMesh, int TriangleID, bool& bIsValidTriangle );
 
 	UFUNCTION(BlueprintPure, Category = "GeometryScript|Materials", meta=(ScriptMethod))
 	static UPARAM(DisplayName = "Target Mesh") UDynamicMesh* 
+	GetMaterialIDsOfTriangles(
+		UDynamicMesh* TargetMesh,
+		FGeometryScriptIndexList TriangleIDList,
+		FGeometryScriptIndexList& MaterialIDList,
+		UGeometryScriptDebug* Debug = nullptr);
+
+	UFUNCTION(BlueprintPure, Category = "GeometryScript|Materials", meta=(ScriptMethod, DisplayName="GetMaterialIDsOfAllTriangles"))
+	static UPARAM(DisplayName = "Target Mesh") UDynamicMesh* 
 	GetAllTriangleMaterialIDs( UDynamicMesh* TargetMesh, FGeometryScriptIndexList& MaterialIDList, bool& bHasMaterialIDs );
 
+	UFUNCTION(BlueprintPure, Category = "GeometryScript|Selection", meta=(ScriptMethod))
+	static UPARAM(DisplayName = "Target Mesh") UDynamicMesh* 
+	GetTrianglesByMaterialID( 
+		UDynamicMesh* TargetMesh, 
+		int MaterialID,
+		FGeometryScriptIndexList& TriangleIDList,
+		UGeometryScriptDebug* Debug = nullptr);
 
-	UFUNCTION(BlueprintCallable, Category = "GeometryScript|Materials", meta=(ScriptMethod))
+	UFUNCTION(BlueprintCallable, Category = "GeometryScript|Materials", meta=(ScriptMethod, DisplayName="SetMaterialIDOnTriangle"))
 	static UPARAM(DisplayName = "Target Mesh") UDynamicMesh* 
 	SetTriangleMaterialID( 
 		UDynamicMesh* TargetMesh, 
@@ -53,14 +68,22 @@ public:
 		bool& bIsValidTriangle, 
 		bool bDeferChangeNotifications = false );
 
-	UFUNCTION(BlueprintCallable, Category = "GeometryScript|Materials", meta=(ScriptMethod))
+	UFUNCTION(BlueprintCallable, Category = "GeometryScript|Materials", meta=(ScriptMethod, DisplayName="SetMaterialIDsOnAllTriangles"))
 	static UPARAM(DisplayName = "Target Mesh") UDynamicMesh*
 	SetAllTriangleMaterialIDs(
 		UDynamicMesh* TargetMesh,
-		FGeometryScriptIndexList& TriangleMaterialIDList,
+		FGeometryScriptIndexList TriangleMaterialIDList,
 		bool bDeferChangeNotifications = false,
 		UGeometryScriptDebug* Debug = nullptr);
 
+	UFUNCTION(BlueprintCallable, Category = "GeometryScript|Materials", meta=(ScriptMethod))
+	static UPARAM(DisplayName = "Target Mesh") UDynamicMesh* 
+	SetMaterialIDOnTriangles(
+		UDynamicMesh* TargetMesh,
+		FGeometryScriptIndexList TriangleIDList,
+		int MaterialID,
+        bool bDeferChangeNotifications = false,
+		UGeometryScriptDebug* Debug = nullptr);
 
 	UFUNCTION(BlueprintCallable, Category = "GeometryScript|Materials", meta=(ScriptMethod))
 	static UPARAM(DisplayName = "Target Mesh") UDynamicMesh* 
