GeometryScript: Add BakeTexture support to GeometryScript.

#rnx

diff --git a/Engine/Plugins/Experimental/GeometryScripting/Source/GeometryScriptingCore/Private/MeshBakeFunctions.cpp b/Engine/Plugins/Experimental/GeometryScripting/Source/GeometryScriptingCore/Private/MeshBakeFunctions.cpp
new file mode 100644
--- /dev/null
+++ b/Engine/Plugins/Experimental/GeometryScripting/Source/GeometryScriptingCore/Private/MeshBakeFunctions.cpp
@@ -0,0 +1,567 @@
+// Copyright Epic Games, Inc. All Rights Reserved.
+
+#include "GeometryScript/MeshBakeFunctions.h"
+#include "UDynamicMesh.h"
+
+#include "Sampling/MeshBakerCommon.h"
+#include "Sampling/MeshMapBaker.h"
+#include "Sampling/MeshConstantMapEvaluator.h"
+#include "Sampling/MeshCurvatureMapEvaluator.h"
+#include "Sampling/MeshOcclusionMapEvaluator.h"
+#include "Sampling/MeshNormalMapEvaluator.h"
+#include "Sampling/MeshPropertyMapEvaluator.h"
+#include "Sampling/MeshResampleImageEvaluator.h"
+
+#include "DynamicMesh/MeshTransforms.h"
+
+#include "AssetUtils/Texture2DBuilder.h"
+#include "AssetUtils/Texture2DUtil.h"
+
+using namespace UE::Geometry;
+
+#define LOCTEXT_NAMESPACE "UGeometryScriptLibrary_MeshBakeFunctions"
+
+namespace GeometryScriptBakeLocals
+{
+	FImageDimensions GetDimensions(const EGeometryScriptBakeResolution Resolution)
+	{
+		int Dimension = 256;
+		switch(Resolution)
+		{
+		case EGeometryScriptBakeResolution::Resolution16:
+			Dimension = 16;
+			break;
+		case EGeometryScriptBakeResolution::Resolution32:
+			Dimension = 32;
+			break;
+		case EGeometryScriptBakeResolution::Resolution64:
+			Dimension = 64;
+			break;
+		case EGeometryScriptBakeResolution::Resolution128:
+			Dimension = 128;
+			break;
+		case EGeometryScriptBakeResolution::Resolution256:
+			Dimension = 256;
+			break;
+		case EGeometryScriptBakeResolution::Resolution512:
+			Dimension = 512;
+			break;
+		case EGeometryScriptBakeResolution::Resolution1024:
+			Dimension = 1024;
+			break;
+		case EGeometryScriptBakeResolution::Resolution2048:
+			Dimension = 2048;
+			break;
+		case EGeometryScriptBakeResolution::Resolution4096:
+			Dimension = 4096;
+			break;
+		case EGeometryScriptBakeResolution::Resolution8192:
+			Dimension = 8192;
+			break;
+		}
+		return FImageDimensions(Dimension, Dimension);
+	}
+
+	int GetSamplesPerPixel(EGeometryScriptBakeSamplesPerPixel SamplesPerPixel)
+	{
+		int Samples = 1;
+		switch(SamplesPerPixel)
+		{
+		case EGeometryScriptBakeSamplesPerPixel::Sample1:
+			Samples = 1;
+			break;
+		case EGeometryScriptBakeSamplesPerPixel::Sample4:
+			Samples = 4;
+			break;
+		case EGeometryScriptBakeSamplesPerPixel::Sample16:
+			Samples = 16;
+			break;
+		case EGeometryScriptBakeSamplesPerPixel::Sample64:
+			Samples = 64;
+			break;
+		case EGeometryScriptBakeSamplesPerPixel::Samples256:
+			Samples = 256;
+			break;
+		}
+		return Samples;
+	}
+	
+	FTexture2DBuilder::ETextureType GetTextureType(const EGeometryScriptBakeTypes BakeType, const EGeometryScriptBakeBitDepth MapFormat)
+	{
+		FTexture2DBuilder::ETextureType TexType = FTexture2DBuilder::ETextureType::Color;
+		switch (BakeType)
+		{
+		default:
+			checkNoEntry();
+			break;
+		case EGeometryScriptBakeTypes::TangentSpaceNormal:
+			TexType = FTexture2DBuilder::ETextureType::NormalMap;
+			break;
+		case EGeometryScriptBakeTypes::AmbientOcclusion:
+			TexType = FTexture2DBuilder::ETextureType::AmbientOcclusion;
+			break;
+		case EGeometryScriptBakeTypes::BentNormal:
+			TexType = FTexture2DBuilder::ETextureType::NormalMap;
+			break;
+		case EGeometryScriptBakeTypes::Curvature:
+		case EGeometryScriptBakeTypes::ObjectSpaceNormal:
+		case EGeometryScriptBakeTypes::FaceNormal:
+		case EGeometryScriptBakeTypes::Position:
+			TexType = FTexture2DBuilder::ETextureType::ColorLinear;
+			break;
+		case EGeometryScriptBakeTypes::MaterialID:
+		case EGeometryScriptBakeTypes::VertexColor:
+			break;
+		case EGeometryScriptBakeTypes::Texture:
+		case EGeometryScriptBakeTypes::MultiTexture:
+			// For texture output with 16-bit source data, output HDR texture
+			if (MapFormat == EGeometryScriptBakeBitDepth::ChannelBits16)
+			{
+				TexType = FTexture2DBuilder::ETextureType::EmissiveHDR;
+			}
+			break;
+		}
+		return TexType;
+	}
+
+	FMeshCurvatureMapEvaluator::ECurvatureType GetCurvatureType(EGeometryScriptBakeCurvatureTypeMode CurvatureType)
+	{
+		FMeshCurvatureMapEvaluator::ECurvatureType Result = FMeshCurvatureMapEvaluator::ECurvatureType::Mean;
+		switch(CurvatureType)
+		{
+		case EGeometryScriptBakeCurvatureTypeMode::Mean:
+			Result = FMeshCurvatureMapEvaluator::ECurvatureType::Mean;
+			break;
+		case EGeometryScriptBakeCurvatureTypeMode::Gaussian:
+			Result = FMeshCurvatureMapEvaluator::ECurvatureType::Gaussian;
+			break;
+		case EGeometryScriptBakeCurvatureTypeMode::Min:
+			Result = FMeshCurvatureMapEvaluator::ECurvatureType::MinPrincipal;
+			break;
+		case EGeometryScriptBakeCurvatureTypeMode::Max:
+			Result = FMeshCurvatureMapEvaluator::ECurvatureType::MaxPrincipal;
+			break;
+		}
+		return Result;
+	}
+
+	FMeshCurvatureMapEvaluator::EColorMode GetCurvatureColorMode(EGeometryScriptBakeCurvatureColorMode ColorMode)
+	{
+		FMeshCurvatureMapEvaluator::EColorMode Result = FMeshCurvatureMapEvaluator::EColorMode::BlackGrayWhite;
+		switch(ColorMode)
+		{
+		case EGeometryScriptBakeCurvatureColorMode::Grayscale:
+			Result = FMeshCurvatureMapEvaluator::EColorMode::BlackGrayWhite;
+			break;
+		case EGeometryScriptBakeCurvatureColorMode::RedGreenBlue:
+			Result = FMeshCurvatureMapEvaluator::EColorMode::RedGreenBlue;
+			break;
+		case EGeometryScriptBakeCurvatureColorMode::RedBlue:
+			Result = FMeshCurvatureMapEvaluator::EColorMode::RedBlue;
+			break;
+		}
+		return Result;
+	}
+
+	FMeshCurvatureMapEvaluator::EClampMode GetCurvatureClampMode(EGeometryScriptBakeCurvatureClampMode ClampMode)
+	{
+		FMeshCurvatureMapEvaluator::EClampMode Result = FMeshCurvatureMapEvaluator::EClampMode::FullRange;
+		switch(ClampMode)
+		{
+		case EGeometryScriptBakeCurvatureClampMode::None:
+			Result = FMeshCurvatureMapEvaluator::EClampMode::FullRange;
+			break;
+		case EGeometryScriptBakeCurvatureClampMode::OnlyNegative:
+			Result = FMeshCurvatureMapEvaluator::EClampMode::Negative;
+			break;
+		case EGeometryScriptBakeCurvatureClampMode::OnlyPositive:
+			Result = FMeshCurvatureMapEvaluator::EClampMode::Positive;
+			break;
+		}
+		return Result;
+	}
+}
+
+
+FGeometryScriptBakeTypeOptions UGeometryScriptLibrary_MeshBakeFunctions::MakeBakeTypeTangentNormal()
+{
+	FGeometryScriptBakeTypeOptions Output;
+	Output.BakeType = EGeometryScriptBakeTypes::TangentSpaceNormal;
+	return Output;
+}
+
+
+FGeometryScriptBakeTypeOptions UGeometryScriptLibrary_MeshBakeFunctions::MakeBakeTypeObjectNormal()
+{
+	FGeometryScriptBakeTypeOptions Output;
+	Output.BakeType = EGeometryScriptBakeTypes::ObjectSpaceNormal;
+	return Output;
+}
+
+
+FGeometryScriptBakeTypeOptions UGeometryScriptLibrary_MeshBakeFunctions::MakeBakeTypeFaceNormal()
+{
+	FGeometryScriptBakeTypeOptions Output;
+	Output.BakeType = EGeometryScriptBakeTypes::FaceNormal;
+	return Output;
+}
+
+FGeometryScriptBakeTypeOptions UGeometryScriptLibrary_MeshBakeFunctions::MakeBakeTypeBentNormal(
+	int OcclusionRays,
+	float MaxDistance,
+	float SpreadAngle)
+{
+	FGeometryScriptBakeTypeOptions Output;
+	Output.BakeType = EGeometryScriptBakeTypes::BentNormal;
+	const TSharedPtr<FGeometryScriptBakeType_Occlusion> OcclusionOptions = MakeShared<FGeometryScriptBakeType_Occlusion>();
+	OcclusionOptions->OcclusionRays = OcclusionRays;
+	OcclusionOptions->MaxDistance = MaxDistance;
+	OcclusionOptions->SpreadAngle = SpreadAngle;
+	Output.Options = OcclusionOptions;
+	return Output;
+}
+
+FGeometryScriptBakeTypeOptions UGeometryScriptLibrary_MeshBakeFunctions::MakeBakeTypePosition()
+{
+	FGeometryScriptBakeTypeOptions Output;
+	Output.BakeType = EGeometryScriptBakeTypes::Position;
+	return Output;
+}
+
+FGeometryScriptBakeTypeOptions UGeometryScriptLibrary_MeshBakeFunctions::MakeBakeTypeCurvature(
+	EGeometryScriptBakeCurvatureTypeMode CurvatureType,
+	EGeometryScriptBakeCurvatureColorMode ColorMapping,
+	float ColorRangeMultiplier,
+	float MinRangeMultiplier,
+	EGeometryScriptBakeCurvatureClampMode Clamping)
+{
+	FGeometryScriptBakeTypeOptions Output;
+	Output.BakeType = EGeometryScriptBakeTypes::Curvature;
+	const TSharedPtr<FGeometryScriptBakeType_Curvature> CurvatureOptions = MakeShared<FGeometryScriptBakeType_Curvature>();
+	CurvatureOptions->CurvatureType = CurvatureType;
+	CurvatureOptions->ColorMapping = ColorMapping;
+	CurvatureOptions->ColorRangeMultiplier = ColorRangeMultiplier;
+	CurvatureOptions->MinRangeMultiplier = MinRangeMultiplier;
+	Output.Options = CurvatureOptions;
+	return Output;
+}
+
+FGeometryScriptBakeTypeOptions UGeometryScriptLibrary_MeshBakeFunctions::MakeBakeTypeAmbientOcclusion(
+	int OcclusionRays,
+	float MaxDistance,
+	float SpreadAngle,
+	float BiasAngle)
+{
+	FGeometryScriptBakeTypeOptions Output;
+	Output.BakeType = EGeometryScriptBakeTypes::AmbientOcclusion;
+	const TSharedPtr<FGeometryScriptBakeType_Occlusion> OcclusionOptions = MakeShared<FGeometryScriptBakeType_Occlusion>();
+	OcclusionOptions->OcclusionRays = OcclusionRays;
+	OcclusionOptions->MaxDistance = MaxDistance;
+	OcclusionOptions->SpreadAngle = SpreadAngle;
+	OcclusionOptions->BiasAngle = BiasAngle;
+	Output.Options = OcclusionOptions;
+	return Output;
+}
+
+FGeometryScriptBakeTypeOptions UGeometryScriptLibrary_MeshBakeFunctions::MakeBakeTypeTexture(
+	UTexture2D* SourceTexture,
+	int SourceUVLayer)
+{
+	FGeometryScriptBakeTypeOptions Output;
+	Output.BakeType = EGeometryScriptBakeTypes::Texture;
+	const TSharedPtr<FGeometryScriptBakeType_Texture> TextureOptions = MakeShared<FGeometryScriptBakeType_Texture>();
+	TextureOptions->SourceTexture = SourceTexture;
+	TextureOptions->SourceUVLayer = SourceUVLayer;
+	Output.Options = TextureOptions;
+	return Output;
+}
+
+FGeometryScriptBakeTypeOptions UGeometryScriptLibrary_MeshBakeFunctions::MakeBakeTypeMultiTexture(
+	const TArray<UTexture2D*>& MaterialIDSourceTextures,
+	int SourceUVLayer)
+{
+	FGeometryScriptBakeTypeOptions Output;
+	Output.BakeType = EGeometryScriptBakeTypes::MultiTexture;
+	const TSharedPtr<FGeometryScriptBakeType_MultiTexture> MultiTextureOptions = MakeShared<FGeometryScriptBakeType_MultiTexture>();
+	MultiTextureOptions->MaterialIDSourceTextures = MaterialIDSourceTextures;
+	MultiTextureOptions->SourceUVLayer = SourceUVLayer;
+	Output.Options = MultiTextureOptions;
+	return Output;
+}
+
+FGeometryScriptBakeTypeOptions UGeometryScriptLibrary_MeshBakeFunctions::MakeBakeTypeVertexColor()
+{
+	FGeometryScriptBakeTypeOptions Output;
+	Output.BakeType = EGeometryScriptBakeTypes::VertexColor;
+	return Output;
+}
+
+FGeometryScriptBakeTypeOptions UGeometryScriptLibrary_MeshBakeFunctions::MakeBakeTypeMaterialID()
+{
+	FGeometryScriptBakeTypeOptions Output;
+	Output.BakeType = EGeometryScriptBakeTypes::MaterialID;
+	return Output;
+}
+
+
+TArray<UTexture2D*> UGeometryScriptLibrary_MeshBakeFunctions::BakeTexture(
+	UDynamicMesh* TargetMesh,
+	FTransform TargetTransform,
+	FGeometryScriptBakeTargetMeshOptions TargetOptions,
+	UDynamicMesh* SourceMesh,
+	FTransform SourceTransform,
+	FGeometryScriptBakeSourceMeshOptions SourceOptions,
+	const TArray<FGeometryScriptBakeTypeOptions>& BakeTypes,
+	FGeometryScriptBakeTextureOptions BakeOptions,
+	UGeometryScriptDebug* Debug)
+{
+	TArray<UTexture2D*> TextureOutput;
+	if (TargetMesh == nullptr)
+	{
+		UE::Geometry::AppendError(Debug, EGeometryScriptErrorType::InvalidInputs, LOCTEXT("BakeTexture_InvalidTargetMesh", "BakeTexture: TargetMesh is Null"));
+		return TextureOutput;
+	}
+	if (SourceMesh == nullptr)
+	{
+		UE::Geometry::AppendError(Debug, EGeometryScriptErrorType::InvalidInputs, LOCTEXT("BakeTexture_InvalidSourceMesh", "BakeTexture: SourceMesh is Null"));
+		return TextureOutput;
+	}
+	if (BakeTypes.Num() == 0)
+	{
+		return TextureOutput;
+	}
+
+	TSharedPtr<FMeshTangentsd, ESPMode::ThreadSafe> TargetMeshTangents;
+	TSharedPtr<FMeshTangentsd, ESPMode::ThreadSafe> SourceMeshTangents;
+	auto GetMeshTangents = [](UDynamicMesh* Mesh, TSharedPtr<FMeshTangentsd, ESPMode::ThreadSafe>& Tangents)
+	{
+		if (!Tangents)
+		{
+			Tangents = MakeShared<FMeshTangentsd, ESPMode::ThreadSafe>(Mesh->GetMeshPtr());
+			Tangents->CopyTriVertexTangents(Mesh->GetMeshRef());
+		}
+	};
+
+	const bool bIsBakeToSelf = (TargetMesh == SourceMesh);
+
+	FDynamicMesh3 SourceMeshCopy;
+	const FDynamicMesh3* SourceMeshOriginal = SourceMesh->GetMeshPtr();
+	const FDynamicMesh3* SourceMeshToUse = SourceMeshOriginal;
+	if (BakeOptions.bProjectionInWorldSpace && !bIsBakeToSelf)
+	{
+		// Transform the SourceMesh into TargetMesh local space using a copy (oof)
+		// TODO: Remove this once we have support for transforming rays in the core bake loop
+		SourceMeshCopy = *SourceMeshOriginal;
+		const FTransformSRT3d SourceToWorld = SourceTransform;
+		MeshTransforms::ApplyTransform(SourceMeshCopy, SourceToWorld);
+		const FTransformSRT3d TargetToWorld = TargetTransform;
+		MeshTransforms::ApplyTransform(SourceMeshCopy, TargetToWorld.Inverse());
+		SourceMeshToUse = &SourceMeshCopy;
+	}
+
+	FImageDimensions BakeDimensions = GeometryScriptBakeLocals::GetDimensions(BakeOptions.Resolution);
+	const FDynamicMeshAABBTree3 DetailSpatial(SourceMeshToUse);
+	FMeshBakerDynamicMeshSampler DetailSampler(SourceMeshToUse, &DetailSpatial);
+
+	FMeshMapBaker Baker;
+	Baker.SetTargetMesh(TargetMesh->GetMeshPtr());
+	Baker.SetTargetMeshUVLayer(TargetOptions.TargetUVLayer);
+	Baker.SetDetailSampler(&DetailSampler);
+	Baker.SetDimensions(BakeDimensions);
+	Baker.SetProjectionDistance(BakeOptions.ProjectionDistance);
+	Baker.SetSamplesPerPixel(GeometryScriptBakeLocals::GetSamplesPerPixel(BakeOptions.SamplesPerPixel));
+
+	bool bSupportsSourceNormalMap = false;
+	TSharedPtr<UE::Geometry::TImageBuilder<FVector4f>, ESPMode::ThreadSafe> SourceTexture;
+	for (const FGeometryScriptBakeTypeOptions& Options : BakeTypes)
+	{
+		switch(Options.BakeType)
+		{
+		case EGeometryScriptBakeTypes::TangentSpaceNormal:
+		{
+			GetMeshTangents(TargetMesh, TargetMeshTangents);
+			TSharedPtr<FMeshNormalMapEvaluator> NormalEval = MakeShared<FMeshNormalMapEvaluator>();
+			bSupportsSourceNormalMap = true;
+			Baker.AddEvaluator(NormalEval);
+			break;
+		}
+		case EGeometryScriptBakeTypes::ObjectSpaceNormal:
+		{
+			TSharedPtr<FMeshPropertyMapEvaluator> PropertyEval = MakeShared<FMeshPropertyMapEvaluator>();
+			PropertyEval->Property = EMeshPropertyMapType::Normal;
+			bSupportsSourceNormalMap = true;
+			Baker.AddEvaluator(PropertyEval);
+			break;
+		}
+		case EGeometryScriptBakeTypes::FaceNormal:
+		{
+			TSharedPtr<FMeshPropertyMapEvaluator> PropertyEval = MakeShared<FMeshPropertyMapEvaluator>();
+			PropertyEval->Property = EMeshPropertyMapType::FacetNormal;
+			Baker.AddEvaluator(PropertyEval);
+			break;
+		}
+		case EGeometryScriptBakeTypes::BentNormal:
+		{
+			GetMeshTangents(TargetMesh, TargetMeshTangents);
+			FGeometryScriptBakeType_Occlusion* OcclusionOptions = static_cast<FGeometryScriptBakeType_Occlusion*>(Options.Options.Get());
+			TSharedPtr<FMeshOcclusionMapEvaluator> OcclusionEval = MakeShared<FMeshOcclusionMapEvaluator>();
+			OcclusionEval->OcclusionType = EMeshOcclusionMapType::BentNormal;
+			OcclusionEval->NumOcclusionRays = OcclusionOptions->OcclusionRays;
+			OcclusionEval->MaxDistance = (OcclusionOptions->MaxDistance == 0) ? TNumericLimits<float>::Max() : OcclusionOptions->MaxDistance;
+			OcclusionEval->SpreadAngle = OcclusionOptions->SpreadAngle;
+			Baker.AddEvaluator(OcclusionEval);
+			break;
+		}
+		case EGeometryScriptBakeTypes::Position:
+		{
+			TSharedPtr<FMeshPropertyMapEvaluator> PropertyEval = MakeShared<FMeshPropertyMapEvaluator>();
+			PropertyEval->Property = EMeshPropertyMapType::Position;
+			Baker.AddEvaluator(PropertyEval);
+			break;
+		}
+		case EGeometryScriptBakeTypes::Curvature:
+		{
+			FGeometryScriptBakeType_Curvature* CurvatureOptions = static_cast<FGeometryScriptBakeType_Curvature*>(Options.Options.Get());
+			TSharedPtr<FMeshCurvatureMapEvaluator> CurvatureEval = MakeShared<FMeshCurvatureMapEvaluator>();
+			CurvatureEval->UseCurvatureType = GeometryScriptBakeLocals::GetCurvatureType(CurvatureOptions->CurvatureType);
+			CurvatureEval->UseColorMode = GeometryScriptBakeLocals::GetCurvatureColorMode(CurvatureOptions->ColorMapping);
+			CurvatureEval->RangeScale = CurvatureOptions->ColorRangeMultiplier;
+			CurvatureEval->MinRangeScale = CurvatureOptions->MinRangeMultiplier;
+			CurvatureEval->UseClampMode = GeometryScriptBakeLocals::GetCurvatureClampMode(CurvatureOptions->Clamping);
+			Baker.AddEvaluator(CurvatureEval);
+			break;
+		}
+		case EGeometryScriptBakeTypes::AmbientOcclusion:
+		{
+			FGeometryScriptBakeType_Occlusion* OcclusionOptions = static_cast<FGeometryScriptBakeType_Occlusion*>(Options.Options.Get());
+			TSharedPtr<FMeshOcclusionMapEvaluator> OcclusionEval = MakeShared<FMeshOcclusionMapEvaluator>();
+			OcclusionEval->OcclusionType = EMeshOcclusionMapType::AmbientOcclusion;
+			OcclusionEval->NumOcclusionRays = OcclusionOptions->OcclusionRays;
+			OcclusionEval->MaxDistance = (OcclusionOptions->MaxDistance == 0) ? TNumericLimits<float>::Max() : OcclusionOptions->MaxDistance;
+			OcclusionEval->SpreadAngle = OcclusionOptions->SpreadAngle;
+			OcclusionEval->BiasAngleDeg = OcclusionOptions->BiasAngle;
+			Baker.AddEvaluator(OcclusionEval);
+			break;
+		}
+		case EGeometryScriptBakeTypes::Texture:
+		{
+			FGeometryScriptBakeType_Texture* TextureOptions = static_cast<FGeometryScriptBakeType_Texture*>(Options.Options.Get());
+
+			// TODO: Add support for sampling different texture maps per Texture evaluator in a single pass. 
+			if (!SourceTexture && TextureOptions->SourceTexture)
+			{
+				SourceTexture = MakeShared<TImageBuilder<FVector4f>>();
+				if (!UE::AssetUtils::ReadTexture(TextureOptions->SourceTexture, *SourceTexture, false))
+				{
+					UE::Geometry::AppendError(Debug, EGeometryScriptErrorType::InvalidInputs, LOCTEXT("BakeTexture_InvalidSourceTexture", "BakeTexture: Failed to read SourceTexture"));
+				}
+				else
+				{
+					DetailSampler.SetTextureMap(SourceMeshToUse, IMeshBakerDetailSampler::FBakeDetailTexture(SourceTexture.Get(), TextureOptions->SourceUVLayer));
+				}
+			}
+			
+			TSharedPtr<FMeshResampleImageEvaluator> TextureEval = MakeShared<FMeshResampleImageEvaluator>();
+			Baker.AddEvaluator(TextureEval);
+			break;
+		}
+		case EGeometryScriptBakeTypes::MultiTexture:
+		{
+			FGeometryScriptBakeType_MultiTexture* TextureOptions = static_cast<FGeometryScriptBakeType_MultiTexture*>(Options.Options.Get());
+			TSharedPtr<FMeshMultiResampleImageEvaluator> TextureEval = MakeShared<FMeshMultiResampleImageEvaluator>();
+
+			if (TextureOptions->MaterialIDSourceTextures.Num())
+			{
+				TextureEval->MultiTextures.SetNum(TextureOptions->MaterialIDSourceTextures.Num());
+				for (int32 MaterialId = 0; MaterialId < TextureEval->MultiTextures.Num(); ++MaterialId)
+				{
+					if (UTexture2D* Texture = TextureOptions->MaterialIDSourceTextures[MaterialId])
+					{
+						TextureEval->MultiTextures[MaterialId] = MakeShared<TImageBuilder<FVector4f>>();
+						if (!UE::AssetUtils::ReadTexture(Texture, *TextureEval->MultiTextures[MaterialId], false))
+						{
+							UE::Geometry::AppendError(Debug, EGeometryScriptErrorType::InvalidInputs, LOCTEXT("BakeTexture_InvalidMultiTexture", "BakeTexture: Failed to read MaterialIDSourceTexture"));
+						}
+					}
+				}
+			}
+			Baker.AddEvaluator(TextureEval);
+			break;
+		}
+		case EGeometryScriptBakeTypes::VertexColor:
+		{
+			TSharedPtr<FMeshPropertyMapEvaluator> PropertyEval = MakeShared<FMeshPropertyMapEvaluator>();
+			PropertyEval->Property = EMeshPropertyMapType::VertexColor;
+			Baker.AddEvaluator(PropertyEval);
+			break;
+		}
+		case EGeometryScriptBakeTypes::MaterialID:
+		{
+			TSharedPtr<FMeshPropertyMapEvaluator> PropertyEval = MakeShared<FMeshPropertyMapEvaluator>();
+			PropertyEval->Property = EMeshPropertyMapType::MaterialID;
+			Baker.AddEvaluator(PropertyEval);
+			break;
+		}
+		default:
+			break;
+		}
+	}
+
+	TSharedPtr<UE::Geometry::TImageBuilder<FVector4f>, ESPMode::ThreadSafe> SourceNormalMap;
+	if (bSupportsSourceNormalMap && SourceOptions.SourceNormalMap)
+	{
+		SourceNormalMap = MakeShared<TImageBuilder<FVector4f>>();
+		if (!UE::AssetUtils::ReadTexture(SourceOptions.SourceNormalMap, *SourceNormalMap, false))
+		{
+			UE::Geometry::AppendError(Debug, EGeometryScriptErrorType::InvalidInputs, LOCTEXT("BakeTexture_InvalidSourceNormalMap", "BakeTexture: Failed to read SourceNormalMap"));
+		}
+		else
+		{
+			DetailSampler.SetNormalMap(SourceMeshToUse,
+			IMeshBakerDetailSampler::FBakeDetailTexture(
+					SourceNormalMap.Get(),
+					SourceOptions.SourceNormalUVLayer));
+
+			GetMeshTangents(SourceMesh, SourceMeshTangents);
+			DetailSampler.SetTangents(SourceMeshToUse, SourceMeshTangents.Get());
+		}
+	}
+
+	if (TargetMeshTangents)
+	{
+		Baker.SetTargetMeshTangents(TargetMeshTangents);
+	}
+
+	Baker.Bake();
+
+	const int NumEval = Baker.NumEvaluators();
+	for (int EvalIdx = 0; EvalIdx < NumEval; ++EvalIdx)
+	{
+		EGeometryScriptBakeTypes BakeType = BakeTypes[EvalIdx].BakeType;
+		const EGeometryScriptBakeBitDepth BakeBitDepth = BakeOptions.BitDepth;
+		
+		// For 8-bit color textures, ensure that the source data is in sRGB.
+		const FTexture2DBuilder::ETextureType TexType = GeometryScriptBakeLocals::GetTextureType(BakeType, BakeBitDepth);
+		const bool bConvertToSRGB = TexType == FTexture2DBuilder::ETextureType::Color;
+		const ETextureSourceFormat SourceDataFormat = BakeBitDepth == EGeometryScriptBakeBitDepth::ChannelBits16 ? TSF_RGBA16F : TSF_BGRA8;
+
+		constexpr int ResultIdx = 0;
+		FTexture2DBuilder TextureBuilder;
+		TextureBuilder.Initialize(TexType, BakeDimensions);
+		TextureBuilder.Copy(*Baker.GetBakeResults(EvalIdx)[ResultIdx], bConvertToSRGB);
+		TextureBuilder.Commit(false);
+
+		// Copy image to source data after commit. This will avoid incurring
+		// the cost of hitting the DDC for texture compile while iterating on
+		// bake settings. Since this dirties the texture, the next time the texture
+		// is used after accepting the final texture, the DDC will trigger and
+		// properly recompile the platform data.
+		const bool bConvertSourceToSRGB = bConvertToSRGB && SourceDataFormat == TSF_BGRA8;
+		TextureBuilder.CopyImageToSourceData(*Baker.GetBakeResults(EvalIdx)[ResultIdx], SourceDataFormat, bConvertSourceToSRGB);
+		TextureOutput.Add(TextureBuilder.GetTexture2D());
+	}
+
+	return TextureOutput;
+}
+
+#undef LOCTEXT_NAMESPACE
diff --git a/Engine/Plugins/Experimental/GeometryScripting/Source/GeometryScriptingCore/Public/GeometryScript/MeshBakeFunctions.h b/Engine/Plugins/Experimental/GeometryScripting/Source/GeometryScriptingCore/Public/GeometryScript/MeshBakeFunctions.h
new file mode 100644
--- /dev/null
+++ b/Engine/Plugins/Experimental/GeometryScripting/Source/GeometryScriptingCore/Public/GeometryScript/MeshBakeFunctions.h
@@ -0,0 +1,319 @@
+// Copyright Epic Games, Inc. All Rights Reserved.
+
+#pragma once
+
+#include "CoreMinimal.h"
+#include "Kismet/BlueprintFunctionLibrary.h"
+#include "GeometryScript/GeometryScriptTypes.h"
+#include "MeshBakeFunctions.generated.h"
+
+class UDynamicMesh;
+
+UENUM(BlueprintType)
+enum class EGeometryScriptBakeResolution : uint8
+{
+	Resolution16 UMETA(DisplayName = "16 x 16"),
+	Resolution32 UMETA(DisplayName = "32 x 32"),
+	Resolution64 UMETA(DisplayName = "64 x 64"),
+	Resolution128 UMETA(DisplayName = "128 x 128"),
+	Resolution256 UMETA(DisplayName = "256 x 256"),
+	Resolution512 UMETA(DisplayName = "512 x 512"),
+	Resolution1024 UMETA(DisplayName = "1024 x 1024"),
+	Resolution2048 UMETA(DisplayName = "2048 x 2048"),
+	Resolution4096 UMETA(DisplayName = "4096 x 4096"),
+	Resolution8192 UMETA(DisplayName = "8192 x 8192")
+};
+
+UENUM(BlueprintType)
+enum class EGeometryScriptBakeBitDepth : uint8
+{
+	ChannelBits8 UMETA(DisplayName = "8 bits/channel"),
+	ChannelBits16 UMETA(DisplayName = "16 bits/channel")
+};
+
+UENUM(BlueprintType)
+enum class EGeometryScriptBakeSamplesPerPixel : uint8
+{
+	Sample1 UMETA(DisplayName = "1"),
+	Sample4 UMETA(DisplayName = "4"),
+	Sample16 UMETA(DisplayName = "16"),
+	Sample64 UMETA(DisplayName = "64"),
+	Samples256 UMETA(DisplayName = "256")
+};
+
+UENUM(BlueprintType)
+enum class EGeometryScriptBakeTypes : uint8
+{
+	/* Normals in tangent space */
+	TangentSpaceNormal     UMETA(DisplayName = "Tangent Normal"),
+	/* Interpolated normals in object space */
+	ObjectSpaceNormal      UMETA(DisplayName = "Object Normal"),
+	/* Geometric face normals in object space */
+	FaceNormal             ,
+	/* Normals skewed towards the least occluded direction */
+	BentNormal             ,
+	/* Positions in object space */
+	Position               ,
+	/* Local curvature of the mesh surface */
+	Curvature              ,
+	/* Ambient occlusion sampled across the hemisphere */
+	AmbientOcclusion       ,
+	/* Transfer a given texture */
+	Texture                ,
+	/* Transfer a texture per material ID */
+	MultiTexture           ,
+	/* Interpolated vertex colors */
+	VertexColor            ,
+	/* Material IDs as unique colors */
+	MaterialID             UMETA(DisplayName = "Material ID"),
+};
+
+USTRUCT(BlueprintType)
+struct GEOMETRYSCRIPTINGCORE_API FGeometryScriptBakeTypes
+{
+	GENERATED_BODY()
+};
+
+USTRUCT(BlueprintType)
+struct GEOMETRYSCRIPTINGCORE_API FGeometryScriptBakeType_Occlusion : public FGeometryScriptBakeTypes
+{
+	GENERATED_BODY()
+public:
+	/** Number of occlusion rays per sample */
+	UPROPERTY(BlueprintReadWrite, Category = OcclusionOutput, meta = (UIMin = "1", UIMax = "1024", ClampMin = "1", ClampMax = "65536"))
+	int32 OcclusionRays = 16;
+
+	/** Maximum distance for occlusion rays to test for intersections; a value of 0 means infinity */
+	UPROPERTY(BlueprintReadWrite, Category = OcclusionOutput, meta = (UIMin = "0.0", UIMax = "1000.0", ClampMin = "0.0", ClampMax = "99999999.0"))
+	float MaxDistance = 0.0f;
+
+	/** Maximum spread angle in degrees for occlusion rays; for example, 180 degrees will cover the entire hemisphere whereas 90 degrees will only cover the center of the hemisphere down to 45 degrees from the horizon. */
+	UPROPERTY(BlueprintReadWrite, Category = OcclusionOutput, meta = (UIMin = "0", UIMax = "180.0", ClampMin = "0", ClampMax = "180.0"))
+	float SpreadAngle = 180.0f;
+
+	/** Angle in degrees from the horizon for occlusion rays for which the contribution is attenuated to reduce faceting artifacts. */
+	UPROPERTY(BlueprintReadWrite, Category = OcclusionOutput, meta = (UIMin = "0", UIMax = "45.0", ClampMin = "0", ClampMax = "89.9"))
+	float BiasAngle = 15.0f;
+};
+
+UENUM(BlueprintType)
+enum class EGeometryScriptBakeCurvatureTypeMode : uint8
+{
+	/** Average of the minimum and maximum principal curvatures */
+	Mean,
+	/** Maximum principal curvature */
+	Max,
+	/** Minimum principal curvature */
+	Min,
+	/** Product of the minimum and maximum principal curvatures */
+	Gaussian
+};
+
+
+UENUM(BlueprintType)
+enum class EGeometryScriptBakeCurvatureColorMode : uint8
+{
+	/** Map curvature values to grayscale such that black is negative, grey is zero, and white is positive */
+	Grayscale,
+	/** Map curvature values to red and blue such that red is negative, black is zero, and blue is positive */
+	RedBlue,
+	/** Map curvature values to red, green, blue such that red is negative, green is zero, and blue is positive */
+	RedGreenBlue
+};
+
+
+UENUM(BlueprintType)
+enum class EGeometryScriptBakeCurvatureClampMode : uint8
+{
+	/** Include both negative and positive curvatures */
+	None,
+	/** Clamp negative curvatures to zero */
+	OnlyPositive,
+	/** Clamp positive curvatures to zero */
+	OnlyNegative
+};
+
+USTRUCT(BlueprintType)
+struct GEOMETRYSCRIPTINGCORE_API FGeometryScriptBakeType_Curvature : public FGeometryScriptBakeTypes
+{
+	GENERATED_BODY()
+public:
+	/** Type of curvature */
+	UPROPERTY(BlueprintReadWrite, Category = CurvatureOutput)
+	EGeometryScriptBakeCurvatureTypeMode CurvatureType = EGeometryScriptBakeCurvatureTypeMode::Mean;
+
+	/** How to map calculated curvature values to colors */
+	UPROPERTY(BlueprintReadWrite, Category = CurvatureOutput)
+	EGeometryScriptBakeCurvatureColorMode ColorMapping = EGeometryScriptBakeCurvatureColorMode::Grayscale;
+
+	/** Multiplier for how the curvature values fill the available range in the selected Color Mapping; a larger value means that higher curvature is required to achieve the maximum color value. */
+	UPROPERTY(BlueprintReadWrite, Category = CurvatureOutput, meta = (UIMin = "0.1", UIMax = "2.0", ClampMin = "0.001", ClampMax = "100.0"))
+	float ColorRangeMultiplier = 1.0f;
+
+	/** Minimum for the curvature values to not be clamped to zero relative to the curvature for the maximum color value; a larger value means that higher curvature is required to not be considered as no curvature. */
+	UPROPERTY(BlueprintReadWrite, Category = CurvatureOutput, AdvancedDisplay, meta = (DisplayName = "Color Range Minimum", UIMin = "0.0", UIMax = "1.0"))
+	float MinRangeMultiplier = 0.0f;
+
+	/** Clamping applied to curvature values before color mapping */
+	UPROPERTY(BlueprintReadWrite, Category = CurvatureOutput)
+	EGeometryScriptBakeCurvatureClampMode Clamping = EGeometryScriptBakeCurvatureClampMode::None;
+};
+
+USTRUCT(BlueprintType)
+struct GEOMETRYSCRIPTINGCORE_API FGeometryScriptBakeType_Texture : public FGeometryScriptBakeTypes
+{
+	GENERATED_BODY()
+public:
+	/** Source mesh texture that is to be resampled into a new texture */
+	UPROPERTY(BlueprintReadWrite, Category = TextureOutput)
+	TObjectPtr<UTexture2D> SourceTexture;
+
+	/** UV channel to use for the source mesh texture */
+	UPROPERTY(BlueprintReadWrite, Category = TextureOutput, meta = (DisplayName = "Source Texture UV Channel", ClampMin="0"))
+	int SourceUVLayer = 0;
+};
+
+USTRUCT(BlueprintType)
+struct GEOMETRYSCRIPTINGCORE_API FGeometryScriptBakeType_MultiTexture : public FGeometryScriptBakeTypes
+{
+	GENERATED_BODY()
+public:
+	/** For each material ID, the source texture that will be resampled in that material's region*/
+	UPROPERTY(BlueprintReadWrite, Category = MultiTextureOutput, meta = (DisplayName = "Source Textures by Material IDs"))
+	TArray<TObjectPtr<UTexture2D>> MaterialIDSourceTextures;
+
+	/** UV channel to use for the source mesh texture */
+	UPROPERTY(BlueprintReadWrite, Category = MultiTextureOutput, meta = (DisplayName = "Source Texture UV Channel", ClampMin="0"))
+	int SourceUVLayer = 0;
+};
+
+/**
+ * Opaque struct for storing bake type options.
+ */
+USTRUCT(BlueprintType)
+struct GEOMETRYSCRIPTINGCORE_API FGeometryScriptBakeTypeOptions
+{
+	GENERATED_BODY()
+
+	/** The bake output type to generate */
+	UPROPERTY(BlueprintReadOnly, Category = Type)
+	EGeometryScriptBakeTypes BakeType = EGeometryScriptBakeTypes::TangentSpaceNormal;
+
+	TSharedPtr<FGeometryScriptBakeTypes> Options;
+};
+
+USTRUCT(BlueprintType)
+struct GEOMETRYSCRIPTINGCORE_API FGeometryScriptBakeTextureOptions
+{
+	GENERATED_BODY()
+
+	/** The pixel resolution of the generated textures */
+	UPROPERTY(BlueprintReadWrite, Category = Options)
+	EGeometryScriptBakeResolution Resolution = EGeometryScriptBakeResolution::Resolution256;
+
+	/** The bit depth for each channel of the generated textures */
+	UPROPERTY(BlueprintReadWrite, Category = Options)
+	EGeometryScriptBakeBitDepth BitDepth = EGeometryScriptBakeBitDepth::ChannelBits8;
+
+	/** Number of samples per pixel */
+	UPROPERTY(BlueprintReadWrite, Category = Options)
+	EGeometryScriptBakeSamplesPerPixel SamplesPerPixel = EGeometryScriptBakeSamplesPerPixel::Sample1;
+
+	/** Maximum allowed distance for the projection from target mesh to source mesh for the sample to be considered valid.
+	 * This is only relevant if a separate source mesh is provided. */
+	UPROPERTY(BlueprintReadWrite, Category = Options)
+	float ProjectionDistance = 3.0f;
+
+	/** If true, uses the world space positions for the projection from target mesh to source mesh, otherwise it uses their object space positions.
+	 * This is only relevant if a separate source mesh is provided. */
+	UPROPERTY(BlueprintReadWrite, Category = Options)
+	bool bProjectionInWorldSpace = false;
+};
+
+USTRUCT(BlueprintType)
+struct GEOMETRYSCRIPTINGCORE_API FGeometryScriptBakeTargetMeshOptions
+{
+	GENERATED_BODY();
+
+	UPROPERTY(BlueprintReadWrite, Category = Options, meta = (DisplayName="Target UV Channel"))
+	int TargetUVLayer = 0;
+};
+
+USTRUCT(BlueprintType)
+struct GEOMETRYSCRIPTINGCORE_API FGeometryScriptBakeSourceMeshOptions
+{
+	GENERATED_BODY()
+
+	UPROPERTY(BlueprintReadWrite, Category = Options)
+	TObjectPtr<UTexture2D> SourceNormalMap = nullptr;
+
+	UPROPERTY(BlueprintReadWrite, Category = Options, meta = (DisplayName="Source Normal UV Channel"))
+	int SourceNormalUVLayer = 0;
+};
+
+UCLASS(meta = (ScriptName = "GeometryScript_Baking"))
+class GEOMETRYSCRIPTINGCORE_API UGeometryScriptLibrary_MeshBakeFunctions : public UBlueprintFunctionLibrary
+{
+	GENERATED_BODY()
+public:
+	UFUNCTION(BlueprintCallable, Category = "GeometryScript|Baking")
+	static UPARAM(DisplayName="Bake Type Out") FGeometryScriptBakeTypeOptions MakeBakeTypeTangentNormal();
+
+	UFUNCTION(BlueprintCallable, Category = "GeometryScript|Baking")
+	static UPARAM(DisplayName="Bake Type Out") FGeometryScriptBakeTypeOptions MakeBakeTypeObjectNormal();
+	
+	UFUNCTION(BlueprintCallable, Category = "GeometryScript|Baking")
+	static UPARAM(DisplayName="Bake Type Out") FGeometryScriptBakeTypeOptions MakeBakeTypeFaceNormal();
+
+	UFUNCTION(BlueprintCallable, Category = "GeometryScript|Baking")
+	static UPARAM(DisplayName="Bake Type Out") FGeometryScriptBakeTypeOptions MakeBakeTypeBentNormal(
+		int OcclusionRays = 16,
+		float MaxDistance = 0.0f,
+		float SpreadAngle = 180.0f);
+
+	UFUNCTION(BlueprintCallable, Category = "GeometryScript|Baking")
+	static UPARAM(DisplayName="Bake Type Out") FGeometryScriptBakeTypeOptions MakeBakeTypePosition();
+
+	UFUNCTION(BlueprintCallable, Category = "GeometryScript|Baking")
+	static UPARAM(DisplayName="Bake Type Out") FGeometryScriptBakeTypeOptions MakeBakeTypeCurvature(
+		EGeometryScriptBakeCurvatureTypeMode CurvatureType = EGeometryScriptBakeCurvatureTypeMode::Mean,
+		EGeometryScriptBakeCurvatureColorMode ColorMapping = EGeometryScriptBakeCurvatureColorMode::Grayscale,
+		float ColorRangeMultiplier = 1.0f,
+		float MinRangeMultiplier = 0.0f,
+		EGeometryScriptBakeCurvatureClampMode Clamping = EGeometryScriptBakeCurvatureClampMode::None);
+
+	UFUNCTION(BlueprintCallable, Category = "GeometryScript|Baking")
+	static UPARAM(DisplayName="Bake Type Out") FGeometryScriptBakeTypeOptions MakeBakeTypeAmbientOcclusion(
+		int OcclusionRays = 16,
+		float MaxDistance = 0.0f,
+		float SpreadAngle = 180.0f,
+		float BiasAngle = 15.0f);
+
+	UFUNCTION(BlueprintCallable, Category = "GeometryScript|Baking")
+	static UPARAM(DisplayName="Bake Type Out") FGeometryScriptBakeTypeOptions MakeBakeTypeTexture(
+		UTexture2D* SourceTexture = nullptr,
+		int SourceUVLayer = 0);
+
+	UFUNCTION(BlueprintCallable, Category = "GeometryScript|Baking")
+	static UPARAM(DisplayName="Bake Type Out") FGeometryScriptBakeTypeOptions MakeBakeTypeMultiTexture(
+		const TArray<UTexture2D*>& MaterialIDSourceTextures,
+		UPARAM(DisplayName="Source UV Channel") int SourceUVLayer = 0);
+
+	UFUNCTION(BlueprintCallable, Category = "GeometryScript|Baking")
+	static UPARAM(DisplayName="Bake Type Out") FGeometryScriptBakeTypeOptions MakeBakeTypeVertexColor();
+
+	UFUNCTION(BlueprintCallable, Category = "GeometryScript|Baking")
+	static UPARAM(DisplayName="Bake Type Out") FGeometryScriptBakeTypeOptions MakeBakeTypeMaterialID();
+	
+	UFUNCTION(BlueprintCallable, Category = "GeometryScript|Baking")
+	static UPARAM(DisplayName="Textures Out") TArray<UTexture2D*> BakeTexture(
+		UDynamicMesh* TargetMesh,
+		FTransform TargetTransform,
+		FGeometryScriptBakeTargetMeshOptions TargetOptions,
+		UDynamicMesh* SourceMesh,
+		FTransform SourceTransform,
+		FGeometryScriptBakeSourceMeshOptions SourceOptions,
+		const TArray<FGeometryScriptBakeTypeOptions>& BakeTypes,
+		FGeometryScriptBakeTextureOptions BakeOptions,
+		UGeometryScriptDebug* Debug = nullptr);
+};
diff --git a/Engine/Plugins/Experimental/MeshModelingToolsetExp/Source/MeshModelingToolsExp/Public/BakeMeshAttributeToolCommon.h b/Engine/Plugins/Experimental/MeshModelingToolsetExp/Source/MeshModelingToolsExp/Public/BakeMeshAttributeToolCommon.h
--- a/Engine/Plugins/Experimental/MeshModelingToolsetExp/Source/MeshModelingToolsExp/Public/BakeMeshAttributeToolCommon.h
+++ b/Engine/Plugins/Experimental/MeshModelingToolsetExp/Source/MeshModelingToolsExp/Public/BakeMeshAttributeToolCommon.h
@@ -154,7 +154,7 @@ UENUM()
 enum class EBakeCurvatureTypeMode
 {
 	/** Average of the minimum and maximum principal curvatures */
-	MeanAverage,
+	MeanAverage UMETA(DisplayName = "Mean"),
 	/** Maximum principal curvature */
 	Max,
 	/** Minimum principal curvature */
diff --git a/Engine/Plugins/Runtime/GeometryProcessing/Source/DynamicMesh/Public/Sampling/MeshBakerCommon.h b/Engine/Plugins/Runtime/GeometryProcessing/Source/DynamicMesh/Public/Sampling/MeshBakerCommon.h
--- a/Engine/Plugins/Runtime/GeometryProcessing/Source/DynamicMesh/Public/Sampling/MeshBakerCommon.h
+++ b/Engine/Plugins/Runtime/GeometryProcessing/Source/DynamicMesh/Public/Sampling/MeshBakerCommon.h
@@ -105,6 +105,11 @@ public:
 	/** @return the triangle count of a given mesh. */
 	virtual int32 GetTriangleCount(const void* Mesh) const = 0;
 
+	/** Set the tangents for a given mesh */
+	virtual void SetTangents(const void* Mesh, FMeshTangentsd* Tangents)
+	{		
+	}
+
 	/** Associate a texture map and UV layer index for a given mesh in the detail set */
 	virtual void SetTextureMap(const void* Mesh, const FBakeDetailTexture& Map) = 0;
 
@@ -337,6 +342,11 @@ public:
 		return DynamicMesh->TriangleCount();
 	}
 
+	virtual void SetTangents(const void* Mesh, FMeshTangentsd* Tangents) override
+	{
+		DetailTangents = Tangents;
+	}
+
 	virtual void SetTextureMap(const void* Mesh, const FBakeDetailTexture& Map) override
 	{
 		DetailTextureMap = Map;
