GeometryScript: Fix uninitialized SourceTexture in MeshBakeFunctions.

#rnx

diff --git a/Engine/Plugins/Experimental/GeometryScripting/Source/GeometryScriptingCore/Public/GeometryScript/MeshBakeFunctions.h b/Engine/Plugins/Experimental/GeometryScripting/Source/GeometryScriptingCore/Public/GeometryScript/MeshBakeFunctions.h
--- a/Engine/Plugins/Experimental/GeometryScripting/Source/GeometryScriptingCore/Public/GeometryScript/MeshBakeFunctions.h
+++ b/Engine/Plugins/Experimental/GeometryScripting/Source/GeometryScriptingCore/Public/GeometryScript/MeshBakeFunctions.h
@@ -166,7 +166,7 @@ struct GEOMETRYSCRIPTINGCORE_API FGeometryScriptBakeType_Texture : public FGeome
 public:
 	/** Source mesh texture that is to be resampled into a new texture */
 	UPROPERTY(BlueprintReadWrite, Category = TextureOutput)
-	TObjectPtr<UTexture2D> SourceTexture;
+	TObjectPtr<UTexture2D> SourceTexture = nullptr;
 
 	/** UV channel to use for the source mesh texture */
 	UPROPERTY(BlueprintReadWrite, Category = TextureOutput, meta = (DisplayName = "Source Texture UV Channel", ClampMin="0"))
