GeometryScript: Added Async variant of BakeTexture blueprint

diff --git a/Engine/Plugins/Experimental/GeometryScripting/Source/GeometryScriptingCore/Private/GeometryScriptTypes.cpp b/Engine/Plugins/Experimental/GeometryScripting/Source/GeometryScriptingCore/Private/GeometryScriptTypes.cpp
--- a/Engine/Plugins/Experimental/GeometryScripting/Source/GeometryScriptingCore/Private/GeometryScriptTypes.cpp
+++ b/Engine/Plugins/Experimental/GeometryScripting/Source/GeometryScriptingCore/Private/GeometryScriptTypes.cpp
@@ -40,4 +40,22 @@ void UE::Geometry::AppendWarning(UGeometryScriptDebug* Debug, EGeometryScriptErr
 	{
 		Debug->Append(Result);
 	}
-}
\ No newline at end of file
+}
+
+void UE::Geometry::AppendError(TArray<FGeometryScriptDebugMessage>* DebugMessages, EGeometryScriptErrorType ErrorType, const FText& Message)
+{
+	FGeometryScriptDebugMessage Result = MakeScriptError(ErrorType, Message);
+	if (DebugMessages != nullptr)
+	{
+		DebugMessages->Add(Result);
+	}
+}
+
+void UE::Geometry::AppendWarning(TArray<FGeometryScriptDebugMessage>* DebugMessages, EGeometryScriptErrorType WarningType, const FText& Message)
+{
+	FGeometryScriptDebugMessage Result = MakeScriptWarning(WarningType, Message);
+	if (DebugMessages != nullptr)
+	{
+		DebugMessages->Add(Result);
+	}
+}
diff --git a/Engine/Plugins/Experimental/GeometryScripting/Source/GeometryScriptingCore/Private/MeshBakeFunctions.cpp b/Engine/Plugins/Experimental/GeometryScripting/Source/GeometryScriptingCore/Private/MeshBakeFunctions.cpp
--- a/Engine/Plugins/Experimental/GeometryScripting/Source/GeometryScriptingCore/Private/MeshBakeFunctions.cpp
+++ b/Engine/Plugins/Experimental/GeometryScripting/Source/GeometryScriptingCore/Private/MeshBakeFunctions.cpp
@@ -4,7 +4,6 @@
 #include "UDynamicMesh.h"
 
 #include "Sampling/MeshBakerCommon.h"
-#include "Sampling/MeshMapBaker.h"
 #include "Sampling/MeshConstantMapEvaluator.h"
 #include "Sampling/MeshCurvatureMapEvaluator.h"
 #include "Sampling/MeshOcclusionMapEvaluator.h"
@@ -16,6 +15,7 @@
 
 #include "AssetUtils/Texture2DBuilder.h"
 #include "AssetUtils/Texture2DUtil.h"
+#include "Async/Async.h"
 
 using namespace UE::Geometry;
 
@@ -85,42 +85,73 @@ namespace GeometryScriptBakeLocals
 		}
 		return Samples;
 	}
-	
-	FTexture2DBuilder::ETextureType GetTextureType(const EGeometryScriptBakeTypes BakeType, const EGeometryScriptBakeBitDepth MapFormat)
+
+	FTexture2DBuilder::ETextureType GetTextureType(const FMeshMapEvaluator* Evaluator, const EGeometryScriptBakeBitDepth MapFormat)
 	{
 		FTexture2DBuilder::ETextureType TexType = FTexture2DBuilder::ETextureType::Color;
-		switch (BakeType)
+		switch (Evaluator->Type())
 		{
 		default:
 			checkNoEntry();
 			break;
-		case EGeometryScriptBakeTypes::TangentSpaceNormal:
+		case EMeshMapEvaluatorType::Normal:
+		{
 			TexType = FTexture2DBuilder::ETextureType::NormalMap;
 			break;
-		case EGeometryScriptBakeTypes::AmbientOcclusion:
-			TexType = FTexture2DBuilder::ETextureType::AmbientOcclusion;
+		}
+		case EMeshMapEvaluatorType::Occlusion:
+		{
+			const FMeshOcclusionMapEvaluator* OcclusionEval = static_cast<const FMeshOcclusionMapEvaluator*>(Evaluator); 
+			if (static_cast<bool>(OcclusionEval->OcclusionType & EMeshOcclusionMapType::AmbientOcclusion))
+			{
+				ensure(OcclusionEval->OcclusionType == EMeshOcclusionMapType::AmbientOcclusion);
+				TexType = FTexture2DBuilder::ETextureType::AmbientOcclusion;
+			}
+			else if (static_cast<bool>(OcclusionEval->OcclusionType & EMeshOcclusionMapType::BentNormal))
+			{
+				ensure(OcclusionEval->OcclusionType == EMeshOcclusionMapType::BentNormal);
+				TexType = FTexture2DBuilder::ETextureType::NormalMap;
+			}
 			break;
-		case EGeometryScriptBakeTypes::BentNormal:
-			TexType = FTexture2DBuilder::ETextureType::NormalMap;
+		}
+		case EMeshMapEvaluatorType::Property:
+		{
+			const FMeshPropertyMapEvaluator* PropertyEval = static_cast<const FMeshPropertyMapEvaluator*>(Evaluator);
+			switch (PropertyEval->Property)
+			{
+			case EMeshPropertyMapType::Normal:
+			case EMeshPropertyMapType::FacetNormal:
+			case EMeshPropertyMapType::Position:
+			case EMeshPropertyMapType::UVPosition:
+				TexType = FTexture2DBuilder::ETextureType::ColorLinear;
+				break;
+			case EMeshPropertyMapType::VertexColor:
+			case EMeshPropertyMapType::MaterialID:
+				TexType = FTexture2DBuilder::ETextureType::Color;
+				break;
+			}
 			break;
-		case EGeometryScriptBakeTypes::Curvature:
-		case EGeometryScriptBakeTypes::ObjectSpaceNormal:
-		case EGeometryScriptBakeTypes::FaceNormal:
-		case EGeometryScriptBakeTypes::Position:
+		}
+		case EMeshMapEvaluatorType::Curvature:
+		{
 			TexType = FTexture2DBuilder::ETextureType::ColorLinear;
 			break;
-		case EGeometryScriptBakeTypes::MaterialID:
-		case EGeometryScriptBakeTypes::VertexColor:
-			break;
-		case EGeometryScriptBakeTypes::Texture:
-		case EGeometryScriptBakeTypes::MultiTexture:
+		}
+		case EMeshMapEvaluatorType::ResampleImage:
+		case EMeshMapEvaluatorType::MultiResampleImage:
+		{
 			// For texture output with 16-bit source data, output HDR texture
 			if (MapFormat == EGeometryScriptBakeBitDepth::ChannelBits16)
 			{
 				TexType = FTexture2DBuilder::ETextureType::EmissiveHDR;
 			}
+			else
+			{
+				TexType = FTexture2DBuilder::ETextureType::Color;
+			}
 			break;
 		}
+		}
 		return TexType;
 	}
 
@@ -180,6 +211,298 @@ namespace GeometryScriptBakeLocals
 		}
 		return Result;
 	}
+
+	TUniquePtr<FMeshMapBaker> BakeTextureImpl(
+		UDynamicMesh* TargetMesh,
+		FTransform TargetTransform,
+		FGeometryScriptBakeTargetMeshOptions TargetOptions,
+		UDynamicMesh* SourceMesh,
+		FTransform SourceTransform,
+		FGeometryScriptBakeSourceMeshOptions SourceOptions,
+		const TArray<FGeometryScriptBakeTypeOptions>& BakeTypes,
+		FGeometryScriptBakeTextureOptions BakeOptions,
+		TArray<FGeometryScriptDebugMessage>* Debug)
+	{
+		if (TargetMesh == nullptr)
+		{
+			UE::Geometry::AppendError(Debug, EGeometryScriptErrorType::InvalidInputs, LOCTEXT("BakeTexture_InvalidTargetMesh", "BakeTexture: TargetMesh is Null"));
+			return nullptr;
+		}
+		if (SourceMesh == nullptr)
+		{
+			UE::Geometry::AppendError(Debug, EGeometryScriptErrorType::InvalidInputs, LOCTEXT("BakeTexture_InvalidSourceMesh", "BakeTexture: SourceMesh is Null"));
+			return nullptr;
+		}
+		if (BakeTypes.Num() == 0)
+		{
+			UE::Geometry::AppendWarning(Debug, EGeometryScriptErrorType::InvalidInputs, LOCTEXT("BakeTexture_BakeTypesEmpty", "BakeTexture: BakeTypes is empty"));
+			return nullptr;
+		}
+
+		TSharedPtr<FMeshTangentsd, ESPMode::ThreadSafe> TargetMeshTangents;
+		TSharedPtr<FMeshTangentsd, ESPMode::ThreadSafe> SourceMeshTangents;
+		auto GetMeshTangents = [TargetMesh, SourceMesh, Debug](UDynamicMesh* Mesh, TSharedPtr<FMeshTangentsd, ESPMode::ThreadSafe>& Tangents)
+		{
+			if (!Tangents)
+			{
+				Tangents = MakeShared<FMeshTangentsd, ESPMode::ThreadSafe>(Mesh->GetMeshPtr());
+				Tangents->CopyTriVertexTangents(Mesh->GetMeshRef());
+
+				// Validate the tangents
+				if (!FDynamicMeshTangents(Mesh->GetMeshPtr()).HasValidTangents(true))
+				{
+					if (Mesh == TargetMesh)
+					{
+						UE::Geometry::AppendError(Debug, EGeometryScriptErrorType::InvalidInputs, LOCTEXT("BakeTexture_InvalidTargetTangents", "BakeTexture: Target Mesh tangents are invalid."));
+					}
+					else
+					{
+						ensure(Mesh == SourceMesh);
+						UE::Geometry::AppendError(Debug, EGeometryScriptErrorType::InvalidInputs, LOCTEXT("BakeTexture_InvalidSourceTangents", "BakeTexture: Source Mesh tangents are invalid."));
+					}
+					return false;
+				}
+			}
+			return true;
+		};
+
+		const bool bIsBakeToSelf = (TargetMesh == SourceMesh);
+
+		FDynamicMesh3 SourceMeshCopy;
+		const FDynamicMesh3* SourceMeshOriginal = SourceMesh->GetMeshPtr();
+		const FDynamicMesh3* SourceMeshToUse = SourceMeshOriginal;
+		if (BakeOptions.bProjectionInWorldSpace && !bIsBakeToSelf)
+		{
+			// Transform the SourceMesh into TargetMesh local space using a copy (oof)
+			// TODO: Remove this once we have support for transforming rays in the core bake loop
+			SourceMeshCopy = *SourceMeshOriginal;
+			const FTransformSRT3d SourceToWorld = SourceTransform;
+			MeshTransforms::ApplyTransform(SourceMeshCopy, SourceToWorld);
+			const FTransformSRT3d TargetToWorld = TargetTransform;
+			MeshTransforms::ApplyTransform(SourceMeshCopy, TargetToWorld.Inverse());
+			SourceMeshToUse = &SourceMeshCopy;
+		}
+
+		FImageDimensions BakeDimensions = GetDimensions(BakeOptions.Resolution);
+		const FDynamicMeshAABBTree3 DetailSpatial(SourceMeshToUse);
+		FMeshBakerDynamicMeshSampler DetailSampler(SourceMeshToUse, &DetailSpatial);
+
+		TUniquePtr<FMeshMapBaker> Result = MakeUnique<FMeshMapBaker>();
+		FMeshMapBaker& Baker = *Result;
+		Baker.SetTargetMesh(TargetMesh->GetMeshPtr());
+		Baker.SetTargetMeshUVLayer(TargetOptions.TargetUVLayer);
+		Baker.SetDetailSampler(&DetailSampler);
+		Baker.SetDimensions(BakeDimensions);
+		Baker.SetProjectionDistance(BakeOptions.ProjectionDistance);
+		Baker.SetSamplesPerPixel(GetSamplesPerPixel(BakeOptions.SamplesPerPixel));
+
+		bool bSupportsSourceNormalMap = false;
+		TSharedPtr<UE::Geometry::TImageBuilder<FVector4f>, ESPMode::ThreadSafe> SourceTexture;
+		for (const FGeometryScriptBakeTypeOptions& Options : BakeTypes)
+		{
+			switch(Options.BakeType)
+			{
+			case EGeometryScriptBakeTypes::TangentSpaceNormal:
+			{
+				if (!GetMeshTangents(TargetMesh, TargetMeshTangents))
+				{
+					return nullptr;
+				}
+				TSharedPtr<FMeshNormalMapEvaluator> NormalEval = MakeShared<FMeshNormalMapEvaluator>();
+				bSupportsSourceNormalMap = true;
+				Baker.AddEvaluator(NormalEval);
+				break;
+			}
+			case EGeometryScriptBakeTypes::ObjectSpaceNormal:
+			{
+				TSharedPtr<FMeshPropertyMapEvaluator> PropertyEval = MakeShared<FMeshPropertyMapEvaluator>();
+				PropertyEval->Property = EMeshPropertyMapType::Normal;
+				bSupportsSourceNormalMap = true;
+				Baker.AddEvaluator(PropertyEval);
+				break;
+			}
+			case EGeometryScriptBakeTypes::FaceNormal:
+			{
+				TSharedPtr<FMeshPropertyMapEvaluator> PropertyEval = MakeShared<FMeshPropertyMapEvaluator>();
+				PropertyEval->Property = EMeshPropertyMapType::FacetNormal;
+				Baker.AddEvaluator(PropertyEval);
+				break;
+			}
+			case EGeometryScriptBakeTypes::BentNormal:
+			{
+				if (!GetMeshTangents(TargetMesh, TargetMeshTangents))
+				{
+					return nullptr;
+				}
+				FGeometryScriptBakeType_Occlusion* OcclusionOptions = static_cast<FGeometryScriptBakeType_Occlusion*>(Options.Options.Get());
+				TSharedPtr<FMeshOcclusionMapEvaluator> OcclusionEval = MakeShared<FMeshOcclusionMapEvaluator>();
+				OcclusionEval->OcclusionType = EMeshOcclusionMapType::BentNormal;
+				OcclusionEval->NumOcclusionRays = OcclusionOptions->OcclusionRays;
+				OcclusionEval->MaxDistance = (OcclusionOptions->MaxDistance == 0) ? TNumericLimits<float>::Max() : OcclusionOptions->MaxDistance;
+				OcclusionEval->SpreadAngle = OcclusionOptions->SpreadAngle;
+				Baker.AddEvaluator(OcclusionEval);
+				break;
+			}
+			case EGeometryScriptBakeTypes::Position:
+			{
+				TSharedPtr<FMeshPropertyMapEvaluator> PropertyEval = MakeShared<FMeshPropertyMapEvaluator>();
+				PropertyEval->Property = EMeshPropertyMapType::Position;
+				Baker.AddEvaluator(PropertyEval);
+				break;
+			}
+			case EGeometryScriptBakeTypes::Curvature:
+			{
+				FGeometryScriptBakeType_Curvature* CurvatureOptions = static_cast<FGeometryScriptBakeType_Curvature*>(Options.Options.Get());
+				TSharedPtr<FMeshCurvatureMapEvaluator> CurvatureEval = MakeShared<FMeshCurvatureMapEvaluator>();
+				CurvatureEval->UseCurvatureType = GetCurvatureType(CurvatureOptions->CurvatureType);
+				CurvatureEval->UseColorMode = GetCurvatureColorMode(CurvatureOptions->ColorMapping);
+				CurvatureEval->RangeScale = CurvatureOptions->ColorRangeMultiplier;
+				CurvatureEval->MinRangeScale = CurvatureOptions->MinRangeMultiplier;
+				CurvatureEval->UseClampMode = GetCurvatureClampMode(CurvatureOptions->Clamping);
+				Baker.AddEvaluator(CurvatureEval);
+				break;
+			}
+			case EGeometryScriptBakeTypes::AmbientOcclusion:
+			{
+				FGeometryScriptBakeType_Occlusion* OcclusionOptions = static_cast<FGeometryScriptBakeType_Occlusion*>(Options.Options.Get());
+				TSharedPtr<FMeshOcclusionMapEvaluator> OcclusionEval = MakeShared<FMeshOcclusionMapEvaluator>();
+				OcclusionEval->OcclusionType = EMeshOcclusionMapType::AmbientOcclusion;
+				OcclusionEval->NumOcclusionRays = OcclusionOptions->OcclusionRays;
+				OcclusionEval->MaxDistance = (OcclusionOptions->MaxDistance == 0) ? TNumericLimits<float>::Max() : OcclusionOptions->MaxDistance;
+				OcclusionEval->SpreadAngle = OcclusionOptions->SpreadAngle;
+				OcclusionEval->BiasAngleDeg = OcclusionOptions->BiasAngle;
+				Baker.AddEvaluator(OcclusionEval);
+				break;
+			}
+			case EGeometryScriptBakeTypes::Texture:
+			{
+				FGeometryScriptBakeType_Texture* TextureOptions = static_cast<FGeometryScriptBakeType_Texture*>(Options.Options.Get());
+
+				// TODO: Add support for sampling different texture maps per Texture evaluator in a single pass. 
+				if (!SourceTexture && TextureOptions->SourceTexture)
+				{
+					SourceTexture = MakeShared<TImageBuilder<FVector4f>>();
+					if (!UE::AssetUtils::ReadTexture(TextureOptions->SourceTexture, *SourceTexture, false))
+					{
+						UE::Geometry::AppendError(Debug, EGeometryScriptErrorType::InvalidInputs, LOCTEXT("BakeTexture_InvalidSourceTexture", "BakeTexture: Failed to read SourceTexture"));
+					}
+					else
+					{
+						DetailSampler.SetTextureMap(SourceMeshToUse, IMeshBakerDetailSampler::FBakeDetailTexture(SourceTexture.Get(), TextureOptions->SourceUVLayer));
+					}
+				}
+				
+				TSharedPtr<FMeshResampleImageEvaluator> TextureEval = MakeShared<FMeshResampleImageEvaluator>();
+				Baker.AddEvaluator(TextureEval);
+				break;
+			}
+			case EGeometryScriptBakeTypes::MultiTexture:
+			{
+				FGeometryScriptBakeType_MultiTexture* TextureOptions = static_cast<FGeometryScriptBakeType_MultiTexture*>(Options.Options.Get());
+				TSharedPtr<FMeshMultiResampleImageEvaluator> TextureEval = MakeShared<FMeshMultiResampleImageEvaluator>();
+
+				if (TextureOptions->MaterialIDSourceTextures.Num())
+				{
+					TextureEval->MultiTextures.SetNum(TextureOptions->MaterialIDSourceTextures.Num());
+					for (int32 MaterialId = 0; MaterialId < TextureEval->MultiTextures.Num(); ++MaterialId)
+					{
+						if (UTexture2D* Texture = TextureOptions->MaterialIDSourceTextures[MaterialId])
+						{
+							TextureEval->MultiTextures[MaterialId] = MakeShared<TImageBuilder<FVector4f>>();
+							if (!UE::AssetUtils::ReadTexture(Texture, *TextureEval->MultiTextures[MaterialId], false))
+							{
+								UE::Geometry::AppendError(Debug, EGeometryScriptErrorType::InvalidInputs, LOCTEXT("BakeTexture_InvalidMultiTexture", "BakeTexture: Failed to read MaterialIDSourceTexture"));
+							}
+						}
+					}
+				}
+				Baker.AddEvaluator(TextureEval);
+				break;
+			}
+			case EGeometryScriptBakeTypes::VertexColor:
+			{
+				TSharedPtr<FMeshPropertyMapEvaluator> PropertyEval = MakeShared<FMeshPropertyMapEvaluator>();
+				PropertyEval->Property = EMeshPropertyMapType::VertexColor;
+				Baker.AddEvaluator(PropertyEval);
+				break;
+			}
+			case EGeometryScriptBakeTypes::MaterialID:
+			{
+				TSharedPtr<FMeshPropertyMapEvaluator> PropertyEval = MakeShared<FMeshPropertyMapEvaluator>();
+				PropertyEval->Property = EMeshPropertyMapType::MaterialID;
+				Baker.AddEvaluator(PropertyEval);
+				break;
+			}
+			default:
+				break;
+			}
+		}
+
+		TSharedPtr<UE::Geometry::TImageBuilder<FVector4f>, ESPMode::ThreadSafe> SourceNormalMap;
+		if (bSupportsSourceNormalMap && SourceOptions.SourceNormalMap)
+		{
+			SourceNormalMap = MakeShared<TImageBuilder<FVector4f>>();
+			if (!UE::AssetUtils::ReadTexture(SourceOptions.SourceNormalMap, *SourceNormalMap, false))
+			{
+				UE::Geometry::AppendError(Debug, EGeometryScriptErrorType::InvalidInputs, LOCTEXT("BakeTexture_InvalidSourceNormalMap", "BakeTexture: Failed to read SourceNormalMap"));
+			}
+			else
+			{
+				DetailSampler.SetNormalMap(SourceMeshToUse,
+				IMeshBakerDetailSampler::FBakeDetailTexture(
+						SourceNormalMap.Get(),
+						SourceOptions.SourceNormalUVLayer));
+
+				if (!GetMeshTangents(SourceMesh, SourceMeshTangents))
+				{
+					return nullptr;
+				}
+				DetailSampler.SetTangents(SourceMeshToUse, SourceMeshTangents.Get());
+			}
+		}
+
+		if (TargetMeshTangents)
+		{
+			Baker.SetTargetMeshTangents(TargetMeshTangents);
+		}
+
+		Baker.Bake();
+
+		return MoveTemp(Result);
+	}
+
+	void GetTexturesFromBaker(FMeshMapBaker* Baker, const EGeometryScriptBakeBitDepth BakeBitDepth, TArray<UTexture2D*>& Textures)
+	{
+		if (!Baker)
+		{
+			return;
+		}
+		
+		const FImageDimensions BakeDimensions = Baker->GetDimensions();
+		const int NumEval = Baker->NumEvaluators();
+		for (int EvalIdx = 0; EvalIdx < NumEval; ++EvalIdx)
+		{
+			// For 8-bit color textures, ensure that the source data is in sRGB.
+			const FTexture2DBuilder::ETextureType TexType = GetTextureType(Baker->GetEvaluator(EvalIdx), BakeBitDepth);
+			const bool bConvertToSRGB = (TexType == FTexture2DBuilder::ETextureType::Color);
+			const ETextureSourceFormat SourceDataFormat = (BakeBitDepth == EGeometryScriptBakeBitDepth::ChannelBits16 ? TSF_RGBA16F : TSF_BGRA8);
+
+			constexpr int ResultIdx = 0;
+			FTexture2DBuilder TextureBuilder;
+			TextureBuilder.Initialize(TexType, BakeDimensions);
+			TextureBuilder.Copy(*Baker->GetBakeResults(EvalIdx)[ResultIdx], bConvertToSRGB);
+			TextureBuilder.Commit(false);
+
+			// Copy image to source data after commit. This will avoid incurring
+			// the cost of hitting the DDC for texture compile while iterating on
+			// bake settings. Since this dirties the texture, the next time the texture
+			// is used after accepting the final texture, the DDC will trigger and
+			// properly recompile the platform data.
+			const bool bConvertSourceToSRGB = bConvertToSRGB && SourceDataFormat == TSF_BGRA8;
+			TextureBuilder.CopyImageToSourceData(*Baker->GetBakeResults(EvalIdx)[ResultIdx], SourceDataFormat, bConvertSourceToSRGB);
+			Textures.Add(TextureBuilder.GetTexture2D());
+		}
+	}
 }
 
 
@@ -315,251 +638,86 @@ TArray<UTexture2D*> UGeometryScriptLibrary_MeshBakeFunctions::BakeTexture(
 	FGeometryScriptBakeTextureOptions BakeOptions,
 	UGeometryScriptDebug* Debug)
 {
-	TArray<UTexture2D*> TextureOutput;
-	if (TargetMesh == nullptr)
-	{
-		UE::Geometry::AppendError(Debug, EGeometryScriptErrorType::InvalidInputs, LOCTEXT("BakeTexture_InvalidTargetMesh", "BakeTexture: TargetMesh is Null"));
-		return TextureOutput;
-	}
-	if (SourceMesh == nullptr)
-	{
-		UE::Geometry::AppendError(Debug, EGeometryScriptErrorType::InvalidInputs, LOCTEXT("BakeTexture_InvalidSourceMesh", "BakeTexture: SourceMesh is Null"));
-		return TextureOutput;
-	}
-	if (BakeTypes.Num() == 0)
-	{
-		return TextureOutput;
-	}
+	const TUniquePtr<FMeshMapBaker> Baker = GeometryScriptBakeLocals::BakeTextureImpl(
+		TargetMesh,
+		TargetTransform,
+		TargetOptions,
+		SourceMesh,
+		SourceTransform,
+		SourceOptions,
+		BakeTypes,
+		BakeOptions,
+		Debug ? &Debug->Messages : nullptr);
 
-	TSharedPtr<FMeshTangentsd, ESPMode::ThreadSafe> TargetMeshTangents;
-	TSharedPtr<FMeshTangentsd, ESPMode::ThreadSafe> SourceMeshTangents;
-	auto GetMeshTangents = [](UDynamicMesh* Mesh, TSharedPtr<FMeshTangentsd, ESPMode::ThreadSafe>& Tangents)
-	{
-		if (!Tangents)
-		{
-			Tangents = MakeShared<FMeshTangentsd, ESPMode::ThreadSafe>(Mesh->GetMeshPtr());
-			Tangents->CopyTriVertexTangents(Mesh->GetMeshRef());
-		}
-	};
-
-	const bool bIsBakeToSelf = (TargetMesh == SourceMesh);
+	TArray<UTexture2D*> TextureOutput;
+	GeometryScriptBakeLocals::GetTexturesFromBaker(Baker.Get(), BakeOptions.BitDepth, TextureOutput);
 
-	FDynamicMesh3 SourceMeshCopy;
-	const FDynamicMesh3* SourceMeshOriginal = SourceMesh->GetMeshPtr();
-	const FDynamicMesh3* SourceMeshToUse = SourceMeshOriginal;
-	if (BakeOptions.bProjectionInWorldSpace && !bIsBakeToSelf)
-	{
-		// Transform the SourceMesh into TargetMesh local space using a copy (oof)
-		// TODO: Remove this once we have support for transforming rays in the core bake loop
-		SourceMeshCopy = *SourceMeshOriginal;
-		const FTransformSRT3d SourceToWorld = SourceTransform;
-		MeshTransforms::ApplyTransform(SourceMeshCopy, SourceToWorld);
-		const FTransformSRT3d TargetToWorld = TargetTransform;
-		MeshTransforms::ApplyTransform(SourceMeshCopy, TargetToWorld.Inverse());
-		SourceMeshToUse = &SourceMeshCopy;
-	}
+	return TextureOutput;
+}
 
-	FImageDimensions BakeDimensions = GeometryScriptBakeLocals::GetDimensions(BakeOptions.Resolution);
-	const FDynamicMeshAABBTree3 DetailSpatial(SourceMeshToUse);
-	FMeshBakerDynamicMeshSampler DetailSampler(SourceMeshToUse, &DetailSpatial);
-
-	FMeshMapBaker Baker;
-	Baker.SetTargetMesh(TargetMesh->GetMeshPtr());
-	Baker.SetTargetMeshUVLayer(TargetOptions.TargetUVLayer);
-	Baker.SetDetailSampler(&DetailSampler);
-	Baker.SetDimensions(BakeDimensions);
-	Baker.SetProjectionDistance(BakeOptions.ProjectionDistance);
-	Baker.SetSamplesPerPixel(GeometryScriptBakeLocals::GetSamplesPerPixel(BakeOptions.SamplesPerPixel));
-
-	bool bSupportsSourceNormalMap = false;
-	TSharedPtr<UE::Geometry::TImageBuilder<FVector4f>, ESPMode::ThreadSafe> SourceTexture;
-	for (const FGeometryScriptBakeTypeOptions& Options : BakeTypes)
+void UGeometryScriptLibrary_MeshBakeFunctions::BakeTextureAsyncBegin(
+	const FBakeTextureDelegate& Completed,
+	const int BakeId,
+	UDynamicMesh* TargetMesh,
+	FTransform TargetTransform,
+	FGeometryScriptBakeTargetMeshOptions TargetOptions,
+	UDynamicMesh* SourceMesh,
+	FTransform SourceTransform,
+	FGeometryScriptBakeSourceMeshOptions SourceOptions,
+	const TArray<FGeometryScriptBakeTypeOptions>& BakeTypes,
+	FGeometryScriptBakeTextureOptions BakeOptions,
+	UGeometryScriptDebug* Debug)
+{
+	AsyncTask(ENamedThreads::AnyBackgroundThreadNormalTask, [Completed, BakeId, TargetMesh, TargetTransform, TargetOptions, SourceMesh, SourceTransform, SourceOptions, BakeTypes, BakeOptions, Debug]()
 	{
-		switch(Options.BakeType)
-		{
-		case EGeometryScriptBakeTypes::TangentSpaceNormal:
-		{
-			GetMeshTangents(TargetMesh, TargetMeshTangents);
-			TSharedPtr<FMeshNormalMapEvaluator> NormalEval = MakeShared<FMeshNormalMapEvaluator>();
-			bSupportsSourceNormalMap = true;
-			Baker.AddEvaluator(NormalEval);
-			break;
-		}
-		case EGeometryScriptBakeTypes::ObjectSpaceNormal:
-		{
-			TSharedPtr<FMeshPropertyMapEvaluator> PropertyEval = MakeShared<FMeshPropertyMapEvaluator>();
-			PropertyEval->Property = EMeshPropertyMapType::Normal;
-			bSupportsSourceNormalMap = true;
-			Baker.AddEvaluator(PropertyEval);
-			break;
-		}
-		case EGeometryScriptBakeTypes::FaceNormal:
-		{
-			TSharedPtr<FMeshPropertyMapEvaluator> PropertyEval = MakeShared<FMeshPropertyMapEvaluator>();
-			PropertyEval->Property = EMeshPropertyMapType::FacetNormal;
-			Baker.AddEvaluator(PropertyEval);
-			break;
-		}
-		case EGeometryScriptBakeTypes::BentNormal:
-		{
-			GetMeshTangents(TargetMesh, TargetMeshTangents);
-			FGeometryScriptBakeType_Occlusion* OcclusionOptions = static_cast<FGeometryScriptBakeType_Occlusion*>(Options.Options.Get());
-			TSharedPtr<FMeshOcclusionMapEvaluator> OcclusionEval = MakeShared<FMeshOcclusionMapEvaluator>();
-			OcclusionEval->OcclusionType = EMeshOcclusionMapType::BentNormal;
-			OcclusionEval->NumOcclusionRays = OcclusionOptions->OcclusionRays;
-			OcclusionEval->MaxDistance = (OcclusionOptions->MaxDistance == 0) ? TNumericLimits<float>::Max() : OcclusionOptions->MaxDistance;
-			OcclusionEval->SpreadAngle = OcclusionOptions->SpreadAngle;
-			Baker.AddEvaluator(OcclusionEval);
-			break;
-		}
-		case EGeometryScriptBakeTypes::Position:
-		{
-			TSharedPtr<FMeshPropertyMapEvaluator> PropertyEval = MakeShared<FMeshPropertyMapEvaluator>();
-			PropertyEval->Property = EMeshPropertyMapType::Position;
-			Baker.AddEvaluator(PropertyEval);
-			break;
-		}
-		case EGeometryScriptBakeTypes::Curvature:
+		// FGeometryScriptDebug is not thread-safe. Instantiate a local
+		// debug message array to collect any debug errors and append those
+		// messages back on the game thread.
+		TSharedPtr<TArray<FGeometryScriptDebugMessage>, ESPMode::ThreadSafe> DebugMessages;
+		if (Debug)
 		{
-			FGeometryScriptBakeType_Curvature* CurvatureOptions = static_cast<FGeometryScriptBakeType_Curvature*>(Options.Options.Get());
-			TSharedPtr<FMeshCurvatureMapEvaluator> CurvatureEval = MakeShared<FMeshCurvatureMapEvaluator>();
-			CurvatureEval->UseCurvatureType = GeometryScriptBakeLocals::GetCurvatureType(CurvatureOptions->CurvatureType);
-			CurvatureEval->UseColorMode = GeometryScriptBakeLocals::GetCurvatureColorMode(CurvatureOptions->ColorMapping);
-			CurvatureEval->RangeScale = CurvatureOptions->ColorRangeMultiplier;
-			CurvatureEval->MinRangeScale = CurvatureOptions->MinRangeMultiplier;
-			CurvatureEval->UseClampMode = GeometryScriptBakeLocals::GetCurvatureClampMode(CurvatureOptions->Clamping);
-			Baker.AddEvaluator(CurvatureEval);
-			break;
+			DebugMessages = MakeShared<TArray<FGeometryScriptDebugMessage>, ESPMode::ThreadSafe>();
 		}
-		case EGeometryScriptBakeTypes::AmbientOcclusion:
-		{
-			FGeometryScriptBakeType_Occlusion* OcclusionOptions = static_cast<FGeometryScriptBakeType_Occlusion*>(Options.Options.Get());
-			TSharedPtr<FMeshOcclusionMapEvaluator> OcclusionEval = MakeShared<FMeshOcclusionMapEvaluator>();
-			OcclusionEval->OcclusionType = EMeshOcclusionMapType::AmbientOcclusion;
-			OcclusionEval->NumOcclusionRays = OcclusionOptions->OcclusionRays;
-			OcclusionEval->MaxDistance = (OcclusionOptions->MaxDistance == 0) ? TNumericLimits<float>::Max() : OcclusionOptions->MaxDistance;
-			OcclusionEval->SpreadAngle = OcclusionOptions->SpreadAngle;
-			OcclusionEval->BiasAngleDeg = OcclusionOptions->BiasAngle;
-			Baker.AddEvaluator(OcclusionEval);
-			break;
-		}
-		case EGeometryScriptBakeTypes::Texture:
+		
+		TUniquePtr<FMeshMapBaker> Baker = GeometryScriptBakeLocals::BakeTextureImpl(
+			TargetMesh,
+			TargetTransform,
+			TargetOptions,
+			SourceMesh,
+			SourceTransform,
+			SourceOptions,
+			BakeTypes,
+			BakeOptions,
+			DebugMessages.Get());
+
+		TSharedPtr<FMeshMapBaker> SharedBaker = MakeShareable<FMeshMapBaker>(Baker.Release());
+
+		AsyncTask(ENamedThreads::GameThread, [Completed, BakeId, SharedBaker, BakeOptions, Debug, DebugMessages]()
 		{
-			FGeometryScriptBakeType_Texture* TextureOptions = static_cast<FGeometryScriptBakeType_Texture*>(Options.Options.Get());
-
-			// TODO: Add support for sampling different texture maps per Texture evaluator in a single pass. 
-			if (!SourceTexture && TextureOptions->SourceTexture)
+			if (Debug && DebugMessages)
 			{
-				SourceTexture = MakeShared<TImageBuilder<FVector4f>>();
-				if (!UE::AssetUtils::ReadTexture(TextureOptions->SourceTexture, *SourceTexture, false))
-				{
-					UE::Geometry::AppendError(Debug, EGeometryScriptErrorType::InvalidInputs, LOCTEXT("BakeTexture_InvalidSourceTexture", "BakeTexture: Failed to read SourceTexture"));
-				}
-				else
+				for (const FGeometryScriptDebugMessage& Message : *DebugMessages)
 				{
-					DetailSampler.SetTextureMap(SourceMeshToUse, IMeshBakerDetailSampler::FBakeDetailTexture(SourceTexture.Get(), TextureOptions->SourceUVLayer));
+					Debug->Messages.Add(Message);
 				}
 			}
 			
-			TSharedPtr<FMeshResampleImageEvaluator> TextureEval = MakeShared<FMeshResampleImageEvaluator>();
-			Baker.AddEvaluator(TextureEval);
-			break;
-		}
-		case EGeometryScriptBakeTypes::MultiTexture:
-		{
-			FGeometryScriptBakeType_MultiTexture* TextureOptions = static_cast<FGeometryScriptBakeType_MultiTexture*>(Options.Options.Get());
-			TSharedPtr<FMeshMultiResampleImageEvaluator> TextureEval = MakeShared<FMeshMultiResampleImageEvaluator>();
-
-			if (TextureOptions->MaterialIDSourceTextures.Num())
-			{
-				TextureEval->MultiTextures.SetNum(TextureOptions->MaterialIDSourceTextures.Num());
-				for (int32 MaterialId = 0; MaterialId < TextureEval->MultiTextures.Num(); ++MaterialId)
-				{
-					if (UTexture2D* Texture = TextureOptions->MaterialIDSourceTextures[MaterialId])
-					{
-						TextureEval->MultiTextures[MaterialId] = MakeShared<TImageBuilder<FVector4f>>();
-						if (!UE::AssetUtils::ReadTexture(Texture, *TextureEval->MultiTextures[MaterialId], false))
-						{
-							UE::Geometry::AppendError(Debug, EGeometryScriptErrorType::InvalidInputs, LOCTEXT("BakeTexture_InvalidMultiTexture", "BakeTexture: Failed to read MaterialIDSourceTexture"));
-						}
-					}
-				}
-			}
-			Baker.AddEvaluator(TextureEval);
-			break;
-		}
-		case EGeometryScriptBakeTypes::VertexColor:
-		{
-			TSharedPtr<FMeshPropertyMapEvaluator> PropertyEval = MakeShared<FMeshPropertyMapEvaluator>();
-			PropertyEval->Property = EMeshPropertyMapType::VertexColor;
-			Baker.AddEvaluator(PropertyEval);
-			break;
-		}
-		case EGeometryScriptBakeTypes::MaterialID:
-		{
-			TSharedPtr<FMeshPropertyMapEvaluator> PropertyEval = MakeShared<FMeshPropertyMapEvaluator>();
-			PropertyEval->Property = EMeshPropertyMapType::MaterialID;
-			Baker.AddEvaluator(PropertyEval);
-			break;
-		}
-		default:
-			break;
-		}
-	}
-
-	TSharedPtr<UE::Geometry::TImageBuilder<FVector4f>, ESPMode::ThreadSafe> SourceNormalMap;
-	if (bSupportsSourceNormalMap && SourceOptions.SourceNormalMap)
-	{
-		SourceNormalMap = MakeShared<TImageBuilder<FVector4f>>();
-		if (!UE::AssetUtils::ReadTexture(SourceOptions.SourceNormalMap, *SourceNormalMap, false))
-		{
-			UE::Geometry::AppendError(Debug, EGeometryScriptErrorType::InvalidInputs, LOCTEXT("BakeTexture_InvalidSourceNormalMap", "BakeTexture: Failed to read SourceNormalMap"));
-		}
-		else
-		{
-			DetailSampler.SetNormalMap(SourceMeshToUse,
-			IMeshBakerDetailSampler::FBakeDetailTexture(
-					SourceNormalMap.Get(),
-					SourceOptions.SourceNormalUVLayer));
-
-			GetMeshTangents(SourceMesh, SourceMeshTangents);
-			DetailSampler.SetTangents(SourceMeshToUse, SourceMeshTangents.Get());
-		}
-	}
-
-	if (TargetMeshTangents)
-	{
-		Baker.SetTargetMeshTangents(TargetMeshTangents);
-	}
+			FGeometryScriptBakeTextureAsyncResult Result;
+			Result.BakeOptions = BakeOptions;
+			Result.BakeResult = SharedBaker;
+			Completed.ExecuteIfBound(BakeId, Result);
+		});
+	});
+}
 
-	Baker.Bake();
+TArray<UTexture2D*> UGeometryScriptLibrary_MeshBakeFunctions::BakeTextureAsyncEnd(
+	const FGeometryScriptBakeTextureAsyncResult& Result)
+{
+	FMeshMapBaker* Baker = Result.BakeResult.Get();
+	const FGeometryScriptBakeTextureOptions BakeOptions = Result.BakeOptions;
 
-	const int NumEval = Baker.NumEvaluators();
-	for (int EvalIdx = 0; EvalIdx < NumEval; ++EvalIdx)
-	{
-		EGeometryScriptBakeTypes BakeType = BakeTypes[EvalIdx].BakeType;
-		const EGeometryScriptBakeBitDepth BakeBitDepth = BakeOptions.BitDepth;
-		
-		// For 8-bit color textures, ensure that the source data is in sRGB.
-		const FTexture2DBuilder::ETextureType TexType = GeometryScriptBakeLocals::GetTextureType(BakeType, BakeBitDepth);
-		const bool bConvertToSRGB = TexType == FTexture2DBuilder::ETextureType::Color;
-		const ETextureSourceFormat SourceDataFormat = BakeBitDepth == EGeometryScriptBakeBitDepth::ChannelBits16 ? TSF_RGBA16F : TSF_BGRA8;
-
-		constexpr int ResultIdx = 0;
-		FTexture2DBuilder TextureBuilder;
-		TextureBuilder.Initialize(TexType, BakeDimensions);
-		TextureBuilder.Copy(*Baker.GetBakeResults(EvalIdx)[ResultIdx], bConvertToSRGB);
-		TextureBuilder.Commit(false);
-
-		// Copy image to source data after commit. This will avoid incurring
-		// the cost of hitting the DDC for texture compile while iterating on
-		// bake settings. Since this dirties the texture, the next time the texture
-		// is used after accepting the final texture, the DDC will trigger and
-		// properly recompile the platform data.
-		const bool bConvertSourceToSRGB = bConvertToSRGB && SourceDataFormat == TSF_BGRA8;
-		TextureBuilder.CopyImageToSourceData(*Baker.GetBakeResults(EvalIdx)[ResultIdx], SourceDataFormat, bConvertSourceToSRGB);
-		TextureOutput.Add(TextureBuilder.GetTexture2D());
-	}
+	TArray<UTexture2D*> TextureOutput;
+	GeometryScriptBakeLocals::GetTexturesFromBaker(Baker, BakeOptions.BitDepth, TextureOutput);
 
 	return TextureOutput;
 }
diff --git a/Engine/Plugins/Experimental/GeometryScripting/Source/GeometryScriptingCore/Public/GeometryScript/GeometryScriptTypes.h b/Engine/Plugins/Experimental/GeometryScripting/Source/GeometryScriptingCore/Public/GeometryScript/GeometryScriptTypes.h
--- a/Engine/Plugins/Experimental/GeometryScripting/Source/GeometryScriptingCore/Public/GeometryScript/GeometryScriptTypes.h
+++ b/Engine/Plugins/Experimental/GeometryScripting/Source/GeometryScriptingCore/Public/GeometryScript/GeometryScriptTypes.h
@@ -378,5 +378,13 @@ namespace Geometry
 
 	GEOMETRYSCRIPTINGCORE_API void AppendError(UGeometryScriptDebug* Debug, EGeometryScriptErrorType ErrorTypeIn, const FText& MessageIn);
 	GEOMETRYSCRIPTINGCORE_API void AppendWarning(UGeometryScriptDebug* Debug, EGeometryScriptErrorType WarningTypeIn, const FText& MessageIn);
+
+	/**
+	 * These variants of AppendError/Warning are for direct write to a debug message array.
+	 * This may be useful in cases where a function is async and needs to accumulate
+	 * debug messages for later collation on a game thread.
+	 */
+	GEOMETRYSCRIPTINGCORE_API void AppendError(TArray<FGeometryScriptDebugMessage>* DebugMessages, EGeometryScriptErrorType ErrorTypeIn, const FText& MessageIn);
+	GEOMETRYSCRIPTINGCORE_API void AppendWarning(TArray<FGeometryScriptDebugMessage>* DebugMessages, EGeometryScriptErrorType WarningTypeIn, const FText& MessageIn);
 }
 }
\ No newline at end of file
diff --git a/Engine/Plugins/Experimental/GeometryScripting/Source/GeometryScriptingCore/Public/GeometryScript/MeshBakeFunctions.h b/Engine/Plugins/Experimental/GeometryScripting/Source/GeometryScriptingCore/Public/GeometryScript/MeshBakeFunctions.h
--- a/Engine/Plugins/Experimental/GeometryScripting/Source/GeometryScriptingCore/Public/GeometryScript/MeshBakeFunctions.h
+++ b/Engine/Plugins/Experimental/GeometryScripting/Source/GeometryScriptingCore/Public/GeometryScript/MeshBakeFunctions.h
@@ -5,6 +5,7 @@
 #include "CoreMinimal.h"
 #include "Kismet/BlueprintFunctionLibrary.h"
 #include "GeometryScript/GeometryScriptTypes.h"
+#include "Sampling/MeshMapBaker.h"
 #include "MeshBakeFunctions.generated.h"
 
 class UDynamicMesh;
@@ -251,30 +252,41 @@ struct GEOMETRYSCRIPTINGCORE_API FGeometryScriptBakeSourceMeshOptions
 	int SourceNormalUVLayer = 0;
 };
 
-UCLASS(meta = (ScriptName = "GeometryScript_Baking"))
+USTRUCT(BlueprintType)
+struct GEOMETRYSCRIPTINGCORE_API FGeometryScriptBakeTextureAsyncResult
+{
+	GENERATED_BODY();
+
+	FGeometryScriptBakeTextureOptions BakeOptions;
+	TSharedPtr<UE::Geometry::FMeshMapBaker, ESPMode::ThreadSafe> BakeResult;
+};
+
+DECLARE_DYNAMIC_DELEGATE_TwoParams(FBakeTextureDelegate, int, BakeId, FGeometryScriptBakeTextureAsyncResult, ResultOut);
+
+UCLASS(meta = (ScriptName = "GeometryScript_Bake"))
 class GEOMETRYSCRIPTINGCORE_API UGeometryScriptLibrary_MeshBakeFunctions : public UBlueprintFunctionLibrary
 {
 	GENERATED_BODY()
 public:
-	UFUNCTION(BlueprintCallable, Category = "GeometryScript|Baking")
+	UFUNCTION(BlueprintCallable, Category = "GeometryScript|Bake")
 	static UPARAM(DisplayName="Bake Type Out") FGeometryScriptBakeTypeOptions MakeBakeTypeTangentNormal();
 
-	UFUNCTION(BlueprintCallable, Category = "GeometryScript|Baking")
+	UFUNCTION(BlueprintCallable, Category = "GeometryScript|Bake")
 	static UPARAM(DisplayName="Bake Type Out") FGeometryScriptBakeTypeOptions MakeBakeTypeObjectNormal();
 	
-	UFUNCTION(BlueprintCallable, Category = "GeometryScript|Baking")
+	UFUNCTION(BlueprintCallable, Category = "GeometryScript|Bake")
 	static UPARAM(DisplayName="Bake Type Out") FGeometryScriptBakeTypeOptions MakeBakeTypeFaceNormal();
 
-	UFUNCTION(BlueprintCallable, Category = "GeometryScript|Baking")
+	UFUNCTION(BlueprintCallable, Category = "GeometryScript|Bake")
 	static UPARAM(DisplayName="Bake Type Out") FGeometryScriptBakeTypeOptions MakeBakeTypeBentNormal(
 		int OcclusionRays = 16,
 		float MaxDistance = 0.0f,
 		float SpreadAngle = 180.0f);
 
-	UFUNCTION(BlueprintCallable, Category = "GeometryScript|Baking")
+	UFUNCTION(BlueprintCallable, Category = "GeometryScript|Bake")
 	static UPARAM(DisplayName="Bake Type Out") FGeometryScriptBakeTypeOptions MakeBakeTypePosition();
 
-	UFUNCTION(BlueprintCallable, Category = "GeometryScript|Baking")
+	UFUNCTION(BlueprintCallable, Category = "GeometryScript|Bake")
 	static UPARAM(DisplayName="Bake Type Out") FGeometryScriptBakeTypeOptions MakeBakeTypeCurvature(
 		EGeometryScriptBakeCurvatureTypeMode CurvatureType = EGeometryScriptBakeCurvatureTypeMode::Mean,
 		EGeometryScriptBakeCurvatureColorMode ColorMapping = EGeometryScriptBakeCurvatureColorMode::Grayscale,
@@ -282,30 +294,30 @@ public:
 		float MinRangeMultiplier = 0.0f,
 		EGeometryScriptBakeCurvatureClampMode Clamping = EGeometryScriptBakeCurvatureClampMode::None);
 
-	UFUNCTION(BlueprintCallable, Category = "GeometryScript|Baking")
+	UFUNCTION(BlueprintCallable, Category = "GeometryScript|Bake")
 	static UPARAM(DisplayName="Bake Type Out") FGeometryScriptBakeTypeOptions MakeBakeTypeAmbientOcclusion(
 		int OcclusionRays = 16,
 		float MaxDistance = 0.0f,
 		float SpreadAngle = 180.0f,
 		float BiasAngle = 15.0f);
 
-	UFUNCTION(BlueprintCallable, Category = "GeometryScript|Baking")
+	UFUNCTION(BlueprintCallable, Category = "GeometryScript|Bake")
 	static UPARAM(DisplayName="Bake Type Out") FGeometryScriptBakeTypeOptions MakeBakeTypeTexture(
 		UTexture2D* SourceTexture = nullptr,
 		int SourceUVLayer = 0);
 
-	UFUNCTION(BlueprintCallable, Category = "GeometryScript|Baking")
+	UFUNCTION(BlueprintCallable, Category = "GeometryScript|Bake")
 	static UPARAM(DisplayName="Bake Type Out") FGeometryScriptBakeTypeOptions MakeBakeTypeMultiTexture(
 		const TArray<UTexture2D*>& MaterialIDSourceTextures,
 		UPARAM(DisplayName="Source UV Channel") int SourceUVLayer = 0);
 
-	UFUNCTION(BlueprintCallable, Category = "GeometryScript|Baking")
+	UFUNCTION(BlueprintCallable, Category = "GeometryScript|Bake")
 	static UPARAM(DisplayName="Bake Type Out") FGeometryScriptBakeTypeOptions MakeBakeTypeVertexColor();
 
-	UFUNCTION(BlueprintCallable, Category = "GeometryScript|Baking")
+	UFUNCTION(BlueprintCallable, Category = "GeometryScript|Bake")
 	static UPARAM(DisplayName="Bake Type Out") FGeometryScriptBakeTypeOptions MakeBakeTypeMaterialID();
 	
-	UFUNCTION(BlueprintCallable, Category = "GeometryScript|Baking")
+	UFUNCTION(BlueprintCallable, Category = "GeometryScript|Bake")
 	static UPARAM(DisplayName="Textures Out") TArray<UTexture2D*> BakeTexture(
 		UDynamicMesh* TargetMesh,
 		FTransform TargetTransform,
@@ -316,4 +328,52 @@ public:
 		const TArray<FGeometryScriptBakeTypeOptions>& BakeTypes,
 		FGeometryScriptBakeTextureOptions BakeOptions,
 		UGeometryScriptDebug* Debug = nullptr);
+
+	/** 
+	 * BakeTextureAsyncBegin() is the entry point for an asynchronous variant of BakeTexture.
+	 * Usage of this method is as follows:
+	 *
+	 * 1. BakeTextureAsyncBegin() kicks off an async compute.
+	 * 2. When the async compute is complete, it invokes the provided delegate back on the game thread.
+	 * 3. The delegate output is consumed by BakeTextureAsyncEnd() which converts the bake results into UTexture2D.
+	 * 4. An optional BakeId can be provided to BakeTextureAsyncBegin() that will be associated with the async compute.
+	 *    The BakeId can be used to distinguish multiple async computes on the same delegate.
+	 *
+	 * @param Completed the delegate to invoke on completion. Delegate signature: (int, FGeometryScriptBakeTextureAsyncResult)
+	 * @param BakeId optional integer identifier to aid in distinguishing multiple async computes on the same delegate.
+	 * @param TargetMesh the target mesh to bake to.
+	 * @param TargetTransform the transform of the target mesh. Only applicable when projecting in world space.
+	 * @param TargetOptions target mesh options such as which UV layer to bake with.
+	 * @param SourceMesh the source mesh to bake from.
+	 * @param SourceTransform the transform of the source mesh. Only applicable when projecting in world space.
+	 * @param SourceOptions source mesh options such source normal map.
+	 * @param BakeTypes the array of bake types & type specific parameters to bake
+	 * @param BakeOptions bake options (ex. resolution, number of samples, projection distance).
+	 * @param Debug debug structure to capture debug messages.
+	 */
+	UFUNCTION(BlueprintCallable, Category = "GeometryScript|Bake")
+	static void BakeTextureAsyncBegin(
+		const FBakeTextureDelegate& Completed,
+		int BakeId,
+		UDynamicMesh* TargetMesh,
+		FTransform TargetTransform,
+		FGeometryScriptBakeTargetMeshOptions TargetOptions,
+		UDynamicMesh* SourceMesh,
+		FTransform SourceTransform,
+		FGeometryScriptBakeSourceMeshOptions SourceOptions,
+		const TArray<FGeometryScriptBakeTypeOptions>& BakeTypes,
+		FGeometryScriptBakeTextureOptions BakeOptions,
+		UGeometryScriptDebug* Debug = nullptr);
+
+	/**
+	 * Converts the bake results of BakeTextureAsyncBegin() into UTexture2D.
+	 * This function is intended to be invoked by the delegate passed to
+	 * BakeTextureAsyncBegin()
+	 *
+	 * @param Result the result of a BakeTextureAsyncBegin() via delegate
+	 * @return array of UTexture2D corresponding to the requested bake types in BakeTextureAsyncBegin().
+	 */
+	UFUNCTION(BlueprintCallable, Category = "GeometryScript|Bake")
+	static UPARAM(DisplayName="Textures Out") TArray<UTexture2D*> BakeTextureAsyncEnd(
+		const FGeometryScriptBakeTextureAsyncResult& Result);
 };
