Add critical GeometryScript operations needed for Lyra. Add ResetDynamicMeshCollision operation that clears out any existing simple collision. Add filter box to FGeometryScriptMeshBevelOptions which can be used to limit the set of bevelled edges in ApplyMeshPolygroupBevel.
#rnx

diff --git a/Engine/Plugins/Experimental/GeometryScripting/Source/GeometryScriptingCore/Private/CollisionFunctions.cpp b/Engine/Plugins/Experimental/GeometryScripting/Source/GeometryScriptingCore/Private/CollisionFunctions.cpp
--- a/Engine/Plugins/Experimental/GeometryScripting/Source/GeometryScriptingCore/Private/CollisionFunctions.cpp
+++ b/Engine/Plugins/Experimental/GeometryScripting/Source/GeometryScriptingCore/Private/CollisionFunctions.cpp
@@ -324,6 +324,53 @@ UDynamicMesh* UGeometryScriptLibrary_CollisionFunctions::SetDynamicMeshCollision
 }
 
 
+void UGeometryScriptLibrary_CollisionFunctions::ResetDynamicMeshCollision(
+	UDynamicMeshComponent* DynamicMeshComponent,
+	bool bEmitTransaction,
+	UGeometryScriptDebug* Debug)
+{
+	if (DynamicMeshComponent == nullptr)
+	{
+		UE::Geometry::AppendError(Debug, EGeometryScriptErrorType::InvalidInputs, LOCTEXT("ResetDynamicMeshCollision_InvalidInput2", "ResetDynamicMeshCollision: Component is Null"));
+		return;
+	}
+
+#if WITH_EDITOR
+	if (bEmitTransaction)
+	{
+		GEditor->BeginTransaction(LOCTEXT("ResetDynamicMeshCollisionTransaction", "Clear Simple Collision"));
+		DynamicMeshComponent->Modify();
+	}
+#endif
+
+	UBodySetup* BodySetup = DynamicMeshComponent->GetBodySetup();
+	if (BodySetup != nullptr)
+	{
+		// mark the BodySetup for modification. Do we need to modify the UStaticMesh??
+#if WITH_EDITOR
+		if (bEmitTransaction)
+		{
+			BodySetup->Modify();
+		}
+#endif
+
+		// clear existing simple collision. This will call BodySetup->InvalidatePhysicsData()
+		BodySetup->RemoveSimpleCollision();
+
+		// do we need to do a post edit change here??
+		DynamicMeshComponent->UpdateCollision(false);
+	}
+
+#if WITH_EDITOR
+	if (bEmitTransaction)
+	{
+		GEditor->EndTransaction();
+	}
+#endif
+
+}
+
+
 
 
 #undef LOCTEXT_NAMESPACE
\ No newline at end of file
diff --git a/Engine/Plugins/Experimental/GeometryScripting/Source/GeometryScriptingCore/Private/MeshModelingFunctions.cpp b/Engine/Plugins/Experimental/GeometryScripting/Source/GeometryScriptingCore/Private/MeshModelingFunctions.cpp
--- a/Engine/Plugins/Experimental/GeometryScripting/Source/GeometryScriptingCore/Private/MeshModelingFunctions.cpp
+++ b/Engine/Plugins/Experimental/GeometryScripting/Source/GeometryScriptingCore/Private/MeshModelingFunctions.cpp
@@ -443,11 +443,97 @@ UDynamicMesh* UGeometryScriptLibrary_MeshModelingFunctions::ApplyMeshPolygroupBe
 	{
 		UE::Geometry::FGroupTopology Topology(&EditMesh, true);
 
+		TArray<int32> BevelGroupEdges;
+		if (Options.bApplyFilterBox)
+		{
+			FAxisAlignedBox3d QueryBox(Options.FilterBox);
+			FTransformSRT3d InvTransform(Options.FilterBoxTransform.Inverse());
+			TSet<int32> FoundEdges;
+
+			// find all mesh edges inside filter shape
+			// TODO: this is hardcoded for what is supported in 5.0, ie only 3D boxes, but should be generalized
+			FDynamicMeshAABBTree3 Spatial(&EditMesh, true);
+			FDynamicMeshAABBTree3::FTreeTraversal EdgeTraversal;
+			EdgeTraversal.NextBoxF = [&QueryBox](const FAxisAlignedBox3d& Box, int Depth) { return Box.Intersects(QueryBox); };
+			EdgeTraversal.NextTriangleF = [&QueryBox, &InvTransform, &FoundEdges, &EditMesh, &Options](int TriangleID)
+			{
+				FIndex3i Edges = EditMesh.GetTriEdges(TriangleID);
+				for (int32 j = 0; j < 3; ++j)
+				{
+					if (FoundEdges.Contains(Edges[j]) == false && EditMesh.IsGroupBoundaryEdge(Edges[j]))
+					{
+						FVector3d A, B;
+						EditMesh.GetEdgeV(Edges[j], A, B);
+						A = InvTransform.TransformPosition(A);
+						B = InvTransform.TransformPosition(B);
+						bool bIntersects = (Options.bFullyContained) ?
+							(QueryBox.Contains(A) && QueryBox.Contains(B)) :
+							(QueryBox.Contains(A) || QueryBox.Contains(B));
+						if (bIntersects)
+						{
+							FoundEdges.Add(Edges[j]);
+						}
+					}
+				}
+			};
+			Spatial.DoTraversal(EdgeTraversal);
+
+			// convert mesh edges to group topology edges
+			TSet<int32> GroupEdges;
+			for (int32 MeshEdgeID : FoundEdges)
+			{
+				int32 GroupEdgeID = Topology.FindGroupEdgeID(MeshEdgeID);
+				if (GroupEdgeID >= 0)
+				{
+					GroupEdges.Add(GroupEdgeID);
+				}
+			}
+
+			// if exclusive was requested we need to check that the entire edge is inside the box
+			if (Options.bFullyContained)
+			{
+				for (int32 GroupEdgeID : GroupEdges)
+				{
+					const TArray<int>& EdgeIDs = Topology.GetGroupEdgeEdges(GroupEdgeID);
+					bool bFoundMissing = false;
+					for (int32 MeshEdgeID : EdgeIDs)
+					{
+						if (FoundEdges.Contains(MeshEdgeID) == false)
+						{
+							bFoundMissing = true;
+							break;
+						}
+					}
+					if (bFoundMissing == false)
+					{
+						BevelGroupEdges.Add(GroupEdgeID);
+					}
+				}
+			}
+			else
+			{
+				BevelGroupEdges = GroupEdges.Array();
+			}
+
+			if (BevelGroupEdges.Num() == 0)
+			{
+				UE::Geometry::AppendError(Debug, EGeometryScriptErrorType::InvalidInputs, LOCTEXT("ApplyMeshPolygroupBevel_FilterIsEmpty", "ApplyMeshPolygroupBevel: Filter box does not contain any Polygroup Edges, bevel will not be applied"));
+				return;
+			}
+		}
+
 		UE::Geometry::FMeshBevel Bevel;
 		Bevel.InsetDistance = Options.BevelDistance;
 		Bevel.MaterialIDMode = (Options.bInferMaterialID) ? FMeshBevel::EMaterialIDMode::InferMaterialID : FMeshBevel::EMaterialIDMode::ConstantMaterialID;
 		Bevel.SetConstantMaterialID = Options.SetMaterialID;
-		Bevel.InitializeFromGroupTopology(EditMesh, Topology);
+		if (BevelGroupEdges.Num() > 0)
+		{
+			Bevel.InitializeFromGroupTopologyEdges(EditMesh, Topology, BevelGroupEdges);
+		}
+		else
+		{
+			Bevel.InitializeFromGroupTopology(EditMesh, Topology);
+		}
 		Bevel.Apply(EditMesh, nullptr);
 
 	}, EDynamicMeshChangeType::GeneralEdit, EDynamicMeshAttributeChangeFlags::Unknown, false);
diff --git a/Engine/Plugins/Experimental/GeometryScripting/Source/GeometryScriptingCore/Public/GeometryScript/CollisionFunctions.h b/Engine/Plugins/Experimental/GeometryScripting/Source/GeometryScriptingCore/Public/GeometryScript/CollisionFunctions.h
--- a/Engine/Plugins/Experimental/GeometryScripting/Source/GeometryScriptingCore/Public/GeometryScript/CollisionFunctions.h
+++ b/Engine/Plugins/Experimental/GeometryScripting/Source/GeometryScriptingCore/Public/GeometryScript/CollisionFunctions.h
@@ -127,4 +127,13 @@ public:
 		FGeometryScriptCollisionFromMeshOptions Options,
 		UGeometryScriptDebug* Debug = nullptr);
 
+
+	UFUNCTION(BlueprintCallable, Category = "GeometryScript|Collision")
+	static void
+	ResetDynamicMeshCollision(
+		UDynamicMeshComponent* Component,
+		bool bEmitTransaction = false,
+		UGeometryScriptDebug* Debug = nullptr);
+
+
 };
\ No newline at end of file
diff --git a/Engine/Plugins/Experimental/GeometryScripting/Source/GeometryScriptingCore/Public/GeometryScript/MeshModelingFunctions.h b/Engine/Plugins/Experimental/GeometryScripting/Source/GeometryScriptingCore/Public/GeometryScript/MeshModelingFunctions.h
--- a/Engine/Plugins/Experimental/GeometryScripting/Source/GeometryScriptingCore/Public/GeometryScript/MeshModelingFunctions.h
+++ b/Engine/Plugins/Experimental/GeometryScripting/Source/GeometryScriptingCore/Public/GeometryScript/MeshModelingFunctions.h
@@ -70,6 +70,34 @@ public:
 
 	UPROPERTY(BlueprintReadWrite, Category = Options)
 	int SetMaterialID = 0;
+
+
+
+	/**
+	 * If true the set of beveled polygroup edges is limited to those that 
+	 * are fully or partially contained within the (transformed) FilterBox
+	 */
+	UPROPERTY(BlueprintReadWrite, Category = FilterShape, AdvancedDisplay)
+	bool bApplyFilterBox = false;
+
+	/**
+	 * Bounding Box used for edge filtering
+	 */
+	UPROPERTY(BlueprintReadWrite, Category = FilterShape, AdvancedDisplay)
+	FBox FilterBox = FBox(EForceInit::ForceInit);
+
+	/**
+	 * Transform applied to the FilterBox
+	 */
+	UPROPERTY(BlueprintReadWrite, Category = FilterShape, AdvancedDisplay)
+	FTransform FilterBoxTransform = FTransform::Identity;
+
+	/**
+	 * If true, then only polygroup edges that are fully contained within the filter box will be beveled,
+	 * otherwise the edge will be beveled if any vertex is within the filter box.
+	 */
+	UPROPERTY(BlueprintReadWrite, Category = FilterShape, AdvancedDisplay)
+	bool bFullyContained = true;
 };
 
 
