GeometryScript: Set MakeBakeType* geometry script functions to BlueprintPure.
Remove USTRUCT and UPROPERTY from internal structs in bake functions.

diff --git a/Engine/Plugins/Experimental/GeometryScripting/Source/GeometryScriptingCore/Public/GeometryScript/MeshBakeFunctions.h b/Engine/Plugins/Experimental/GeometryScripting/Source/GeometryScriptingCore/Public/GeometryScript/MeshBakeFunctions.h
--- a/Engine/Plugins/Experimental/GeometryScripting/Source/GeometryScriptingCore/Public/GeometryScript/MeshBakeFunctions.h
+++ b/Engine/Plugins/Experimental/GeometryScripting/Source/GeometryScriptingCore/Public/GeometryScript/MeshBakeFunctions.h
@@ -78,31 +78,22 @@ enum class EGeometryScriptBakeNormalSpace : uint8
 	Object
 };
 
-USTRUCT(BlueprintType)
 struct GEOMETRYSCRIPTINGCORE_API FGeometryScriptBakeTypes
 {
-	GENERATED_BODY()
 };
 
-USTRUCT(BlueprintType)
 struct GEOMETRYSCRIPTINGCORE_API FGeometryScriptBakeType_Occlusion : public FGeometryScriptBakeTypes
 {
-	GENERATED_BODY()
-public:
 	/** Number of occlusion rays per sample */
-	UPROPERTY(BlueprintReadWrite, Category = OcclusionOutput, meta = (UIMin = "1", UIMax = "1024", ClampMin = "1", ClampMax = "65536"))
 	int32 OcclusionRays = 16;
 
 	/** Maximum distance for occlusion rays to test for intersections; a value of 0 means infinity */
-	UPROPERTY(BlueprintReadWrite, Category = OcclusionOutput, meta = (UIMin = "0.0", UIMax = "1000.0", ClampMin = "0.0", ClampMax = "99999999.0"))
 	float MaxDistance = 0.0f;
 
 	/** Maximum spread angle in degrees for occlusion rays; for example, 180 degrees will cover the entire hemisphere whereas 90 degrees will only cover the center of the hemisphere down to 45 degrees from the horizon. */
-	UPROPERTY(BlueprintReadWrite, Category = OcclusionOutput, meta = (UIMin = "0", UIMax = "180.0", ClampMin = "0", ClampMax = "180.0"))
 	float SpreadAngle = 180.0f;
 
 	/** Angle in degrees from the horizon for occlusion rays for which the contribution is attenuated to reduce faceting artifacts. */
-	UPROPERTY(BlueprintReadWrite, Category = OcclusionOutput, meta = (UIMin = "0", UIMax = "45.0", ClampMin = "0", ClampMax = "89.9"))
 	float BiasAngle = 15.0f;
 };
 
@@ -143,57 +134,39 @@ enum class EGeometryScriptBakeCurvatureClampMode : uint8
 	OnlyNegative
 };
 
-USTRUCT(BlueprintType)
 struct GEOMETRYSCRIPTINGCORE_API FGeometryScriptBakeType_Curvature : public FGeometryScriptBakeTypes
 {
-	GENERATED_BODY()
-public:
 	/** Type of curvature */
-	UPROPERTY(BlueprintReadWrite, Category = CurvatureOutput)
 	EGeometryScriptBakeCurvatureTypeMode CurvatureType = EGeometryScriptBakeCurvatureTypeMode::Mean;
 
 	/** How to map calculated curvature values to colors */
-	UPROPERTY(BlueprintReadWrite, Category = CurvatureOutput)
 	EGeometryScriptBakeCurvatureColorMode ColorMapping = EGeometryScriptBakeCurvatureColorMode::Grayscale;
 
 	/** Multiplier for how the curvature values fill the available range in the selected Color Mapping; a larger value means that higher curvature is required to achieve the maximum color value. */
-	UPROPERTY(BlueprintReadWrite, Category = CurvatureOutput, meta = (UIMin = "0.1", UIMax = "2.0", ClampMin = "0.001", ClampMax = "100.0"))
 	float ColorRangeMultiplier = 1.0f;
 
 	/** Minimum for the curvature values to not be clamped to zero relative to the curvature for the maximum color value; a larger value means that higher curvature is required to not be considered as no curvature. */
-	UPROPERTY(BlueprintReadWrite, Category = CurvatureOutput, AdvancedDisplay, meta = (DisplayName = "Color Range Minimum", UIMin = "0.0", UIMax = "1.0"))
 	float MinRangeMultiplier = 0.0f;
 
 	/** Clamping applied to curvature values before color mapping */
-	UPROPERTY(BlueprintReadWrite, Category = CurvatureOutput)
 	EGeometryScriptBakeCurvatureClampMode Clamping = EGeometryScriptBakeCurvatureClampMode::None;
 };
 
-USTRUCT(BlueprintType)
 struct GEOMETRYSCRIPTINGCORE_API FGeometryScriptBakeType_Texture : public FGeometryScriptBakeTypes
 {
-	GENERATED_BODY()
-public:
 	/** Source mesh texture that is to be resampled into a new texture */
-	UPROPERTY(BlueprintReadWrite, Category = TextureOutput)
 	TObjectPtr<UTexture2D> SourceTexture = nullptr;
 
 	/** UV channel to use for the source mesh texture */
-	UPROPERTY(BlueprintReadWrite, Category = TextureOutput, meta = (DisplayName = "Source Texture UV Channel", ClampMin="0"))
 	int SourceUVLayer = 0;
 };
 
-USTRUCT(BlueprintType)
 struct GEOMETRYSCRIPTINGCORE_API FGeometryScriptBakeType_MultiTexture : public FGeometryScriptBakeTypes
 {
-	GENERATED_BODY()
-public:
 	/** For each material ID, the source texture that will be resampled in that material's region*/
-	UPROPERTY(BlueprintReadWrite, Category = MultiTextureOutput, meta = (DisplayName = "Source Textures by Material IDs"))
 	TArray<TObjectPtr<UTexture2D>> MaterialIDSourceTextures;
 
 	/** UV channel to use for the source mesh texture */
-	UPROPERTY(BlueprintReadWrite, Category = MultiTextureOutput, meta = (DisplayName = "Source Texture UV Channel", ClampMin="0"))
 	int SourceUVLayer = 0;
 };
 
@@ -280,25 +253,25 @@ class GEOMETRYSCRIPTINGCORE_API UGeometryScriptLibrary_MeshBakeFunctions : publi
 {
 	GENERATED_BODY()
 public:
-	UFUNCTION(BlueprintCallable, Category = "GeometryScript|Bake")
+	UFUNCTION(BlueprintPure, Category = "GeometryScript|Bake")
 	static UPARAM(DisplayName="Bake Type Out") FGeometryScriptBakeTypeOptions MakeBakeTypeTangentNormal();
 
-	UFUNCTION(BlueprintCallable, Category = "GeometryScript|Bake")
+	UFUNCTION(BlueprintPure, Category = "GeometryScript|Bake")
 	static UPARAM(DisplayName="Bake Type Out") FGeometryScriptBakeTypeOptions MakeBakeTypeObjectNormal();
 	
-	UFUNCTION(BlueprintCallable, Category = "GeometryScript|Bake")
+	UFUNCTION(BlueprintPure, Category = "GeometryScript|Bake")
 	static UPARAM(DisplayName="Bake Type Out") FGeometryScriptBakeTypeOptions MakeBakeTypeFaceNormal();
 
-	UFUNCTION(BlueprintCallable, Category = "GeometryScript|Bake")
+	UFUNCTION(BlueprintPure, Category = "GeometryScript|Bake")
 	static UPARAM(DisplayName="Bake Type Out") FGeometryScriptBakeTypeOptions MakeBakeTypeBentNormal(
 		int OcclusionRays = 16,
 		float MaxDistance = 0.0f,
 		float SpreadAngle = 180.0f);
 
-	UFUNCTION(BlueprintCallable, Category = "GeometryScript|Bake")
+	UFUNCTION(BlueprintPure, Category = "GeometryScript|Bake")
 	static UPARAM(DisplayName="Bake Type Out") FGeometryScriptBakeTypeOptions MakeBakeTypePosition();
 
-	UFUNCTION(BlueprintCallable, Category = "GeometryScript|Bake")
+	UFUNCTION(BlueprintPure, Category = "GeometryScript|Bake")
 	static UPARAM(DisplayName="Bake Type Out") FGeometryScriptBakeTypeOptions MakeBakeTypeCurvature(
 		EGeometryScriptBakeCurvatureTypeMode CurvatureType = EGeometryScriptBakeCurvatureTypeMode::Mean,
 		EGeometryScriptBakeCurvatureColorMode ColorMapping = EGeometryScriptBakeCurvatureColorMode::Grayscale,
@@ -306,27 +279,27 @@ public:
 		float MinRangeMultiplier = 0.0f,
 		EGeometryScriptBakeCurvatureClampMode Clamping = EGeometryScriptBakeCurvatureClampMode::None);
 
-	UFUNCTION(BlueprintCallable, Category = "GeometryScript|Bake")
+	UFUNCTION(BlueprintPure, Category = "GeometryScript|Bake")
 	static UPARAM(DisplayName="Bake Type Out") FGeometryScriptBakeTypeOptions MakeBakeTypeAmbientOcclusion(
 		int OcclusionRays = 16,
 		float MaxDistance = 0.0f,
 		float SpreadAngle = 180.0f,
 		float BiasAngle = 15.0f);
 
-	UFUNCTION(BlueprintCallable, Category = "GeometryScript|Bake")
+	UFUNCTION(BlueprintPure, Category = "GeometryScript|Bake")
 	static UPARAM(DisplayName="Bake Type Out") FGeometryScriptBakeTypeOptions MakeBakeTypeTexture(
 		UTexture2D* SourceTexture = nullptr,
 		int SourceUVLayer = 0);
 
-	UFUNCTION(BlueprintCallable, Category = "GeometryScript|Bake")
+	UFUNCTION(BlueprintPure, Category = "GeometryScript|Bake")
 	static UPARAM(DisplayName="Bake Type Out") FGeometryScriptBakeTypeOptions MakeBakeTypeMultiTexture(
 		const TArray<UTexture2D*>& MaterialIDSourceTextures,
 		UPARAM(DisplayName="Source UV Channel") int SourceUVLayer = 0);
 
-	UFUNCTION(BlueprintCallable, Category = "GeometryScript|Bake")
+	UFUNCTION(BlueprintPure, Category = "GeometryScript|Bake")
 	static UPARAM(DisplayName="Bake Type Out") FGeometryScriptBakeTypeOptions MakeBakeTypeVertexColor();
 
-	UFUNCTION(BlueprintCallable, Category = "GeometryScript|Bake")
+	UFUNCTION(BlueprintPure, Category = "GeometryScript|Bake")
 	static UPARAM(DisplayName="Bake Type Out") FGeometryScriptBakeTypeOptions MakeBakeTypeMaterialID();
 	
 	UFUNCTION(BlueprintCallable, Category = "GeometryScript|Bake")
