GeometryScript: Add support for converting between SkeletalMesh & DynamicMesh.

diff --git a/Engine/Plugins/Experimental/GeometryScripting/Source/GeometryScriptingCore/GeometryScriptingCore.Build.cs b/Engine/Plugins/Experimental/GeometryScripting/Source/GeometryScriptingCore/GeometryScriptingCore.Build.cs
--- a/Engine/Plugins/Experimental/GeometryScripting/Source/GeometryScriptingCore/GeometryScriptingCore.Build.cs
+++ b/Engine/Plugins/Experimental/GeometryScripting/Source/GeometryScriptingCore/GeometryScriptingCore.Build.cs
@@ -42,6 +42,7 @@ public class GeometryScriptingCore : ModuleRules
 				"Engine",
 				"MeshDescription",
 				"StaticMeshDescription",
+				"SkeletalMeshDescription",
 				"MeshConversion",
 				"GeometryAlgorithms",
 				"ModelingOperators",
diff --git a/Engine/Plugins/Experimental/GeometryScripting/Source/GeometryScriptingCore/Private/MeshAssetFunctions.cpp b/Engine/Plugins/Experimental/GeometryScripting/Source/GeometryScriptingCore/Private/MeshAssetFunctions.cpp
--- a/Engine/Plugins/Experimental/GeometryScripting/Source/GeometryScriptingCore/Private/MeshAssetFunctions.cpp
+++ b/Engine/Plugins/Experimental/GeometryScripting/Source/GeometryScriptingCore/Private/MeshAssetFunctions.cpp
@@ -12,6 +12,8 @@
 #include "StaticMeshOperations.h"
 #include "RenderingThread.h"
 #include "PhysicsEngine/BodySetup.h"
+#include "Rendering/SkeletalMeshModel.h"
+#include "SkeletalMeshAttributes.h"
 
 #include "MeshDescriptionToDynamicMesh.h"
 #include "DynamicMeshToMeshDescription.h"
@@ -296,8 +298,159 @@ void UGeometryScriptLibrary_StaticMeshFunctions::GetSectionMaterialListFromStati
 }
 
 
+UDynamicMesh* UGeometryScriptLibrary_StaticMeshFunctions::CopyMeshFromSkeletalMesh(
+		USkeletalMesh* FromSkeletalMeshAsset, 
+		UDynamicMesh* ToDynamicMesh,
+		FGeometryScriptCopyMeshFromAssetOptions AssetOptions,
+		FGeometryScriptMeshReadLOD RequestedLOD,
+		TEnumAsByte<EGeometryScriptOutcomePins>& Outcome,
+		UGeometryScriptDebug* Debug)
+{
+	Outcome = EGeometryScriptOutcomePins::Failure;
 
+	if (FromSkeletalMeshAsset == nullptr)
+	{
+		UE::Geometry::AppendError(Debug, EGeometryScriptErrorType::InvalidInputs, LOCTEXT("CopyMeshFromSkeletalMesh_InvalidInput1", "CopyMeshFromSkeletalMesh: FromSkeletalMeshAsset is Null"));
+		return ToDynamicMesh;
+	}
+	if (ToDynamicMesh == nullptr)
+	{
+		UE::Geometry::AppendError(Debug, EGeometryScriptErrorType::InvalidInputs, LOCTEXT("CopyMeshFromSkeletalMesh_InvalidInput2", "CopyMeshFromSkeletalMesh: ToDynamicMesh is Null"));
+		return ToDynamicMesh;
+	}
+	if (RequestedLOD.LODType != EGeometryScriptLODType::MaxAvailable && RequestedLOD.LODType != EGeometryScriptLODType::SourceModel)
+	{
+		UE::Geometry::AppendError(Debug, EGeometryScriptErrorType::InvalidInputs, LOCTEXT("CopyMeshFromSkeletalMesh_LODNotAvailable", "CopyMeshFromSkeletalMesh: Requested LOD is not available"));
+		return ToDynamicMesh;
+	}
+
+	// TODO: Consolidate this code with SkeletalMeshToolTarget::GetMeshDescription(..) 
+#if WITH_EDITOR
+	const int32 UseLODIndex = FMath::Clamp(RequestedLOD.LODIndex, 0, FromSkeletalMeshAsset->GetLODNum() - 1);;
+	
+	FMeshDescription SourceMesh;
+
+	// Check first if we have bulk data available and non-empty.
+	if (FromSkeletalMeshAsset->IsLODImportedDataBuildAvailable(UseLODIndex) && !FromSkeletalMeshAsset->IsLODImportedDataEmpty(UseLODIndex))
+	{
+		FSkeletalMeshImportData SkeletalMeshImportData;
+		FromSkeletalMeshAsset->LoadLODImportedData(UseLODIndex, SkeletalMeshImportData);
+		SkeletalMeshImportData.GetMeshDescription(SourceMesh);
+	}
+	else
+	{
+		// Fall back on the LOD model directly if no bulk data exists. When we commit
+		// the mesh description, we override using the bulk data. This can happen for older
+		// skeletal meshes, from UE 4.24 and earlier.
+		const FSkeletalMeshModel* SkeletalMeshModel = FromSkeletalMeshAsset->GetImportedModel();
+		if (SkeletalMeshModel && SkeletalMeshModel->LODModels.IsValidIndex(UseLODIndex))
+		{
+			SkeletalMeshModel->LODModels[UseLODIndex].GetMeshDescription(SourceMesh, FromSkeletalMeshAsset);
+		}			
+	}
 
+	FDynamicMesh3 NewMesh;
+	FMeshDescriptionToDynamicMesh Converter;
+	Converter.Convert(&SourceMesh, NewMesh, AssetOptions.bRequestTangents);
+	
+	ToDynamicMesh->SetMesh(MoveTemp(NewMesh));
+	
+	Outcome = EGeometryScriptOutcomePins::Success;
+#else
+	UE::Geometry::AppendError(Debug, EGeometryScriptErrorType::InvalidInputs, LOCTEXT("CopyMeshFromSkeletalMesh_EditorOnly", "CopyMeshFromSkeletalMesh: Not currently supported at Runtime"));
+#endif
+	
+	return ToDynamicMesh;
+}
+
+
+UDynamicMesh* UGeometryScriptLibrary_StaticMeshFunctions::CopyMeshToSkeletalMesh(
+		UDynamicMesh* FromDynamicMesh, 
+		USkeletalMesh* ToSkeletalMeshAsset,
+		FGeometryScriptCopyMeshToAssetOptions Options,
+		FGeometryScriptMeshWriteLOD TargetLOD,
+		TEnumAsByte<EGeometryScriptOutcomePins>& Outcome,
+		UGeometryScriptDebug* Debug)
+{
+	Outcome = EGeometryScriptOutcomePins::Failure;
+
+	if (FromDynamicMesh == nullptr)
+	{
+		UE::Geometry::AppendError(Debug, EGeometryScriptErrorType::InvalidInputs, LOCTEXT("CopyMeshToSkeletalMesh_InvalidInput1", "CopyMeshToSkeletalMesh: FromDynamicMesh is Null"));
+		return FromDynamicMesh;
+	}
+	if (ToSkeletalMeshAsset == nullptr)
+	{
+		UE::Geometry::AppendError(Debug, EGeometryScriptErrorType::InvalidInputs, LOCTEXT("CopyMeshToSkeletalMesh_InvalidInput2", "CopyMeshToSkeletalMesh: ToSkeletalMeshAsset is Null"));
+		return FromDynamicMesh;
+	}
+	if (TargetLOD.bWriteHiResSource)
+	{
+		UE::Geometry::AppendError(Debug, EGeometryScriptErrorType::InvalidInputs, LOCTEXT("CopyMeshToSkeletalMesh_Unsupported", "CopyMeshToSkeletalMesh: Writing HiResSource LOD is not yet supported"));
+		return FromDynamicMesh;
+	}
+
+	// TODO: Consolidate this code with SkeletalMeshToolTarget::CommitMeshDescription
+#if WITH_EDITOR
+	if (ToSkeletalMeshAsset->GetPathName().StartsWith(TEXT("/Engine/")))
+	{
+		const FText Error = FText::Format(LOCTEXT("CopyMeshToSkeletalMesh_BuiltInAsset", "CopyMeshToSkeletalMesh: Cannot modify built-in engine asset: {0}"), FText::FromString(*ToSkeletalMeshAsset->GetPathName()));
+		UE::Geometry::AppendError(Debug, EGeometryScriptErrorType::InvalidInputs, Error);
+		return FromDynamicMesh;
+	}
+
+	// flush any pending rendering commands, which might touch a component while we are rebuilding it's mesh
+	FlushRenderingCommands();
+
+	if (Options.bEmitTransaction)
+	{
+		GEditor->BeginTransaction(LOCTEXT("UpdateSkeletalMesh", "Update Skeletal Mesh"));
+	}
+
+	// make sure transactional flag is on for this asset
+	ToSkeletalMeshAsset->SetFlags(RF_Transactional);
+
+	verify(ToSkeletalMeshAsset->Modify());
+
+	FMeshDescription MeshDescription;
+	FSkeletalMeshAttributes MeshAttributes(MeshDescription);
+	MeshAttributes.Register();
+	
+	FConversionToMeshDescriptionOptions ConversionOptions;
+	FDynamicMeshToMeshDescription Converter(ConversionOptions);
+	FromDynamicMesh->ProcessMesh([&](const FDynamicMesh3& ReadMesh)
+	{
+		Converter.Convert(&ReadMesh, MeshDescription, !Options.bEnableRecomputeTangents);
+	});
+
+	FSkeletalMeshImportData SkeletalMeshImportData = 
+		FSkeletalMeshImportData::CreateFromMeshDescription(MeshDescription);
+	ToSkeletalMeshAsset->SaveLODImportedData(TargetLOD.LODIndex, SkeletalMeshImportData);
+
+	// Make sure the mesh builder knows it's the latest variety, so that the render data gets
+	// properly rebuilt.
+	ToSkeletalMeshAsset->SetLODImportedDataVersions(TargetLOD.LODIndex, ESkeletalMeshGeoImportVersions::LatestVersion, ESkeletalMeshSkinningImportVersions::LatestVersion);
+	ToSkeletalMeshAsset->SetUseLegacyMeshDerivedDataKey(false);
+
+	// TODO: Options.bReplaceMaterials support?
+
+	if (Options.bDeferMeshPostEditChange == false)
+	{
+		ToSkeletalMeshAsset->PostEditChange();
+	}
+
+	if (Options.bEmitTransaction)
+	{
+		GEditor->EndTransaction();
+	}
+	
+	Outcome = EGeometryScriptOutcomePins::Success;
+#else
+	UE::Geometry::AppendError(Debug, EGeometryScriptErrorType::InvalidInputs, LOCTEXT("CopyMeshToSkeletalMesh_EditorOnly", "CopyMeshToSkeletalMesh: Not currently supported at Runtime"));
+#endif
+	
+	return FromDynamicMesh;
+}
 
 
 
diff --git a/Engine/Plugins/Experimental/GeometryScripting/Source/GeometryScriptingCore/Public/GeometryScript/MeshAssetFunctions.h b/Engine/Plugins/Experimental/GeometryScripting/Source/GeometryScriptingCore/Public/GeometryScript/MeshAssetFunctions.h
--- a/Engine/Plugins/Experimental/GeometryScripting/Source/GeometryScriptingCore/Public/GeometryScript/MeshAssetFunctions.h
+++ b/Engine/Plugins/Experimental/GeometryScripting/Source/GeometryScriptingCore/Public/GeometryScript/MeshAssetFunctions.h
@@ -64,7 +64,9 @@ public:
 
 
 
-
+// Although the class name indicates StaticMeshFunctions, that was a naming mistake that is difficult
+// to correct. This class is intended to serve as a generic asset utils function library. The naming
+// issue is only visible at the C++ level. It is not visible in Python or BP.
 UCLASS(meta = (ScriptName = "GeometryScript_AssetUtils"))
 class GEOMETRYSCRIPTINGCORE_API UGeometryScriptLibrary_StaticMeshFunctions : public UBlueprintFunctionLibrary
 {
@@ -103,6 +105,28 @@ public:
 		TArray<int32>& MaterialIndex,
 		TEnumAsByte<EGeometryScriptOutcomePins>& Outcome,
 		UGeometryScriptDebug* Debug = nullptr);
+
+
+	UFUNCTION(BlueprintCallable, Category = "GeometryScript|SkeletalMesh", meta = (ExpandEnumAsExecs = "Outcome"))
+	static UPARAM(DisplayName = "Dynamic Mesh") UDynamicMesh* 
+	CopyMeshFromSkeletalMesh(
+		USkeletalMesh* FromSkeletalMeshAsset, 
+		UDynamicMesh* ToDynamicMesh,
+		FGeometryScriptCopyMeshFromAssetOptions AssetOptions,
+		FGeometryScriptMeshReadLOD RequestedLOD,
+		TEnumAsByte<EGeometryScriptOutcomePins>& Outcome,
+		UGeometryScriptDebug* Debug = nullptr);
+
+
+	UFUNCTION(BlueprintCallable, Category = "GeometryScript|SkeletalMesh", meta = (ExpandEnumAsExecs = "Outcome"))
+	static UPARAM(DisplayName = "Dynamic Mesh") UDynamicMesh* 
+	CopyMeshToSkeletalMesh(
+		UDynamicMesh* FromDynamicMesh, 
+		USkeletalMesh* ToSkeletalMeshAsset,
+		FGeometryScriptCopyMeshToAssetOptions Options,
+		FGeometryScriptMeshWriteLOD TargetLOD,
+		TEnumAsByte<EGeometryScriptOutcomePins>& Outcome,
+		UGeometryScriptDebug* Debug = nullptr);
 };
 
 
