GeometryScript: AppendRectangle and AppendRoundRectangle released in Preview 1 have a bug where X/Y dimension parameters are incorrectly interpreted. Resolved by replacing with new AppendRectangleXY/AppendRoundRectangleXY versions, and adding redirectors for existing functions to new versions with "Compatibility_5_0" suffix and function text/warning to not use in future.

diff --git a/Engine/Plugins/Experimental/GeometryScripting/Source/GeometryScriptingCore/Private/GeometryScriptingCoreModule.cpp b/Engine/Plugins/Experimental/GeometryScripting/Source/GeometryScriptingCore/Private/GeometryScriptingCoreModule.cpp
--- a/Engine/Plugins/Experimental/GeometryScripting/Source/GeometryScriptingCore/Private/GeometryScriptingCoreModule.cpp
+++ b/Engine/Plugins/Experimental/GeometryScripting/Source/GeometryScriptingCore/Private/GeometryScriptingCoreModule.cpp
@@ -1,12 +1,20 @@
 // Copyright Epic Games, Inc. All Rights Reserved.
 
 #include "GeometryScriptingCoreModule.h"
+#include "UObject/CoreRedirects.h"
 
 #define LOCTEXT_NAMESPACE "FGeometryScriptingCoreModule"
 
 void FGeometryScriptingCoreModule::StartupModule()
 {
 	// This code will execute after your module is loaded into memory; the exact timing is specified in the .uplugin file per-module
+
+	TArray<FCoreRedirect> Redirects;
+	// UE5.0 - AppendRectangle and AppendRoundRectangle shipped in Preview 1 interpreted dimensions incorrectly
+	Redirects.Emplace(ECoreRedirectFlags::Type_Function, TEXT("AppendRectangle"), TEXT("AppendRectangle_Compatibility_5_0"));
+	Redirects.Emplace(ECoreRedirectFlags::Type_Function, TEXT("AppendRoundRectangle"), TEXT("AppendRoundRectangle_Compatibility_5_0"));
+	
+	FCoreRedirects::AddRedirectList(Redirects, TEXT("GeometryScriptingCore"));
 }
 
 void FGeometryScriptingCoreModule::ShutdownModule()
diff --git a/Engine/Plugins/Experimental/GeometryScripting/Source/GeometryScriptingCore/Private/MeshPrimitiveFunctions.cpp b/Engine/Plugins/Experimental/GeometryScripting/Source/GeometryScriptingCore/Private/MeshPrimitiveFunctions.cpp
--- a/Engine/Plugins/Experimental/GeometryScripting/Source/GeometryScriptingCore/Private/MeshPrimitiveFunctions.cpp
+++ b/Engine/Plugins/Experimental/GeometryScripting/Source/GeometryScriptingCore/Private/MeshPrimitiveFunctions.cpp
@@ -646,8 +646,7 @@ UDynamicMesh* UGeometryScriptLibrary_MeshPrimitiveFunctions::AppendSimpleSweptPo
 
 
 
-
-UDynamicMesh* UGeometryScriptLibrary_MeshPrimitiveFunctions::AppendRectangle(
+UDynamicMesh* UGeometryScriptLibrary_MeshPrimitiveFunctions::AppendRectangleXY(
 	UDynamicMesh* TargetMesh,
 	FGeometryScriptPrimitiveOptions PrimitiveOptions,
 	FTransform Transform,
@@ -659,15 +658,15 @@ UDynamicMesh* UGeometryScriptLibrary_MeshPrimitiveFunctions::AppendRectangle(
 {
 	if (TargetMesh == nullptr)
 	{
-		UE::Geometry::AppendError(Debug, EGeometryScriptErrorType::InvalidInputs, LOCTEXT("PrimitiveFunctions_AppendRectangle", "AppendRectangle: TargetMesh is Null"));
+		UE::Geometry::AppendError(Debug, EGeometryScriptErrorType::InvalidInputs, LOCTEXT("PrimitiveFunctions_AppendRectangleXY", "AppendRectangle: TargetMesh is Null"));
 		return TargetMesh;
 	}
 
 	FRectangleMeshGenerator RectGenerator;
 	RectGenerator.Origin = FVector3d(0, 0, 0);
 	RectGenerator.Normal = FVector3f::UnitZ();
-	RectGenerator.Width = DimensionX / 2;
-	RectGenerator.Height = DimensionY / 2;
+	RectGenerator.Width = DimensionX;
+	RectGenerator.Height = DimensionY;
 	RectGenerator.WidthVertexCount = FMath::Max(0, StepsWidth);
 	RectGenerator.HeightVertexCount = FMath::Max(0, StepsHeight);
 	RectGenerator.bSinglePolyGroup = (PrimitiveOptions.PolygroupMode != EGeometryScriptPrimitivePolygroupMode::PerQuad);
@@ -681,7 +680,21 @@ UDynamicMesh* UGeometryScriptLibrary_MeshPrimitiveFunctions::AppendRectangle(
 
 
 
-UDynamicMesh* UGeometryScriptLibrary_MeshPrimitiveFunctions::AppendRoundRectangle(
+UDynamicMesh* UGeometryScriptLibrary_MeshPrimitiveFunctions::AppendRectangle_Compatibility_5_0(
+	UDynamicMesh* TargetMesh,
+	FGeometryScriptPrimitiveOptions PrimitiveOptions,
+	FTransform Transform,
+	float DimensionX,
+	float DimensionY,
+	int32 StepsWidth,
+	int32 StepsHeight,
+	UGeometryScriptDebug* Debug)
+{
+	return AppendRectangleXY(TargetMesh, PrimitiveOptions, Transform, DimensionX * 0.5, DimensionY * 0.5, StepsWidth, StepsHeight, Debug);
+}
+
+
+UDynamicMesh* UGeometryScriptLibrary_MeshPrimitiveFunctions::AppendRoundRectangleXY(
 	UDynamicMesh* TargetMesh,
 	FGeometryScriptPrimitiveOptions PrimitiveOptions,
 	FTransform Transform,
@@ -695,15 +708,15 @@ UDynamicMesh* UGeometryScriptLibrary_MeshPrimitiveFunctions::AppendRoundRectangl
 {
 	if (TargetMesh == nullptr)
 	{
-		UE::Geometry::AppendError(Debug, EGeometryScriptErrorType::InvalidInputs, LOCTEXT("PrimitiveFunctions_AppendRoundRectangle", "AppendRoundRectangle: TargetMesh is Null"));
+		UE::Geometry::AppendError(Debug, EGeometryScriptErrorType::InvalidInputs, LOCTEXT("PrimitiveFunctions_AppendRoundRectangleXY", "AppendRoundRectangle: TargetMesh is Null"));
 		return TargetMesh;
 	}
 
 	FRoundedRectangleMeshGenerator RectGenerator;
 	RectGenerator.Origin = FVector3d(0, 0, 0);
 	RectGenerator.Normal = FVector3f::UnitZ();
-	RectGenerator.Width = DimensionX / 2;
-	RectGenerator.Height = DimensionY / 2;
+	RectGenerator.Width = DimensionX;
+	RectGenerator.Height = DimensionY;
 	RectGenerator.WidthVertexCount = FMath::Max(0, StepsWidth);
 	RectGenerator.HeightVertexCount = FMath::Max(0, StepsHeight);
 	RectGenerator.Radius = FMath::Max(FMathf::ZeroTolerance, CornerRadius);
@@ -718,6 +731,23 @@ UDynamicMesh* UGeometryScriptLibrary_MeshPrimitiveFunctions::AppendRoundRectangl
 
 
 
+UDynamicMesh* UGeometryScriptLibrary_MeshPrimitiveFunctions::AppendRoundRectangle_Compatibility_5_0(
+	UDynamicMesh* TargetMesh,
+	FGeometryScriptPrimitiveOptions PrimitiveOptions,
+	FTransform Transform,
+	float DimensionX,
+	float DimensionY,
+	float CornerRadius,
+	int32 StepsWidth,
+	int32 StepsHeight,
+	int32 StepsRound,
+	UGeometryScriptDebug* Debug)
+{
+	return AppendRoundRectangleXY(TargetMesh, PrimitiveOptions, Transform, DimensionX * 0.5, DimensionY * 0.5, CornerRadius, StepsWidth, StepsHeight, StepsRound, Debug);
+}
+
+
+
 UDynamicMesh* UGeometryScriptLibrary_MeshPrimitiveFunctions::AppendDisc(
 	UDynamicMesh* TargetMesh,
 	FGeometryScriptPrimitiveOptions PrimitiveOptions,
diff --git a/Engine/Plugins/Experimental/GeometryScripting/Source/GeometryScriptingCore/Public/GeometryScript/GeometryScriptTypes.h b/Engine/Plugins/Experimental/GeometryScripting/Source/GeometryScriptingCore/Public/GeometryScript/GeometryScriptTypes.h
--- a/Engine/Plugins/Experimental/GeometryScripting/Source/GeometryScriptingCore/Public/GeometryScript/GeometryScriptTypes.h
+++ b/Engine/Plugins/Experimental/GeometryScripting/Source/GeometryScriptingCore/Public/GeometryScript/GeometryScriptTypes.h
@@ -160,14 +160,20 @@ public:
 // Polygroups
 //
 
+/**
+ * FGeometryScriptGroupLayer identifies a Polygroup Layer of a Mesh.
+ * The Default Layer always exists, Extended layers may or may not exist.
+ */
 USTRUCT(BlueprintType)
 struct GEOMETRYSCRIPTINGCORE_API FGeometryScriptGroupLayer
 {
 	GENERATED_BODY()
 public:
+	/** If true,the default/standard PolyGroup Layer is used */
 	UPROPERTY(BlueprintReadWrite, Category = LOD)
 	bool bDefaultLayer = true;
 
+	/** Index of an extended PolyGroup Layer (which may or may not exist on any given Mesh) */
 	UPROPERTY(BlueprintReadWrite, Category = LOD)
 	int ExtendedLayerIndex = 0;
 };
diff --git a/Engine/Plugins/Experimental/GeometryScripting/Source/GeometryScriptingCore/Public/GeometryScript/MeshPrimitiveFunctions.h b/Engine/Plugins/Experimental/GeometryScripting/Source/GeometryScriptingCore/Public/GeometryScript/MeshPrimitiveFunctions.h
--- a/Engine/Plugins/Experimental/GeometryScripting/Source/GeometryScriptingCore/Public/GeometryScript/MeshPrimitiveFunctions.h
+++ b/Engine/Plugins/Experimental/GeometryScripting/Source/GeometryScriptingCore/Public/GeometryScript/MeshPrimitiveFunctions.h
@@ -264,7 +264,7 @@ public:
 
 	UFUNCTION(BlueprintCallable, Category = "GeometryScript|Primitives", meta=(ScriptMethod))
 	static UPARAM(DisplayName = "Target Mesh") UDynamicMesh* 
-	AppendRectangle( 
+	AppendRectangleXY( 
 		UDynamicMesh* TargetMesh, 
 		FGeometryScriptPrimitiveOptions PrimitiveOptions,
 		FTransform Transform,
@@ -276,7 +276,7 @@ public:
 
 	UFUNCTION(BlueprintCallable, Category = "GeometryScript|Primitives", meta=(ScriptMethod))
 	static UPARAM(DisplayName = "Target Mesh") UDynamicMesh* 
-	AppendRoundRectangle( 
+	AppendRoundRectangleXY( 
 		UDynamicMesh* TargetMesh, 
 		FGeometryScriptPrimitiveOptions PrimitiveOptions,
 		FTransform Transform,
@@ -346,4 +346,56 @@ public:
 		bool bFloating = false,
 		UGeometryScriptDebug* Debug = nullptr);
 
+
+
+
+	//---------------------------------------------
+	// Backwards-Compatibility implementations
+	//---------------------------------------------
+	// These are versions/variants of the above functions that were released
+	// in previous UE 5.x versions, that have since been updated. 
+	// To avoid breaking user scripts, these previous versions are currently kept and 
+	// called via redirectors registered in GeometryScriptingCoreModule.cpp.
+	// 
+	// These functions may be deprecated in future UE releases.
+	//
+
+
+	/**
+	 * 5.0 Preview 1 Compatibility version of AppendRectangleXY. Incorrectly interprets the input dimensions.
+	 * Incorrectly divides the input DimensionX and DimensionY by 2.
+	 * @warning It is strongly recommended that callers of this function update to the current AppendRectangleXY function!
+	 */
+	UFUNCTION(BlueprintCallable, Category = "GeometryScript|Compatibility", meta=(ScriptMethod))
+	static UPARAM(DisplayName = "Target Mesh") UDynamicMesh* 
+	AppendRectangle_Compatibility_5_0( 
+		UDynamicMesh* TargetMesh, 
+		FGeometryScriptPrimitiveOptions PrimitiveOptions,
+		FTransform Transform,
+		float DimensionX = 100,
+		float DimensionY = 100,
+		int32 StepsWidth = 0,
+		int32 StepsHeight = 0,
+		UGeometryScriptDebug* Debug = nullptr);
+
+
+	/**
+	 * 5.0 Preview 1 Compatibility version of AppendRoundRectangleXY. 
+	 * Incorrectly divides the input DimensionX and DimensionY by 2.
+	 * @warning It is strongly recommended that callers of this function update to the current AppendRoundRectangleXY function!
+	 */
+	UFUNCTION(BlueprintCallable, Category = "GeometryScript|Compatibility", meta=(ScriptMethod))
+	static UPARAM(DisplayName = "Target Mesh") UDynamicMesh* 
+	AppendRoundRectangle_Compatibility_5_0( 
+		UDynamicMesh* TargetMesh, 
+		FGeometryScriptPrimitiveOptions PrimitiveOptions,
+		FTransform Transform,
+		float DimensionX = 100,
+		float DimensionY = 100,
+		float CornerRadius = 5,
+		int32 StepsWidth = 0,
+		int32 StepsHeight = 0,
+		int32 StepsRound = 4,
+		UGeometryScriptDebug* Debug = nullptr);
+
 };
\ No newline at end of file
