GeometryScript: add Nanite setup options to FGeometryScriptCopyMeshToAssetOptions / CopyMeshToStaticMesh
#rnx

diff --git a/Engine/Plugins/Experimental/GeometryScripting/Source/GeometryScriptingCore/Private/MeshAssetFunctions.cpp b/Engine/Plugins/Experimental/GeometryScripting/Source/GeometryScriptingCore/Private/MeshAssetFunctions.cpp
--- a/Engine/Plugins/Experimental/GeometryScripting/Source/GeometryScriptingCore/Private/MeshAssetFunctions.cpp
+++ b/Engine/Plugins/Experimental/GeometryScripting/Source/GeometryScriptingCore/Private/MeshAssetFunctions.cpp
@@ -212,6 +212,11 @@ UDynamicMesh*  UGeometryScriptLibrary_StaticMeshFunctions::CopyMeshToStaticMesh(
 	ThisSourceModel.ReductionSettings.PercentTriangles = 1.f;
 	ThisSourceModel.ReductionSettings.PercentVertices = 1.f;
 
+	if (Options.bApplyNaniteSettings)
+	{
+		ToStaticMeshAsset->NaniteSettings = Options.NaniteSettings;
+	}
+
 	if (Options.bReplaceMaterials)
 	{
 		bool bHaveSlotNames = (Options.NewMaterialSlotNames.Num() == Options.NewMaterials.Num());
diff --git a/Engine/Plugins/Experimental/GeometryScripting/Source/GeometryScriptingCore/Public/GeometryScript/MeshAssetFunctions.h b/Engine/Plugins/Experimental/GeometryScripting/Source/GeometryScriptingCore/Public/GeometryScript/MeshAssetFunctions.h
--- a/Engine/Plugins/Experimental/GeometryScripting/Source/GeometryScriptingCore/Public/GeometryScript/MeshAssetFunctions.h
+++ b/Engine/Plugins/Experimental/GeometryScripting/Source/GeometryScriptingCore/Public/GeometryScript/MeshAssetFunctions.h
@@ -5,6 +5,7 @@
 #include "CoreMinimal.h"
 #include "Kismet/BlueprintFunctionLibrary.h"
 #include "GeometryScript/GeometryScriptTypes.h"
+#include "Engine/EngineTypes.h"
 #include "MeshAssetFunctions.generated.h"
 
 class UStaticMesh;
@@ -54,6 +55,13 @@ public:
 	UPROPERTY(BlueprintReadWrite, Category = Options)
 	TArray<FName> NewMaterialSlotNames;
 
+	/** If enabled, NaniteSettings will be applied to the target Asset if possible */
+	UPROPERTY(BlueprintReadWrite, Category = Options)
+	bool bApplyNaniteSettings = false;
+
+	/** Nanite Settings applied to the target Asset, if bApplyNaniteSettings = true */
+	UPROPERTY(BlueprintReadWrite, Category = Options)
+	FMeshNaniteSettings NaniteSettings = FMeshNaniteSettings();
 
 	UPROPERTY(BlueprintReadWrite, Category = Options)
 	bool bEmitTransaction = true;
