GeometryScript: resolve issue w/ exposing FMeshNaniteSettings as a UFunction input, introduced in 19212237 - it is not a valid type in BP due to uint32 uproperty. Replace with FGeometryScriptNaniteOptions that mirrors necessary fields. Update affected content.
#rnx

diff --git a/Engine/Plugins/Experimental/GeometryScripting/Source/GeometryScriptingCore/Private/MeshAssetFunctions.cpp b/Engine/Plugins/Experimental/GeometryScripting/Source/GeometryScriptingCore/Private/MeshAssetFunctions.cpp
--- a/Engine/Plugins/Experimental/GeometryScripting/Source/GeometryScriptingCore/Private/MeshAssetFunctions.cpp
+++ b/Engine/Plugins/Experimental/GeometryScripting/Source/GeometryScriptingCore/Private/MeshAssetFunctions.cpp
@@ -5,6 +5,7 @@
 #include "DynamicMesh/DynamicMesh3.h"
 #include "UDynamicMesh.h"
 
+#include "Engine/EngineTypes.h"
 #include "Engine/StaticMesh.h"
 #include "Engine/SkeletalMesh.h"
 #include "Engine/StaticMeshActor.h"
@@ -214,7 +215,11 @@ UDynamicMesh*  UGeometryScriptLibrary_StaticMeshFunctions::CopyMeshToStaticMesh(
 
 	if (Options.bApplyNaniteSettings)
 	{
-		ToStaticMeshAsset->NaniteSettings = Options.NaniteSettings;
+		FMeshNaniteSettings NewNaniteSettings = ToStaticMeshAsset->NaniteSettings;
+		NewNaniteSettings.bEnabled = Options.NaniteSettings.bEnabled;
+		NewNaniteSettings.FallbackPercentTriangles = FMath::Clamp(Options.NaniteSettings.FallbackPercentTriangles / 100.0f, 0.0f, 1.0f);
+		NewNaniteSettings.FallbackRelativeError = FMath::Max(Options.NaniteSettings.FallbackRelativeError, 0.0f);
+		ToStaticMeshAsset->NaniteSettings = NewNaniteSettings;
 	}
 
 	if (Options.bReplaceMaterials)
diff --git a/Engine/Plugins/Experimental/GeometryScripting/Source/GeometryScriptingCore/Public/GeometryScript/MeshAssetFunctions.h b/Engine/Plugins/Experimental/GeometryScripting/Source/GeometryScriptingCore/Public/GeometryScript/MeshAssetFunctions.h
--- a/Engine/Plugins/Experimental/GeometryScripting/Source/GeometryScriptingCore/Public/GeometryScript/MeshAssetFunctions.h
+++ b/Engine/Plugins/Experimental/GeometryScripting/Source/GeometryScriptingCore/Public/GeometryScript/MeshAssetFunctions.h
@@ -5,7 +5,6 @@
 #include "CoreMinimal.h"
 #include "Kismet/BlueprintFunctionLibrary.h"
 #include "GeometryScript/GeometryScriptTypes.h"
-#include "Engine/EngineTypes.h"
 #include "MeshAssetFunctions.generated.h"
 
 class UStaticMesh;
@@ -29,6 +28,26 @@ public:
 	bool bIgnoreRemoveDegenerates = true;
 };
 
+/**
+ * Configuration settings for Nanite Rendering on StaticMesh Assets
+ */
+USTRUCT(BlueprintType)
+struct GEOMETRYSCRIPTINGCORE_API FGeometryScriptNaniteOptions
+{
+	GENERATED_BODY()
+
+	/** Set Nanite to Enabled/Disabled */
+	UPROPERTY(BlueprintReadWrite, Category = Options)
+	bool bEnabled = true;
+
+	/** Percentage of triangles to maintain in Fallback Mesh used when Nanite is unavailable */
+	UPROPERTY(BlueprintReadWrite, Category = Options)
+	float FallbackPercentTriangles = 100.0f;
+
+	/** Relative Error to maintain in Fallback Mesh used when Nanite is unavailable. Overrides FallbackPercentTriangles. Set to 0 to only use FallbackPercentTriangles (default). */
+	UPROPERTY(BlueprintReadWrite, Category = Options)
+	float FallbackRelativeError = 0.0f;
+};
 
 
 USTRUCT(BlueprintType)
@@ -61,7 +80,7 @@ public:
 
 	/** Nanite Settings applied to the target Asset, if bApplyNaniteSettings = true */
 	UPROPERTY(BlueprintReadWrite, Category = Options)
-	FMeshNaniteSettings NaniteSettings = FMeshNaniteSettings();
+	FGeometryScriptNaniteOptions NaniteSettings = FGeometryScriptNaniteOptions();
 
 	UPROPERTY(BlueprintReadWrite, Category = Options)
 	bool bEmitTransaction = true;
