GeometryScript: Added CreateNewTexture2DAsset function

diff --git a/Engine/Plugins/Experimental/GeometryScripting/Source/GeometryScriptingEditor/Private/CreateNewAssetUtilityFunctions.cpp b/Engine/Plugins/Experimental/GeometryScripting/Source/GeometryScriptingEditor/Private/CreateNewAssetUtilityFunctions.cpp
--- a/Engine/Plugins/Experimental/GeometryScripting/Source/GeometryScriptingEditor/Private/CreateNewAssetUtilityFunctions.cpp
+++ b/Engine/Plugins/Experimental/GeometryScripting/Source/GeometryScriptingEditor/Private/CreateNewAssetUtilityFunctions.cpp
@@ -26,6 +26,7 @@
 #include "DynamicMeshToMeshDescription.h"
 #include "ConversionUtils/DynamicMeshToVolume.h"
 #include "AssetUtils/CreateStaticMeshUtil.h"
+#include "AssetUtils/CreateTexture2DUtil.h"
 #include "ModelingObjectsCreationAPI.h"
 
 using namespace UE::Geometry;
@@ -210,6 +211,64 @@ UStaticMesh* UGeometryScriptLibrary_CreateNewAssetFunctions::CreateNewStaticMesh
 }
 
 
+UTexture2D* UGeometryScriptLibrary_CreateNewAssetFunctions::CreateNewTexture2DAsset(
+		UTexture2D* FromTexture, 
+		FString AssetPathAndName,
+		FGeometryScriptCreateNewTexture2DAssetOptions Options,
+		TEnumAsByte<EGeometryScriptOutcomePins>& Outcome,
+		UGeometryScriptDebug* Debug)
+{
+	Outcome = EGeometryScriptOutcomePins::Failure;
+	if (FromTexture == nullptr)
+	{
+		UE::Geometry::AppendError(Debug, EGeometryScriptErrorType::InvalidInputs, LOCTEXT("CreateNewTexture2DAsset_InvalidInput1", "CreateNewTexture2DAsset: FromTexture is Null"));
+		return nullptr;
+	}
+
+	UE::AssetUtils::FTexture2DAssetOptions AssetOptions;
+	AssetOptions.NewAssetPath = AssetPathAndName;
+	AssetOptions.bOverwriteIfExists = Options.bOverwriteIfExists;
+
+	UE::AssetUtils::FTexture2DAssetResults ResultData;
+	UE::AssetUtils::ECreateTexture2DResult AssetResult = UE::AssetUtils::SaveGeneratedTexture2DAsset(
+		FromTexture, AssetOptions, ResultData);
+
+	switch (AssetResult)
+	{
+	case UE::AssetUtils::ECreateTexture2DResult::NameError:
+	{
+		const FText Error = FText::Format(LOCTEXT("CreateNewTexture2DAsset_InvalidInputName", "CreateNewTexture2DAsset: AssetPathAndName '{0}' already exists."), FText::FromString(AssetPathAndName));
+		UE::Geometry::AppendError(Debug, EGeometryScriptErrorType::InvalidInputs, Error);
+		break;
+	}
+	case UE::AssetUtils::ECreateTexture2DResult::OverwriteTypeError:
+	{
+		const FText Error = FText::Format(LOCTEXT("CreateNewTexture2DAsset_InvalidOverwriteType", "CreateNewTexture2DAsset: AssetPathAndName '{0}' already exists and is not a UTexture2D."), FText::FromString(AssetPathAndName));
+		UE::Geometry::AppendError(Debug, EGeometryScriptErrorType::InvalidInputs, Error);
+		break;
+	}
+	case UE::AssetUtils::ECreateTexture2DResult::InvalidInputTexture:
+	{
+		const FText Error = FText::Format(LOCTEXT("CreateNewTexture2DAsset_InvalidInputPackage", "CreateNewTexture2DAsset: Failed to read input texture '{0}'."), FText::FromString(FromTexture->GetPathName()));
+		UE::Geometry::AppendError(Debug, EGeometryScriptErrorType::InvalidInputs, Error);
+		break;
+	}
+	case UE::AssetUtils::ECreateTexture2DResult::InvalidPackage:
+	case UE::AssetUtils::ECreateTexture2DResult::UnknownError:
+	{
+		UE::Geometry::AppendError(Debug, EGeometryScriptErrorType::OperationFailed, LOCTEXT("CreateNewTexture2DAsset_Failed", "CreateNewTexture2DAsset: Failed to create new Asset."));
+		break;
+	}
+	case UE::AssetUtils::ECreateTexture2DResult::Ok:
+	{
+		Outcome = EGeometryScriptOutcomePins::Success;
+		break;
+	}
+	}
+	
+	return ResultData.Texture;
+}
+
 
 
 
diff --git a/Engine/Plugins/Experimental/GeometryScripting/Source/GeometryScriptingEditor/Public/GeometryScript/CreateNewAssetUtilityFunctions.h b/Engine/Plugins/Experimental/GeometryScripting/Source/GeometryScriptingEditor/Public/GeometryScript/CreateNewAssetUtilityFunctions.h
--- a/Engine/Plugins/Experimental/GeometryScripting/Source/GeometryScriptingEditor/Public/GeometryScript/CreateNewAssetUtilityFunctions.h
+++ b/Engine/Plugins/Experimental/GeometryScripting/Source/GeometryScriptingEditor/Public/GeometryScript/CreateNewAssetUtilityFunctions.h
@@ -65,6 +65,15 @@ public:
 };
 
 
+USTRUCT(BlueprintType)
+struct GEOMETRYSCRIPTINGEDITOR_API FGeometryScriptCreateNewTexture2DAssetOptions
+{
+	GENERATED_BODY()
+public:
+	/** If true, overwrite any existing texture assets using the same AssetPathAndName */
+	UPROPERTY(BlueprintReadWrite, Category = Options)
+	bool bOverwriteIfExists = false;
+};
 
 
 UCLASS(meta = (ScriptName = "GeometryScript_NewAssetUtils"))
@@ -105,5 +114,14 @@ public:
 		TEnumAsByte<EGeometryScriptOutcomePins>& Outcome,
 		UGeometryScriptDebug* Debug = nullptr);
 
+	UFUNCTION(BlueprintCallable, Category = "GeometryScript|AssetManagement", meta = (ExpandEnumAsExecs = "Outcome"))
+	static UPARAM(DisplayName = "Texture 2D Asset") UTexture2D* 
+	CreateNewTexture2DAsset(
+		UTexture2D* FromTexture, 
+		FString AssetPathAndName,
+		FGeometryScriptCreateNewTexture2DAssetOptions Options,
+		TEnumAsByte<EGeometryScriptOutcomePins>& Outcome,
+		UGeometryScriptDebug* Debug = nullptr);
+
 };
 
diff --git a/Engine/Plugins/Runtime/MeshModelingToolset/Source/ModelingComponents/Private/AssetUtils/Texture2DBuilder.cpp b/Engine/Plugins/Runtime/MeshModelingToolset/Source/ModelingComponents/Private/AssetUtils/Texture2DBuilder.cpp
--- a/Engine/Plugins/Runtime/MeshModelingToolset/Source/ModelingComponents/Private/AssetUtils/Texture2DBuilder.cpp
+++ b/Engine/Plugins/Runtime/MeshModelingToolset/Source/ModelingComponents/Private/AssetUtils/Texture2DBuilder.cpp
@@ -57,6 +57,76 @@ bool FTexture2DBuilder::InitializeAndReplaceExistingTexture(UTexture2D* Existing
 }
 
 
+bool FTexture2DBuilder::InitializeAndReplaceExistingTexture(UTexture2D* ExistingTexture, UTexture2D* SourceTexture)
+{
+	if (!ExistingTexture || !SourceTexture)
+	{
+		return false;
+	}
+
+	// create new texture
+	const EPixelFormat InFormat = SourceTexture->GetPixelFormat();
+
+	int32 InSizeX = SourceTexture->GetSizeX();
+	int32 InSizeY = SourceTexture->GetSizeY();
+
+	if (ensureMsgf(InSizeX > 0 && InSizeY > 0 &&
+		(InSizeX % GPixelFormats[InFormat].BlockSizeX) == 0 &&
+		(InSizeY % GPixelFormats[InFormat].BlockSizeY) == 0, TEXT("Invalid size and/or pixel format for new texture")))
+	{
+		UTexture2D* NewTexture = NewObject<UTexture2D>(
+			ExistingTexture->GetOuter(),
+			ExistingTexture->GetFName(),
+			ExistingTexture->GetFlags()
+			);
+
+		NewTexture->SetPlatformData(new FTexturePlatformData());
+		NewTexture->GetPlatformData()->SizeX = InSizeX;
+		NewTexture->GetPlatformData()->SizeY = InSizeY;
+		NewTexture->GetPlatformData()->PixelFormat = InFormat;
+
+		// Allocate first mipmap.
+		const int32 NumBlocksX = InSizeX / GPixelFormats[InFormat].BlockSizeX;
+		const int32 NumBlocksY = InSizeY / GPixelFormats[InFormat].BlockSizeY;
+		FTexture2DMipMap* Mip = new FTexture2DMipMap();
+		NewTexture->GetPlatformData()->Mips.Add(Mip);
+		Mip->SizeX = InSizeX;
+		Mip->SizeY = InSizeY;
+		Mip->BulkData.Lock(LOCK_READ_WRITE);
+		Mip->BulkData.Realloc(NumBlocksX * NumBlocksY * GPixelFormats[InFormat].BlockBytes);
+		Mip->BulkData.Unlock();
+
+		// Initialize texture settings from SourceTexture.
+		NewTexture->SRGB = SourceTexture->SRGB;
+		NewTexture->CompressionSettings = SourceTexture->CompressionSettings;
+		NewTexture->LODGroup = SourceTexture->LODGroup;
+#if WITH_EDITOR
+		NewTexture->MipGenSettings = SourceTexture->MipGenSettings;
+#endif
+		NewTexture->UpdateResource();
+
+		BuildType = ETextureType::Color; // This will only be used to determine clear color.
+		Dimensions = FImageDimensions(InSizeX, InSizeY);
+		RawTexture2D = NewTexture;
+		CurrentPixelFormat = InFormat;
+
+		// lock
+		if (LockForEditing() == false)
+		{
+			return false;
+		}
+
+		if (IsEditable())
+		{
+			Clear();
+		}
+		return true;
+	}
+	
+	return false;
+}
+
+
 bool FTexture2DBuilder::InitializeInternal(ETextureType BuildTypeIn, FImageDimensions DimensionsIn, UTexture2D* CreatedTextureIn)
 {
 	check(DimensionsIn.IsSquare());
diff --git a/Engine/Plugins/Runtime/MeshModelingToolset/Source/ModelingComponents/Public/AssetUtils/Texture2DBuilder.h b/Engine/Plugins/Runtime/MeshModelingToolset/Source/ModelingComponents/Public/AssetUtils/Texture2DBuilder.h
--- a/Engine/Plugins/Runtime/MeshModelingToolset/Source/ModelingComponents/Public/AssetUtils/Texture2DBuilder.h
+++ b/Engine/Plugins/Runtime/MeshModelingToolset/Source/ModelingComponents/Public/AssetUtils/Texture2DBuilder.h
@@ -112,6 +112,14 @@ public:
 	bool InitializeAndReplaceExistingTexture(UTexture2D* ExistingTexture, ETextureType BuildTypeIn, FImageDimensions DimensionsIn);
 
 
+	/**
+	 * Create a new UTexture2D 'over' an existing UTexture2D using the same texture settings
+	 * as a given source UTexture2D. This method only initializes the texture settings of ExistingTexture
+	 * to match those of SourceTexture. To populate texture data, use Copy(..).
+	 */
+	bool InitializeAndReplaceExistingTexture(UTexture2D* ExistingTexture, UTexture2D* SourceTexture);
+
+
 	/**
 	 * Lock the Mip 0 buffer for editing
 	 * @return true if lock was successfull. Will return false if already locked.
diff --git a/Engine/Plugins/Runtime/MeshModelingToolset/Source/ModelingComponentsEditorOnly/Private/AssetUtils/CreateTexture2DUtil.cpp b/Engine/Plugins/Runtime/MeshModelingToolset/Source/ModelingComponentsEditorOnly/Private/AssetUtils/CreateTexture2DUtil.cpp
--- a/Engine/Plugins/Runtime/MeshModelingToolset/Source/ModelingComponentsEditorOnly/Private/AssetUtils/CreateTexture2DUtil.cpp
+++ b/Engine/Plugins/Runtime/MeshModelingToolset/Source/ModelingComponentsEditorOnly/Private/AssetUtils/CreateTexture2DUtil.cpp
@@ -1,6 +1,8 @@
 // Copyright Epic Games, Inc. All Rights Reserved. 
 
 #include "AssetUtils/CreateTexture2DUtil.h"
+#include "AssetUtils/Texture2DBuilder.h"
+#include "AssetUtils/Texture2DUtil.h"
 
 UE::AssetUtils::ECreateTexture2DResult UE::AssetUtils::SaveGeneratedTexture2DAsset(
 	UTexture2D* GeneratedTexture,
@@ -37,16 +39,48 @@ UE::AssetUtils::ECreateTexture2DResult UE::AssetUtils::SaveGeneratedTexture2DAss
 		return ECreateTexture2DResult::InvalidPackage;
 	}
 
-	// move texture from Transient package to real package
-	bool bRenameOK = GeneratedTexture->Rename(*NewObjectName, UsePackage, REN_None);
-	if (ensure(bRenameOK) == false)
+	// check for existing asset of the same type.
+	if (UObject* ExistingAsset = StaticFindObject(nullptr, UsePackage, *NewObjectName, false))
 	{
-		return ECreateTexture2DResult::NameError;
+		UTexture2D* ExistingTexture2D = Cast<UTexture2D>(ExistingAsset);
+		if (ExistingTexture2D && Options.bOverwriteIfExists)
+		{
+			// TODO: Direct copy between UTexture2D
+			UE::Geometry::TImageBuilder<FVector4f> TexImage;
+			if (!UE::AssetUtils::ReadTexture(GeneratedTexture, TexImage))
+			{
+				return ECreateTexture2DResult::InvalidInputTexture;
+			}
+			
+			UE::Geometry::FTexture2DBuilder Builder;
+			Builder.InitializeAndReplaceExistingTexture(ExistingTexture2D, GeneratedTexture);
+			Builder.Copy(TexImage);
+			Builder.Commit(true);
+			GeneratedTexture = Builder.GetTexture2D();
+		}
+		else if (!ExistingTexture2D && Options.bOverwriteIfExists)
+		{
+			return ECreateTexture2DResult::OverwriteTypeError;
+		}
+		else
+		{
+			return ECreateTexture2DResult::NameError;
+		}
+	}
+	else
+	{
+		// move texture from Transient package to real package
+		const bool bRenameOK = GeneratedTexture->Rename(*NewObjectName, UsePackage, REN_None);
+		if (ensure(bRenameOK) == false)
+		{
+			return ECreateTexture2DResult::NameError;
+		}
+
+		// remove transient flag, add public/standalone/transactional
+		GeneratedTexture->ClearFlags(RF_Transient);
+		GeneratedTexture->SetFlags(RF_Public | RF_Standalone | RF_Transactional);
 	}
 
-	// remove transient flag, add public/standalone/transactional
-	GeneratedTexture->ClearFlags(RF_Transient);
-	GeneratedTexture->SetFlags(RF_Public | RF_Standalone | RF_Transactional);
 	[[maybe_unused]] bool MarkedDirty = GeneratedTexture->MarkPackageDirty();
 
 	// do we need to Modify() it? we are not doing any undo/redo
diff --git a/Engine/Plugins/Runtime/MeshModelingToolset/Source/ModelingComponentsEditorOnly/Public/AssetUtils/CreateTexture2DUtil.h b/Engine/Plugins/Runtime/MeshModelingToolset/Source/ModelingComponentsEditorOnly/Public/AssetUtils/CreateTexture2DUtil.h
--- a/Engine/Plugins/Runtime/MeshModelingToolset/Source/ModelingComponentsEditorOnly/Public/AssetUtils/CreateTexture2DUtil.h
+++ b/Engine/Plugins/Runtime/MeshModelingToolset/Source/ModelingComponentsEditorOnly/Public/AssetUtils/CreateTexture2DUtil.h
@@ -11,7 +11,7 @@ namespace UE
 namespace AssetUtils
 {
 	/**
-	 * Result enum returned by Create() functions below to indicate succes/error conditions
+	 * Result enum returned by Create() functions below to indicate success/error conditions
 	 */
 	enum class ECreateTexture2DResult
 	{
@@ -19,6 +19,7 @@ namespace AssetUtils
 		InvalidPackage = 1,
 		InvalidInputTexture = 2,
 		NameError = 3,
+		OverwriteTypeError = 4,
 
 		UnknownError = 100
 	};
@@ -34,6 +35,9 @@ namespace AssetUtils
 		FString NewAssetPath;
 
 		bool bDeferPostEditChange = false;
+
+		// if NewAssetPath already exists, update the existing Texture2D.
+		bool bOverwriteIfExists = false;
 	};
 
 	/**
