GeometryScript: fix crash in simple collision setup nodes
#rnx

diff --git a/Engine/Plugins/Experimental/GeometryScripting/Source/GeometryScriptingCore/Private/CollisionFunctions.cpp b/Engine/Plugins/Experimental/GeometryScripting/Source/GeometryScriptingCore/Private/CollisionFunctions.cpp
--- a/Engine/Plugins/Experimental/GeometryScripting/Source/GeometryScriptingCore/Private/CollisionFunctions.cpp
+++ b/Engine/Plugins/Experimental/GeometryScripting/Source/GeometryScriptingCore/Private/CollisionFunctions.cpp
@@ -57,6 +57,7 @@ void ComputeCollisionFromMesh(
 	else
 	{
 		Submeshes.SetNum(NumComponents);
+		SubmeshPointers.SetNum(NumComponents);
 		ParallelFor(NumComponents, [&](int32 k)
 		{
 			FDynamicSubmesh3 Submesh(&Mesh, Components[k].Indices, (int32)EMeshComponents::None, false);
