GeometryScript: Add BakeVertex script support.
Hide async Bake script functions from Blueprint until they can be reworked.

diff --git a/Engine/Plugins/Experimental/GeometryScripting/Source/GeometryScriptingCore/Private/MeshBakeFunctions.cpp b/Engine/Plugins/Experimental/GeometryScripting/Source/GeometryScriptingCore/Private/MeshBakeFunctions.cpp
--- a/Engine/Plugins/Experimental/GeometryScripting/Source/GeometryScriptingCore/Private/MeshBakeFunctions.cpp
+++ b/Engine/Plugins/Experimental/GeometryScripting/Source/GeometryScriptingCore/Private/MeshBakeFunctions.cpp
@@ -212,6 +212,231 @@ namespace GeometryScriptBakeLocals
 		return Result;
 	}
 
+	bool GetMeshTangents(const FDynamicMesh3* Mesh, TSharedPtr<FMeshTangentsd, ESPMode::ThreadSafe>& Tangents)
+	{
+		if (!Tangents)
+		{
+			Tangents = MakeShared<FMeshTangentsd, ESPMode::ThreadSafe>(Mesh);
+			Tangents->CopyTriVertexTangents(*Mesh);
+
+			// Validate the tangents
+			if (!FDynamicMeshTangents(Mesh).HasValidTangents(true))
+			{
+				return false;
+			}
+		}
+		return true;
+	};
+
+	struct FEvaluatorState
+	{
+		const FDynamicMesh3* TargetMesh = nullptr;
+		const FDynamicMesh3* SourceMesh = nullptr;
+		IMeshBakerDetailSampler* DetailSampler = nullptr;
+		TSharedPtr<FMeshTangentsd, ESPMode::ThreadSafe> TargetMeshTangents;
+		TSharedPtr<FMeshTangentsd, ESPMode::ThreadSafe> SourceMeshTangents;
+		TSharedPtr<UE::Geometry::TImageBuilder<FVector4f>, ESPMode::ThreadSafe> SourceTexture;
+		TSharedPtr<UE::Geometry::TImageBuilder<FVector4f>, ESPMode::ThreadSafe> SourceNormalMap;
+		bool bSupportsSourceNormalMap = false;
+	};
+
+	bool GetSourceNormalMap(
+		FEvaluatorState& EvalState,
+		const FGeometryScriptBakeSourceMeshOptions& SourceOptions,
+		const FText& DebugPrefix,
+		TArray<FGeometryScriptDebugMessage>* Debug)
+	{
+		if (!EvalState.bSupportsSourceNormalMap || !SourceOptions.SourceNormalMap)
+		{
+			return false;
+		}
+		
+		EvalState.SourceNormalMap = MakeShared<TImageBuilder<FVector4f>>();
+		if (!UE::AssetUtils::ReadTexture(SourceOptions.SourceNormalMap, *EvalState.SourceNormalMap, false))
+		{
+			UE::Geometry::AppendError(Debug, EGeometryScriptErrorType::InvalidInputs, FText::Format(LOCTEXT("Bake_InvalidSourceNormalMap", "{0}: Failed to read SourceNormalMap"), DebugPrefix));
+			return false;
+		}
+		else
+		{
+			EvalState.DetailSampler->SetNormalTextureMap(EvalState.SourceMesh,
+			IMeshBakerDetailSampler::FBakeDetailNormalTexture(
+					EvalState.SourceNormalMap.Get(),
+					SourceOptions.SourceNormalUVLayer,
+					SourceOptions.SourceNormalSpace == EGeometryScriptBakeNormalSpace::Tangent ? IMeshBakerDetailSampler::EBakeDetailNormalSpace::Tangent : IMeshBakerDetailSampler::EBakeDetailNormalSpace::Object));
+
+			if (!GetMeshTangents(EvalState.SourceMesh, EvalState.SourceMeshTangents))
+			{
+				UE::Geometry::AppendError(Debug, EGeometryScriptErrorType::InvalidInputs, FText::Format(LOCTEXT("Bake_InvalidSourceTangents", "{0}: Source Mesh tangents are invalid."), DebugPrefix));
+				return false;
+			}
+			EvalState.DetailSampler->SetTangents(EvalState.SourceMesh, EvalState.SourceMeshTangents.Get());
+		}
+		return true;
+	}
+
+	TSharedPtr<FMeshMapEvaluator> CreateEvaluator(
+		FEvaluatorState& EvalState,
+		const TFunctionRef<bool(EGeometryScriptBakeTypes, TArray<FGeometryScriptDebugMessage>*)> IsValidType,
+		const FGeometryScriptBakeTypeOptions Options,
+		const FText& DebugPrefix,
+		TArray<FGeometryScriptDebugMessage>* Debug)
+	{
+		// Channel evaluators only support a subset of bake types
+		if (!IsValidType(Options.BakeType, Debug))
+		{
+			return nullptr;
+		}
+
+		auto GetTargetMeshTangents = [Debug, &DebugPrefix](FEvaluatorState& State)
+		{
+			const bool bSuccess = GetMeshTangents(State.TargetMesh, State.TargetMeshTangents); 
+			if (!bSuccess)
+			{
+				UE::Geometry::AppendError(Debug, EGeometryScriptErrorType::InvalidInputs, FText::Format(LOCTEXT("Bake_InvalidTargetTangents", "{0}: Target Mesh tangents are invalid."), DebugPrefix));
+			}
+			return bSuccess;
+		};
+		
+		TSharedPtr<FMeshMapEvaluator> Result;
+		switch(Options.BakeType)
+		{
+		case EGeometryScriptBakeTypes::TangentSpaceNormal:
+		{
+			TSharedPtr<FMeshNormalMapEvaluator> NormalEval = MakeShared<FMeshNormalMapEvaluator>();
+			if (!GetTargetMeshTangents(EvalState))
+			{
+				return nullptr;
+			}
+			EvalState.bSupportsSourceNormalMap = true;
+			Result = NormalEval;
+			break;
+		}
+		case EGeometryScriptBakeTypes::ObjectSpaceNormal:
+		{
+			TSharedPtr<FMeshPropertyMapEvaluator> PropertyEval = MakeShared<FMeshPropertyMapEvaluator>();
+			PropertyEval->Property = EMeshPropertyMapType::Normal;
+			EvalState.bSupportsSourceNormalMap = true;
+			Result = PropertyEval;
+			break;
+		}
+		case EGeometryScriptBakeTypes::FaceNormal:
+		{
+			TSharedPtr<FMeshPropertyMapEvaluator> PropertyEval = MakeShared<FMeshPropertyMapEvaluator>();
+			PropertyEval->Property = EMeshPropertyMapType::FacetNormal;
+			Result = PropertyEval;
+			break;
+		}
+		case EGeometryScriptBakeTypes::BentNormal:
+		{
+			FGeometryScriptBakeType_Occlusion* OcclusionOptions = static_cast<FGeometryScriptBakeType_Occlusion*>(Options.Options.Get());
+			TSharedPtr<FMeshOcclusionMapEvaluator> OcclusionEval = MakeShared<FMeshOcclusionMapEvaluator>();
+			OcclusionEval->OcclusionType = EMeshOcclusionMapType::BentNormal;
+			OcclusionEval->NumOcclusionRays = OcclusionOptions->OcclusionRays;
+			OcclusionEval->MaxDistance = (OcclusionOptions->MaxDistance == 0) ? TNumericLimits<float>::Max() : OcclusionOptions->MaxDistance;
+			OcclusionEval->SpreadAngle = OcclusionOptions->SpreadAngle;
+			if (!GetTargetMeshTangents(EvalState))
+			{
+				return nullptr;
+			}
+			Result = OcclusionEval;
+			break;
+		}
+		case EGeometryScriptBakeTypes::Position:
+		{
+			TSharedPtr<FMeshPropertyMapEvaluator> PropertyEval = MakeShared<FMeshPropertyMapEvaluator>();
+			PropertyEval->Property = EMeshPropertyMapType::Position;
+			Result = PropertyEval;
+			break;
+		}
+		case EGeometryScriptBakeTypes::Curvature:
+		{
+			FGeometryScriptBakeType_Curvature* CurvatureOptions = static_cast<FGeometryScriptBakeType_Curvature*>(Options.Options.Get());
+			TSharedPtr<FMeshCurvatureMapEvaluator> CurvatureEval = MakeShared<FMeshCurvatureMapEvaluator>();
+			CurvatureEval->UseCurvatureType = GetCurvatureType(CurvatureOptions->CurvatureType);
+			CurvatureEval->UseColorMode = GetCurvatureColorMode(CurvatureOptions->ColorMapping);
+			CurvatureEval->RangeScale = CurvatureOptions->ColorRangeMultiplier;
+			CurvatureEval->MinRangeScale = CurvatureOptions->MinRangeMultiplier;
+			CurvatureEval->UseClampMode = GetCurvatureClampMode(CurvatureOptions->Clamping);
+			Result = CurvatureEval;
+			break;
+		}
+		case EGeometryScriptBakeTypes::AmbientOcclusion:
+		{
+			FGeometryScriptBakeType_Occlusion* OcclusionOptions = static_cast<FGeometryScriptBakeType_Occlusion*>(Options.Options.Get());
+			TSharedPtr<FMeshOcclusionMapEvaluator> OcclusionEval = MakeShared<FMeshOcclusionMapEvaluator>();
+			OcclusionEval->OcclusionType = EMeshOcclusionMapType::AmbientOcclusion;
+			OcclusionEval->NumOcclusionRays = OcclusionOptions->OcclusionRays;
+			OcclusionEval->MaxDistance = (OcclusionOptions->MaxDistance == 0) ? TNumericLimits<float>::Max() : OcclusionOptions->MaxDistance;
+			OcclusionEval->SpreadAngle = OcclusionOptions->SpreadAngle;
+			OcclusionEval->BiasAngleDeg = OcclusionOptions->BiasAngle;
+			Result = OcclusionEval;
+			break;
+		}
+		case EGeometryScriptBakeTypes::Texture:
+		{
+			FGeometryScriptBakeType_Texture* TextureOptions = static_cast<FGeometryScriptBakeType_Texture*>(Options.Options.Get());
+			TSharedPtr<FMeshResampleImageEvaluator> TextureEval = MakeShared<FMeshResampleImageEvaluator>();
+
+			// TODO: Add support for sampling different texture maps per Texture evaluator in a single pass. 
+			if (!EvalState.SourceTexture && TextureOptions->SourceTexture)
+			{
+				EvalState.SourceTexture = MakeShared<TImageBuilder<FVector4f>>();
+				if (!UE::AssetUtils::ReadTexture(TextureOptions->SourceTexture, *EvalState.SourceTexture, false))
+				{
+					UE::Geometry::AppendError(Debug, EGeometryScriptErrorType::InvalidInputs, LOCTEXT("BakeTexture_InvalidSourceTexture", "BakeTexture: Failed to read SourceTexture"));
+				}
+				else
+				{
+					EvalState.DetailSampler->SetTextureMap(EvalState.SourceMesh, IMeshBakerDetailSampler::FBakeDetailTexture(EvalState.SourceTexture.Get(), TextureOptions->SourceUVLayer));
+				}
+			}
+			Result = TextureEval;
+			break;
+		}
+		case EGeometryScriptBakeTypes::MultiTexture:
+		{
+			FGeometryScriptBakeType_MultiTexture* TextureOptions = static_cast<FGeometryScriptBakeType_MultiTexture*>(Options.Options.Get());
+			TSharedPtr<FMeshMultiResampleImageEvaluator> TextureEval = MakeShared<FMeshMultiResampleImageEvaluator>();
+
+			if (TextureOptions->MaterialIDSourceTextures.Num())
+			{
+				TextureEval->MultiTextures.SetNum(TextureOptions->MaterialIDSourceTextures.Num());
+				for (int32 MaterialId = 0; MaterialId < TextureEval->MultiTextures.Num(); ++MaterialId)
+				{
+					if (UTexture2D* Texture = TextureOptions->MaterialIDSourceTextures[MaterialId])
+					{
+						TextureEval->MultiTextures[MaterialId] = MakeShared<TImageBuilder<FVector4f>>();
+						if (!UE::AssetUtils::ReadTexture(Texture, *TextureEval->MultiTextures[MaterialId], false))
+						{
+							UE::Geometry::AppendError(Debug, EGeometryScriptErrorType::InvalidInputs, FText::Format(LOCTEXT("Bake_InvalidMultiTexture", "{0}: Failed to read MaterialIDSourceTexture"), DebugPrefix));
+						}
+					}
+				}
+			}
+			Result = TextureEval;
+			break;
+		}
+		case EGeometryScriptBakeTypes::VertexColor:
+		{
+			TSharedPtr<FMeshPropertyMapEvaluator> PropertyEval = MakeShared<FMeshPropertyMapEvaluator>();
+			PropertyEval->Property = EMeshPropertyMapType::VertexColor;
+			Result = PropertyEval;
+			break;
+		}
+		case EGeometryScriptBakeTypes::MaterialID:
+		{
+			TSharedPtr<FMeshPropertyMapEvaluator> PropertyEval = MakeShared<FMeshPropertyMapEvaluator>();
+			PropertyEval->Property = EMeshPropertyMapType::MaterialID;
+			Result = PropertyEval;
+			break;
+		}
+		default:
+			break;
+		}
+		return Result;
+	}
+
 	TUniquePtr<FMeshMapBaker> BakeTextureImpl(
 		UDynamicMesh* TargetMesh,
 		FTransform TargetTransform,
@@ -239,32 +464,7 @@ namespace GeometryScriptBakeLocals
 			return nullptr;
 		}
 
-		TSharedPtr<FMeshTangentsd, ESPMode::ThreadSafe> TargetMeshTangents;
-		TSharedPtr<FMeshTangentsd, ESPMode::ThreadSafe> SourceMeshTangents;
-		auto GetMeshTangents = [TargetMesh, SourceMesh, Debug](UDynamicMesh* Mesh, TSharedPtr<FMeshTangentsd, ESPMode::ThreadSafe>& Tangents)
-		{
-			if (!Tangents)
-			{
-				Tangents = MakeShared<FMeshTangentsd, ESPMode::ThreadSafe>(Mesh->GetMeshPtr());
-				Tangents->CopyTriVertexTangents(Mesh->GetMeshRef());
-
-				// Validate the tangents
-				if (!FDynamicMeshTangents(Mesh->GetMeshPtr()).HasValidTangents(true))
-				{
-					if (Mesh == TargetMesh)
-					{
-						UE::Geometry::AppendError(Debug, EGeometryScriptErrorType::InvalidInputs, LOCTEXT("BakeTexture_InvalidTargetTangents", "BakeTexture: Target Mesh tangents are invalid."));
-					}
-					else
-					{
-						ensure(Mesh == SourceMesh);
-						UE::Geometry::AppendError(Debug, EGeometryScriptErrorType::InvalidInputs, LOCTEXT("BakeTexture_InvalidSourceTangents", "BakeTexture: Source Mesh tangents are invalid."));
-					}
-					return false;
-				}
-			}
-			return true;
-		};
+		const FText BakeTexturePrefix = LOCTEXT("BakeTexture_Prefix", "BakeTexture");
 
 		const bool bIsBakeToSelf = (TargetMesh == SourceMesh);
 
@@ -296,175 +496,34 @@ namespace GeometryScriptBakeLocals
 		Baker.SetProjectionDistance(BakeOptions.ProjectionDistance);
 		Baker.SetSamplesPerPixel(GetSamplesPerPixel(BakeOptions.SamplesPerPixel));
 
-		bool bSupportsSourceNormalMap = false;
-		TSharedPtr<UE::Geometry::TImageBuilder<FVector4f>, ESPMode::ThreadSafe> SourceTexture;
-		for (const FGeometryScriptBakeTypeOptions& Options : BakeTypes)
+		auto IsValidBakeType = [](EGeometryScriptBakeTypes, TArray<FGeometryScriptDebugMessage>*)
 		{
-			switch(Options.BakeType)
-			{
-			case EGeometryScriptBakeTypes::TangentSpaceNormal:
-			{
-				if (!GetMeshTangents(TargetMesh, TargetMeshTangents))
-				{
-					return nullptr;
-				}
-				TSharedPtr<FMeshNormalMapEvaluator> NormalEval = MakeShared<FMeshNormalMapEvaluator>();
-				bSupportsSourceNormalMap = true;
-				Baker.AddEvaluator(NormalEval);
-				break;
-			}
-			case EGeometryScriptBakeTypes::ObjectSpaceNormal:
-			{
-				TSharedPtr<FMeshPropertyMapEvaluator> PropertyEval = MakeShared<FMeshPropertyMapEvaluator>();
-				PropertyEval->Property = EMeshPropertyMapType::Normal;
-				bSupportsSourceNormalMap = true;
-				Baker.AddEvaluator(PropertyEval);
-				break;
-			}
-			case EGeometryScriptBakeTypes::FaceNormal:
-			{
-				TSharedPtr<FMeshPropertyMapEvaluator> PropertyEval = MakeShared<FMeshPropertyMapEvaluator>();
-				PropertyEval->Property = EMeshPropertyMapType::FacetNormal;
-				Baker.AddEvaluator(PropertyEval);
-				break;
-			}
-			case EGeometryScriptBakeTypes::BentNormal:
-			{
-				if (!GetMeshTangents(TargetMesh, TargetMeshTangents))
-				{
-					return nullptr;
-				}
-				FGeometryScriptBakeType_Occlusion* OcclusionOptions = static_cast<FGeometryScriptBakeType_Occlusion*>(Options.Options.Get());
-				TSharedPtr<FMeshOcclusionMapEvaluator> OcclusionEval = MakeShared<FMeshOcclusionMapEvaluator>();
-				OcclusionEval->OcclusionType = EMeshOcclusionMapType::BentNormal;
-				OcclusionEval->NumOcclusionRays = OcclusionOptions->OcclusionRays;
-				OcclusionEval->MaxDistance = (OcclusionOptions->MaxDistance == 0) ? TNumericLimits<float>::Max() : OcclusionOptions->MaxDistance;
-				OcclusionEval->SpreadAngle = OcclusionOptions->SpreadAngle;
-				Baker.AddEvaluator(OcclusionEval);
-				break;
-			}
-			case EGeometryScriptBakeTypes::Position:
-			{
-				TSharedPtr<FMeshPropertyMapEvaluator> PropertyEval = MakeShared<FMeshPropertyMapEvaluator>();
-				PropertyEval->Property = EMeshPropertyMapType::Position;
-				Baker.AddEvaluator(PropertyEval);
-				break;
-			}
-			case EGeometryScriptBakeTypes::Curvature:
-			{
-				FGeometryScriptBakeType_Curvature* CurvatureOptions = static_cast<FGeometryScriptBakeType_Curvature*>(Options.Options.Get());
-				TSharedPtr<FMeshCurvatureMapEvaluator> CurvatureEval = MakeShared<FMeshCurvatureMapEvaluator>();
-				CurvatureEval->UseCurvatureType = GetCurvatureType(CurvatureOptions->CurvatureType);
-				CurvatureEval->UseColorMode = GetCurvatureColorMode(CurvatureOptions->ColorMapping);
-				CurvatureEval->RangeScale = CurvatureOptions->ColorRangeMultiplier;
-				CurvatureEval->MinRangeScale = CurvatureOptions->MinRangeMultiplier;
-				CurvatureEval->UseClampMode = GetCurvatureClampMode(CurvatureOptions->Clamping);
-				Baker.AddEvaluator(CurvatureEval);
-				break;
-			}
-			case EGeometryScriptBakeTypes::AmbientOcclusion:
-			{
-				FGeometryScriptBakeType_Occlusion* OcclusionOptions = static_cast<FGeometryScriptBakeType_Occlusion*>(Options.Options.Get());
-				TSharedPtr<FMeshOcclusionMapEvaluator> OcclusionEval = MakeShared<FMeshOcclusionMapEvaluator>();
-				OcclusionEval->OcclusionType = EMeshOcclusionMapType::AmbientOcclusion;
-				OcclusionEval->NumOcclusionRays = OcclusionOptions->OcclusionRays;
-				OcclusionEval->MaxDistance = (OcclusionOptions->MaxDistance == 0) ? TNumericLimits<float>::Max() : OcclusionOptions->MaxDistance;
-				OcclusionEval->SpreadAngle = OcclusionOptions->SpreadAngle;
-				OcclusionEval->BiasAngleDeg = OcclusionOptions->BiasAngle;
-				Baker.AddEvaluator(OcclusionEval);
-				break;
-			}
-			case EGeometryScriptBakeTypes::Texture:
-			{
-				FGeometryScriptBakeType_Texture* TextureOptions = static_cast<FGeometryScriptBakeType_Texture*>(Options.Options.Get());
-
-				// TODO: Add support for sampling different texture maps per Texture evaluator in a single pass. 
-				if (!SourceTexture && TextureOptions->SourceTexture)
-				{
-					SourceTexture = MakeShared<TImageBuilder<FVector4f>>();
-					if (!UE::AssetUtils::ReadTexture(TextureOptions->SourceTexture, *SourceTexture, false))
-					{
-						UE::Geometry::AppendError(Debug, EGeometryScriptErrorType::InvalidInputs, LOCTEXT("BakeTexture_InvalidSourceTexture", "BakeTexture: Failed to read SourceTexture"));
-					}
-					else
-					{
-						DetailSampler.SetTextureMap(SourceMeshToUse, IMeshBakerDetailSampler::FBakeDetailTexture(SourceTexture.Get(), TextureOptions->SourceUVLayer));
-					}
-				}
-				
-				TSharedPtr<FMeshResampleImageEvaluator> TextureEval = MakeShared<FMeshResampleImageEvaluator>();
-				Baker.AddEvaluator(TextureEval);
-				break;
-			}
-			case EGeometryScriptBakeTypes::MultiTexture:
-			{
-				FGeometryScriptBakeType_MultiTexture* TextureOptions = static_cast<FGeometryScriptBakeType_MultiTexture*>(Options.Options.Get());
-				TSharedPtr<FMeshMultiResampleImageEvaluator> TextureEval = MakeShared<FMeshMultiResampleImageEvaluator>();
+			return true;
+		};
 
-				if (TextureOptions->MaterialIDSourceTextures.Num())
-				{
-					TextureEval->MultiTextures.SetNum(TextureOptions->MaterialIDSourceTextures.Num());
-					for (int32 MaterialId = 0; MaterialId < TextureEval->MultiTextures.Num(); ++MaterialId)
-					{
-						if (UTexture2D* Texture = TextureOptions->MaterialIDSourceTextures[MaterialId])
-						{
-							TextureEval->MultiTextures[MaterialId] = MakeShared<TImageBuilder<FVector4f>>();
-							if (!UE::AssetUtils::ReadTexture(Texture, *TextureEval->MultiTextures[MaterialId], false))
-							{
-								UE::Geometry::AppendError(Debug, EGeometryScriptErrorType::InvalidInputs, LOCTEXT("BakeTexture_InvalidMultiTexture", "BakeTexture: Failed to read MaterialIDSourceTexture"));
-							}
-						}
-					}
-				}
-				Baker.AddEvaluator(TextureEval);
-				break;
-			}
-			case EGeometryScriptBakeTypes::VertexColor:
-			{
-				TSharedPtr<FMeshPropertyMapEvaluator> PropertyEval = MakeShared<FMeshPropertyMapEvaluator>();
-				PropertyEval->Property = EMeshPropertyMapType::VertexColor;
-				Baker.AddEvaluator(PropertyEval);
-				break;
-			}
-			case EGeometryScriptBakeTypes::MaterialID:
+		FEvaluatorState EvalState;
+		EvalState.TargetMesh = TargetMesh->GetMeshPtr();
+		EvalState.SourceMesh = SourceMeshToUse;
+		EvalState.DetailSampler = &DetailSampler;
+		for (const FGeometryScriptBakeTypeOptions& Options : BakeTypes)
+		{
+			TSharedPtr<FMeshMapEvaluator> Eval = CreateEvaluator(EvalState, IsValidBakeType, Options, BakeTexturePrefix, Debug);
+			if (!Eval)
 			{
-				TSharedPtr<FMeshPropertyMapEvaluator> PropertyEval = MakeShared<FMeshPropertyMapEvaluator>();
-				PropertyEval->Property = EMeshPropertyMapType::MaterialID;
-				Baker.AddEvaluator(PropertyEval);
-				break;
-			}
-			default:
-				break;
+				// Abort if any evaluators failed to build.
+				return nullptr;
 			}
+			Baker.AddEvaluator(Eval);
 		}
 
-		TSharedPtr<UE::Geometry::TImageBuilder<FVector4f>, ESPMode::ThreadSafe> SourceNormalMap;
-		if (bSupportsSourceNormalMap && SourceOptions.SourceNormalMap)
+		if (EvalState.bSupportsSourceNormalMap && SourceOptions.SourceNormalMap)
 		{
-			SourceNormalMap = MakeShared<TImageBuilder<FVector4f>>();
-			if (!UE::AssetUtils::ReadTexture(SourceOptions.SourceNormalMap, *SourceNormalMap, false))
-			{
-				UE::Geometry::AppendError(Debug, EGeometryScriptErrorType::InvalidInputs, LOCTEXT("BakeTexture_InvalidSourceNormalMap", "BakeTexture: Failed to read SourceNormalMap"));
-			}
-			else
-			{
-				DetailSampler.SetNormalTextureMap(SourceMeshToUse,
-				IMeshBakerDetailSampler::FBakeDetailNormalTexture(
-						SourceNormalMap.Get(),
-						SourceOptions.SourceNormalUVLayer,
-						SourceOptions.SourceNormalSpace == EGeometryScriptBakeNormalSpace::Tangent ? IMeshBakerDetailSampler::EBakeDetailNormalSpace::Tangent : IMeshBakerDetailSampler::EBakeDetailNormalSpace::Object));
-
-				if (!GetMeshTangents(SourceMesh, SourceMeshTangents))
-				{
-					return nullptr;
-				}
-				DetailSampler.SetTangents(SourceMeshToUse, SourceMeshTangents.Get());
-			}
+			GetSourceNormalMap(EvalState, SourceOptions, BakeTexturePrefix, Debug);
 		}
 
-		if (TargetMeshTangents)
+		if (EvalState.TargetMeshTangents)
 		{
-			Baker.SetTargetMeshTangents(TargetMeshTangents);
+			Baker.SetTargetMeshTangents(EvalState.TargetMeshTangents);
 		}
 
 		Baker.Bake();
@@ -504,6 +563,167 @@ namespace GeometryScriptBakeLocals
 			Textures.Add(TextureBuilder.GetTexture2D());
 		}
 	}
+
+	TUniquePtr<FMeshVertexBaker> BakeVertexImpl(
+		UDynamicMesh* TargetMesh,
+		const FTransform& TargetTransform,
+		FGeometryScriptBakeTargetMeshOptions TargetOptions,
+		UDynamicMesh* SourceMesh,
+		const FTransform& SourceTransform,
+		FGeometryScriptBakeSourceMeshOptions SourceOptions,
+		FGeometryScriptBakeOutputType BakeTypes,
+		FGeometryScriptBakeVertexOptions BakeOptions,
+		TArray<FGeometryScriptDebugMessage>* Debug)
+	{
+		if (TargetMesh == nullptr)
+		{
+			UE::Geometry::AppendError(Debug, EGeometryScriptErrorType::InvalidInputs, LOCTEXT("BakeVertex_InvalidTargetMesh", "BakeVertex: TargetMesh is Null"));
+			return nullptr;
+		}
+		if (SourceMesh == nullptr)
+		{
+			UE::Geometry::AppendError(Debug, EGeometryScriptErrorType::InvalidInputs, LOCTEXT("BakeVertex_InvalidSourceMesh", "BakeVertex: SourceMesh is Null"));
+			return nullptr;
+		}
+
+		const FText BakeVertexPrefix = LOCTEXT("BakeVertex_Prefix", "BakeVertex");
+
+		const bool bIsBakeToSelf = (TargetMesh == SourceMesh);
+
+		// Initialize the color overlay on the TargetMesh
+		FDynamicMesh3& TargetMeshRef = TargetMesh->GetMeshRef();
+		TargetMeshRef.EnableAttributes();
+		TargetMeshRef.Attributes()->DisablePrimaryColors();
+		TargetMeshRef.Attributes()->EnablePrimaryColors();
+
+		FDynamicMeshNormalOverlay* NormalOverlay = TargetMeshRef.Attributes()->PrimaryNormals();
+		FDynamicMeshUVOverlay* UVOverlay = TargetMeshRef.Attributes()->PrimaryUV();
+		TargetMeshRef.Attributes()->PrimaryColors()->CreateFromPredicate(
+			[&BakeOptions, NormalOverlay, UVOverlay](int ParentVID, int TriIDA, int TriIDB) -> bool
+			{
+				auto OverlayCanShare = [&] (auto Overlay) -> bool
+				{
+					return Overlay ? Overlay->AreTrianglesConnected(TriIDA, TriIDB) : true;
+				};
+				
+				bool bCanShare = true;
+				if (BakeOptions.bSplitAtNormalSeams)
+				{
+					bCanShare = bCanShare && OverlayCanShare(NormalOverlay);
+				}
+				if (BakeOptions.bSplitAtUVSeams)
+				{
+					bCanShare = bCanShare && OverlayCanShare(UVOverlay);
+				}
+				return bCanShare;
+			}, 0.0f);
+		
+
+		// Initialize the source mesh
+		FDynamicMesh3 SourceMeshCopy;
+		const FDynamicMesh3* SourceMeshOriginal = SourceMesh->GetMeshPtr();
+		const FDynamicMesh3* SourceMeshToUse = SourceMeshOriginal;
+		if (BakeOptions.bProjectionInWorldSpace && !bIsBakeToSelf)
+		{
+			// Transform the SourceMesh into TargetMesh local space using a copy (oof)
+			// TODO: Remove this once we have support for transforming rays in the core bake loop
+			SourceMeshCopy = *SourceMeshOriginal;
+			const FTransformSRT3d SourceToWorld = SourceTransform;
+			MeshTransforms::ApplyTransform(SourceMeshCopy, SourceToWorld);
+			const FTransformSRT3d TargetToWorld = TargetTransform;
+			MeshTransforms::ApplyTransform(SourceMeshCopy, TargetToWorld.Inverse());
+			SourceMeshToUse = &SourceMeshCopy;
+		}
+
+		const FDynamicMeshAABBTree3 DetailSpatial(SourceMeshToUse);
+		FMeshBakerDynamicMeshSampler DetailSampler(SourceMeshToUse, &DetailSpatial);
+
+		TUniquePtr<FMeshVertexBaker> Result = MakeUnique<FMeshVertexBaker>();
+		FMeshVertexBaker& Baker = *Result;
+		Baker.BakeMode = BakeTypes.OutputMode == EGeometryScriptBakeOutputMode::RGBA ? FMeshVertexBaker::EBakeMode::RGBA : FMeshVertexBaker::EBakeMode::PerChannel;
+		Baker.SetTargetMesh(TargetMesh->GetMeshPtr());
+		Baker.SetTargetMeshUVLayer(TargetOptions.TargetUVLayer);
+		Baker.SetDetailSampler(&DetailSampler);
+		Baker.SetProjectionDistance(BakeOptions.ProjectionDistance);
+
+		FEvaluatorState EvalState;
+		EvalState.TargetMesh = TargetMesh->GetMeshPtr();
+		EvalState.SourceMesh = SourceMeshToUse;
+		EvalState.DetailSampler = &DetailSampler;
+
+		if (BakeTypes.OutputMode == EGeometryScriptBakeOutputMode::RGBA)
+		{
+			auto IsValidBakeType = [&BakeVertexPrefix](EGeometryScriptBakeTypes BakeType, TArray<FGeometryScriptDebugMessage>* Debug)
+			{
+				const bool bIsValid = BakeType != EGeometryScriptBakeTypes::VertexColor;
+				if (!bIsValid)
+				{
+					const FText BakeTypeName = FText::FromName(StaticEnum<EGeometryScriptBakeTypes>()->GetNameByIndex(static_cast<int32>(BakeType)));
+					UE::Geometry::AppendError(Debug, EGeometryScriptErrorType::InvalidInputs, FText::Format(LOCTEXT("BakeVertex_InvalidColorEval", "{0}: {1} bake type is not a supported RGBA evaluator."), BakeVertexPrefix, BakeTypeName));
+				}
+				return bIsValid;
+			};
+			
+			Baker.BakeMode = FMeshVertexBaker::EBakeMode::RGBA;
+			Baker.ColorEvaluator = CreateEvaluator(EvalState, IsValidBakeType, BakeTypes.RGBA, BakeVertexPrefix, Debug);
+		}
+		else
+		{
+			ensure(BakeTypes.OutputMode == EGeometryScriptBakeOutputMode::PerChannel);
+
+			auto IsValidBakeType = [&BakeVertexPrefix](EGeometryScriptBakeTypes BakeType, TArray<FGeometryScriptDebugMessage>* Debug)
+			{
+				const bool bIsValid = (BakeType == EGeometryScriptBakeTypes::AmbientOcclusion || BakeType == EGeometryScriptBakeTypes::Curvature);
+				if (!bIsValid)
+				{
+					const FText BakeTypeName = FText::FromName(StaticEnum<EGeometryScriptBakeTypes>()->GetNameByIndex(static_cast<int32>(BakeType)));
+					UE::Geometry::AppendError(Debug, EGeometryScriptErrorType::InvalidInputs, FText::Format(LOCTEXT("BakeVertex_InvalidChannelEval", "{0}: {1} bake type is not a supported per-channel evaluator."), BakeVertexPrefix, BakeTypeName));
+				}
+				return bIsValid;
+			};
+			
+			Baker.BakeMode = FMeshVertexBaker::EBakeMode::PerChannel;
+			Baker.ChannelEvaluators[0] = CreateEvaluator(EvalState, IsValidBakeType, BakeTypes.R, BakeVertexPrefix, Debug);
+			Baker.ChannelEvaluators[1] = CreateEvaluator(EvalState, IsValidBakeType, BakeTypes.G, BakeVertexPrefix, Debug);
+			Baker.ChannelEvaluators[2] = CreateEvaluator(EvalState, IsValidBakeType, BakeTypes.B, BakeVertexPrefix, Debug);
+			Baker.ChannelEvaluators[3] = CreateEvaluator(EvalState, IsValidBakeType, BakeTypes.A, BakeVertexPrefix, Debug);
+		}
+
+		TSharedPtr<UE::Geometry::TImageBuilder<FVector4f>, ESPMode::ThreadSafe> SourceNormalMap;
+		TSharedPtr<FMeshTangentsd, ESPMode::ThreadSafe> SourceMeshTangents;
+		if (EvalState.bSupportsSourceNormalMap && SourceOptions.SourceNormalMap)
+		{
+			GetSourceNormalMap(EvalState, SourceOptions, BakeVertexPrefix, Debug);
+		}
+
+		if (EvalState.TargetMeshTangents)
+		{
+			Baker.SetTargetMeshTangents(EvalState.TargetMeshTangents);
+		}
+
+		Baker.Bake();
+		
+		return MoveTemp(Result);
+	}
+
+	bool ApplyVertexBakeToMesh(const FMeshVertexBaker* Baker, UDynamicMesh* Mesh)
+	{
+		if (!Baker || !Mesh)
+		{
+			return false;
+		}
+
+		FDynamicMesh3& MeshRef = Mesh->GetMeshRef();		
+		const TImageBuilder<FVector4f>* ImageResult = Baker->GetBakeResult();
+		const int NumColors = MeshRef.Attributes()->PrimaryColors()->ElementCount();
+		check(NumColors == ImageResult->GetDimensions().GetWidth());
+		for (int Idx = 0; Idx < NumColors; ++Idx)
+		{
+			const FVector4f& Pixel = ImageResult->GetPixel(Idx);
+			MeshRef.Attributes()->PrimaryColors()->SetElement(Idx, Pixel);
+		}
+		return true;
+	}
 }
 
 
@@ -693,7 +913,7 @@ void UGeometryScriptLibrary_MeshBakeFunctions::BakeTextureAsyncBegin(
 
 		TSharedPtr<FMeshMapBaker> SharedBaker = MakeShareable<FMeshMapBaker>(Baker.Release());
 
-		AsyncTask(ENamedThreads::GameThread, [Completed, BakeId, SharedBaker, BakeOptions, Debug, DebugMessages]()
+		AsyncTask(ENamedThreads::GameThread, [Completed, BakeId, SharedBaker, BakeOptions, TargetMesh, SourceMesh, Debug, DebugMessages]()
 		{
 			if (Debug && DebugMessages)
 			{
@@ -723,4 +943,99 @@ TArray<UTexture2D*> UGeometryScriptLibrary_MeshBakeFunctions::BakeTextureAsyncEn
 	return TextureOutput;
 }
 
+UDynamicMesh* UGeometryScriptLibrary_MeshBakeFunctions::BakeVertex(
+		UDynamicMesh* TargetMesh,
+		FTransform TargetTransform,
+		FGeometryScriptBakeTargetMeshOptions TargetOptions,
+		UDynamicMesh* SourceMesh,
+		FTransform SourceTransform,
+		FGeometryScriptBakeSourceMeshOptions SourceOptions,
+		FGeometryScriptBakeOutputType BakeTypes,
+		FGeometryScriptBakeVertexOptions BakeOptions,
+		UGeometryScriptDebug* Debug)
+{
+	const TUniquePtr<FMeshVertexBaker> Baker = GeometryScriptBakeLocals::BakeVertexImpl(
+		TargetMesh,
+		TargetTransform,
+		TargetOptions,
+		SourceMesh,
+		SourceTransform,
+		SourceOptions,
+		BakeTypes,
+		BakeOptions,
+		Debug ? &Debug->Messages : nullptr);
+
+	// Extract the vertex bake data and apply to target mesh.
+	GeometryScriptBakeLocals::ApplyVertexBakeToMesh(Baker.Get(), TargetMesh);
+
+	return TargetMesh;
+}
+
+void UGeometryScriptLibrary_MeshBakeFunctions::BakeVertexAsyncBegin(
+	const FBakeVertexDelegate& Completed,
+	int BakeId,
+	UDynamicMesh* TargetMesh,
+	FTransform TargetTransform,
+	FGeometryScriptBakeTargetMeshOptions TargetOptions,
+	UDynamicMesh* SourceMesh,
+	FTransform SourceTransform,
+	FGeometryScriptBakeSourceMeshOptions SourceOptions,
+	FGeometryScriptBakeOutputType BakeTypes,
+	FGeometryScriptBakeVertexOptions BakeOptions,
+	UGeometryScriptDebug* Debug)
+{
+	AsyncTask(ENamedThreads::AnyBackgroundThreadNormalTask, [Completed, BakeId, TargetMesh, TargetTransform, TargetOptions, SourceMesh, SourceTransform, SourceOptions, BakeTypes, BakeOptions, Debug]()
+	{
+		// FGeometryScriptDebug is not thread-safe. Instantiate a local
+		// debug message array to collect any debug errors and append those
+		// messages back on the game thread.
+		TSharedPtr<TArray<FGeometryScriptDebugMessage>, ESPMode::ThreadSafe> DebugMessages;
+		if (Debug)
+		{
+			DebugMessages = MakeShared<TArray<FGeometryScriptDebugMessage>, ESPMode::ThreadSafe>();
+		}
+		
+		TUniquePtr<FMeshVertexBaker> Baker = GeometryScriptBakeLocals::BakeVertexImpl(
+			TargetMesh,
+			TargetTransform,
+			TargetOptions,
+			SourceMesh,
+			SourceTransform,
+			SourceOptions,
+			BakeTypes,
+			BakeOptions,
+			DebugMessages.Get());
+
+		TSharedPtr<FMeshVertexBaker> SharedBaker = MakeShareable<FMeshVertexBaker>(Baker.Release());
+
+		AsyncTask(ENamedThreads::GameThread, [Completed, BakeId, SharedBaker, TargetMesh, Debug, DebugMessages]()
+		{
+			if (Debug && DebugMessages)
+			{
+				for (const FGeometryScriptDebugMessage& Message : *DebugMessages)
+				{
+					Debug->Messages.Add(Message);
+				}
+			}
+			
+			FGeometryScriptBakeVertexAsyncResult Result;
+			Result.BakeResult = SharedBaker;
+			Result.TargetMesh = TargetMesh;
+			Completed.ExecuteIfBound(BakeId, Result);
+		});
+	});
+}
+
+UDynamicMesh* UGeometryScriptLibrary_MeshBakeFunctions::BakeVertexAsyncEnd(const FGeometryScriptBakeVertexAsyncResult& Result)
+{
+	FMeshVertexBaker* Baker = Result.BakeResult.Get();
+
+	// Extract the vertex bake data and apply to target mesh.
+	GeometryScriptBakeLocals::ApplyVertexBakeToMesh(Baker, Result.TargetMesh);
+
+	return Result.TargetMesh;
+}
+
+
+
 #undef LOCTEXT_NAMESPACE
diff --git a/Engine/Plugins/Experimental/GeometryScripting/Source/GeometryScriptingCore/Public/GeometryScript/MeshBakeFunctions.h b/Engine/Plugins/Experimental/GeometryScripting/Source/GeometryScriptingCore/Public/GeometryScript/MeshBakeFunctions.h
--- a/Engine/Plugins/Experimental/GeometryScripting/Source/GeometryScriptingCore/Public/GeometryScript/MeshBakeFunctions.h
+++ b/Engine/Plugins/Experimental/GeometryScripting/Source/GeometryScriptingCore/Public/GeometryScript/MeshBakeFunctions.h
@@ -6,6 +6,7 @@
 #include "Kismet/BlueprintFunctionLibrary.h"
 #include "GeometryScript/GeometryScriptTypes.h"
 #include "Sampling/MeshMapBaker.h"
+#include "Sampling/MeshVertexBaker.h"
 #include "MeshBakeFunctions.generated.h"
 
 class UDynamicMesh;
@@ -69,6 +70,13 @@ enum class EGeometryScriptBakeTypes : uint8
 	MaterialID             UMETA(DisplayName = "Material ID"),
 };
 
+UENUM(BlueprintType)
+enum class EGeometryScriptBakeOutputMode : uint8
+{
+	RGBA,
+	PerChannel
+};
+
 UENUM(BlueprintType)
 enum class EGeometryScriptBakeNormalSpace : uint8
 {
@@ -213,6 +221,55 @@ struct GEOMETRYSCRIPTINGCORE_API FGeometryScriptBakeTextureOptions
 	bool bProjectionInWorldSpace = false;
 };
 
+USTRUCT(BlueprintType)
+struct GEOMETRYSCRIPTINGCORE_API FGeometryScriptBakeVertexOptions
+{
+	GENERATED_BODY()
+
+	/** If true, compute a separate vertex color for each unique normal on a vertex */
+	UPROPERTY(BlueprintReadWrite, Category = Options)
+	bool bSplitAtNormalSeams = false;
+
+	/** If true, compute a separate vertex color for each unique UV on a vertex. */
+	UPROPERTY(BlueprintReadWrite, Category = Options, meta=(DisplayName = "Split at UV Seams"))
+	bool bSplitAtUVSeams = false;
+
+	/** Maximum allowed distance for the projection from target mesh to source mesh for the sample to be considered valid.
+	 * This is only relevant if a separate source mesh is provided. */
+	UPROPERTY(BlueprintReadWrite, Category = Options)
+	float ProjectionDistance = 3.0f;
+
+	/** If true, uses the world space positions for the projection from target mesh to source mesh, otherwise it uses their object space positions.
+	 * This is only relevant if a separate source mesh is provided. */
+	UPROPERTY(BlueprintReadWrite, Category = Options)
+	bool bProjectionInWorldSpace = false;
+};
+
+USTRUCT(BlueprintType)
+struct GEOMETRYSCRIPTINGCORE_API FGeometryScriptBakeOutputType
+{
+	GENERATED_BODY()
+	
+	/** The bake output mode */
+	UPROPERTY(BlueprintReadWrite, Category = Options)
+	EGeometryScriptBakeOutputMode OutputMode = EGeometryScriptBakeOutputMode::RGBA;
+
+	UPROPERTY(BlueprintReadWrite, Category = Output)
+	FGeometryScriptBakeTypeOptions RGBA;
+
+	UPROPERTY(BlueprintReadWrite, Category = Output)
+	FGeometryScriptBakeTypeOptions R;
+
+	UPROPERTY(BlueprintReadWrite, Category = Output)
+	FGeometryScriptBakeTypeOptions G;
+
+	UPROPERTY(BlueprintReadWrite, Category = Output)
+	FGeometryScriptBakeTypeOptions B;
+
+	UPROPERTY(BlueprintReadWrite, Category = Output)
+	FGeometryScriptBakeTypeOptions A;
+};
+
 USTRUCT(BlueprintType)
 struct GEOMETRYSCRIPTINGCORE_API FGeometryScriptBakeTargetMeshOptions
 {
@@ -246,7 +303,17 @@ struct GEOMETRYSCRIPTINGCORE_API FGeometryScriptBakeTextureAsyncResult
 	TSharedPtr<UE::Geometry::FMeshMapBaker, ESPMode::ThreadSafe> BakeResult;
 };
 
+USTRUCT(BlueprintType)
+struct GEOMETRYSCRIPTINGCORE_API FGeometryScriptBakeVertexAsyncResult
+{
+	GENERATED_BODY();
+
+	TSharedPtr<UE::Geometry::FMeshVertexBaker, ESPMode::ThreadSafe> BakeResult;
+	TObjectPtr<UDynamicMesh> TargetMesh;
+};
+
 DECLARE_DYNAMIC_DELEGATE_TwoParams(FBakeTextureDelegate, int, BakeId, FGeometryScriptBakeTextureAsyncResult, ResultOut);
+DECLARE_DYNAMIC_DELEGATE_TwoParams(FBakeVertexDelegate, int, BakeId, FGeometryScriptBakeVertexAsyncResult, ResultOut);
 
 UCLASS(meta = (ScriptName = "GeometryScript_Bake"))
 class GEOMETRYSCRIPTINGCORE_API UGeometryScriptLibrary_MeshBakeFunctions : public UBlueprintFunctionLibrary
@@ -336,7 +403,7 @@ public:
 	 * @param BakeOptions bake options (ex. resolution, number of samples, projection distance).
 	 * @param Debug debug structure to capture debug messages.
 	 */
-	UFUNCTION(BlueprintCallable, Category = "GeometryScript|Bake")
+	UFUNCTION(/*BlueprintCallable,*/ Category = "GeometryScript|Bake")
 	static void BakeTextureAsyncBegin(
 		const FBakeTextureDelegate& Completed,
 		int BakeId,
@@ -358,7 +425,37 @@ public:
 	 * @param Result the result of a BakeTextureAsyncBegin() via delegate
 	 * @return array of UTexture2D corresponding to the requested bake types in BakeTextureAsyncBegin().
 	 */
-	UFUNCTION(BlueprintCallable, Category = "GeometryScript|Bake")
+	UFUNCTION(/*BlueprintCallable,*/ Category = "GeometryScript|Bake")
 	static UPARAM(DisplayName="Textures Out") TArray<UTexture2D*> BakeTextureAsyncEnd(
 		const FGeometryScriptBakeTextureAsyncResult& Result);
+
+	UFUNCTION(BlueprintCallable, Category = "GeometryScript|Bake")
+	static UPARAM(DisplayName="Target Mesh") UDynamicMesh* BakeVertex(
+		UDynamicMesh* TargetMesh,
+		FTransform TargetTransform,
+		FGeometryScriptBakeTargetMeshOptions TargetOptions,
+		UDynamicMesh* SourceMesh,
+		FTransform SourceTransform,
+		FGeometryScriptBakeSourceMeshOptions SourceOptions,
+		FGeometryScriptBakeOutputType BakeTypes,
+		FGeometryScriptBakeVertexOptions BakeOptions,
+		UGeometryScriptDebug* Debug = nullptr);
+
+	UFUNCTION(/*BlueprintCallable,*/ Category = "GeometryScript|Bake")
+	static void BakeVertexAsyncBegin(
+		const FBakeVertexDelegate& Completed,
+		int BakeId,
+		UDynamicMesh* TargetMesh,
+		FTransform TargetTransform,
+		FGeometryScriptBakeTargetMeshOptions TargetOptions,
+		UDynamicMesh* SourceMesh,
+		FTransform SourceTransform,
+		FGeometryScriptBakeSourceMeshOptions SourceOptions,
+		FGeometryScriptBakeOutputType BakeTypes,
+		FGeometryScriptBakeVertexOptions BakeOptions,
+		UGeometryScriptDebug* Debug = nullptr);
+
+	UFUNCTION(/*BlueprintCallable,*/ Category = "GeometryScript|Bake")
+	static UPARAM(DisplayName="Target Mesh") UDynamicMesh* BakeVertexAsyncEnd(
+		const FGeometryScriptBakeVertexAsyncResult& Result);
 };
