Geometry Scripting: Added the Spectral Confrormal UV Solver option to the RecomputeMeshUVs function.

diff --git a/Engine/Plugins/Experimental/GeometryScripting/Source/GeometryScriptingCore/Private/MeshUVFunctions.cpp b/Engine/Plugins/Experimental/GeometryScripting/Source/GeometryScriptingCore/Private/MeshUVFunctions.cpp
--- a/Engine/Plugins/Experimental/GeometryScripting/Source/GeometryScriptingCore/Private/MeshUVFunctions.cpp
+++ b/Engine/Plugins/Experimental/GeometryScripting/Source/GeometryScriptingCore/Private/MeshUVFunctions.cpp
@@ -437,22 +437,32 @@ UDynamicMesh* UGeometryScriptLibrary_MeshUVFunctions::RecomputeMeshUVs(
 			bComponentSolved[k] = false;
 			switch (Options.Method)
 			{
-			case EGeometryScriptUVFlattenMethod::ExpMap:
-			{
-				FDynamicMeshUVEditor::FExpMapOptions ExpMapOptions;
-				ExpMapOptions.NormalSmoothingRounds = Options.ExpMapOptions.NormalSmoothingRounds;
-				ExpMapOptions.NormalSmoothingAlpha = Options.ExpMapOptions.NormalSmoothingAlpha;
-				bComponentSolved[k] = UVEditor.SetTriangleUVsFromExpMap(ComponentTris, ExpMapOptions);
-			}
-			break;
-
-			case EGeometryScriptUVFlattenMethod::Conformal:
-				bComponentSolved[k] = UVEditor.SetTriangleUVsFromFreeBoundaryConformal(ComponentTris);
-				if ( bComponentSolved[k] )
+				case EGeometryScriptUVFlattenMethod::ExpMap:
+				{
+					FDynamicMeshUVEditor::FExpMapOptions ExpMapOptions;
+					ExpMapOptions.NormalSmoothingRounds = Options.ExpMapOptions.NormalSmoothingRounds;
+					ExpMapOptions.NormalSmoothingAlpha = Options.ExpMapOptions.NormalSmoothingAlpha;
+					bComponentSolved[k] = UVEditor.SetTriangleUVsFromExpMap(ComponentTris, ExpMapOptions);
+					break;
+				}
+				case EGeometryScriptUVFlattenMethod::Conformal:
+				{
+					bComponentSolved[k] = UVEditor.SetTriangleUVsFromFreeBoundaryConformal(ComponentTris);
+					if ( bComponentSolved[k] )
+					{
+						UVEditor.ScaleUVAreaTo3DArea(ComponentTris, true);
+					}
+					break;
+				}
+				case EGeometryScriptUVFlattenMethod::SpectralConformal:
 				{
-					UVEditor.ScaleUVAreaTo3DArea(ComponentTris, true);
+					bComponentSolved[k] = UVEditor.SetTriangleUVsFromFreeBoundarySpectralConformal(ComponentTris, false, Options.SpectralConformalOptions.bPreserveIrregularity);
+					if ( bComponentSolved[k] )
+					{
+						UVEditor.ScaleUVAreaTo3DArea(ComponentTris, true);
+					}
+					break;
 				}
-				break;
 			}
 		}
 
diff --git a/Engine/Plugins/Experimental/GeometryScripting/Source/GeometryScriptingCore/Public/GeometryScript/MeshUVFunctions.h b/Engine/Plugins/Experimental/GeometryScripting/Source/GeometryScriptingCore/Public/GeometryScript/MeshUVFunctions.h
--- a/Engine/Plugins/Experimental/GeometryScripting/Source/GeometryScriptingCore/Public/GeometryScript/MeshUVFunctions.h
+++ b/Engine/Plugins/Experimental/GeometryScripting/Source/GeometryScriptingCore/Public/GeometryScript/MeshUVFunctions.h
@@ -26,7 +26,8 @@ UENUM(BlueprintType)
 enum class EGeometryScriptUVFlattenMethod : uint8
 {
 	ExpMap = 0,
-	Conformal = 1
+	Conformal = 1,
+	SpectralConformal = 2
 };
 
 UENUM(BlueprintType)
@@ -49,6 +50,15 @@ struct GEOMETRYSCRIPTINGCORE_API FGeometryScriptExpMapUVOptions
 	float NormalSmoothingAlpha = 0.25f;
 };
 
+USTRUCT(BlueprintType)
+struct GEOMETRYSCRIPTINGCORE_API FGeometryScriptSpectralConformalUVOptions
+{
+	GENERATED_BODY()
+
+	UPROPERTY(BlueprintReadWrite, Category = Options)
+	bool bPreserveIrregularity = true;
+};
+
 
 USTRUCT(BlueprintType)
 struct GEOMETRYSCRIPTINGCORE_API FGeometryScriptRecomputeUVsOptions
@@ -56,7 +66,7 @@ struct GEOMETRYSCRIPTINGCORE_API FGeometryScriptRecomputeUVsOptions
 	GENERATED_BODY()
 public:
 	UPROPERTY(BlueprintReadWrite, Category = Options)
-	EGeometryScriptUVFlattenMethod Method = EGeometryScriptUVFlattenMethod::Conformal;
+	EGeometryScriptUVFlattenMethod Method = EGeometryScriptUVFlattenMethod::SpectralConformal;
 
 	UPROPERTY(BlueprintReadWrite, Category = Options)
 	EGeometryScriptUVIslandSource IslandSource = EGeometryScriptUVIslandSource::UVIslands;
@@ -64,6 +74,9 @@ public:
 	UPROPERTY(BlueprintReadWrite, Category = Options)
 	FGeometryScriptExpMapUVOptions ExpMapOptions;
 
+	UPROPERTY(BlueprintReadWrite, Category = Options)
+	FGeometryScriptSpectralConformalUVOptions SpectralConformalOptions;
+
 	UPROPERTY(BlueprintReadWrite, Category = Options)
 	FGeometryScriptGroupLayer GroupLayer;
 
