Add PolyPath functionality to Geometry Script

PolyPath is an optionally-closed polyline/polygon with 3D vertices + some utility functions and automatic conversion to/from arrays

Also add a USplineComponent->PolyPath conversion function

diff --git a/Engine/Plugins/Experimental/GeometryScripting/Source/GeometryScriptingCore/Private/PolyPathFunctions.cpp b/Engine/Plugins/Experimental/GeometryScripting/Source/GeometryScriptingCore/Private/PolyPathFunctions.cpp
new file mode 100644
--- /dev/null
+++ b/Engine/Plugins/Experimental/GeometryScripting/Source/GeometryScriptingCore/Private/PolyPathFunctions.cpp
@@ -0,0 +1,180 @@
+// Copyright Epic Games, Inc. All Rights Reserved.
+
+#include "GeometryScript/PolyPathFunctions.h"
+
+#include "Curve/CurveUtil.h"
+
+#define LOCTEXT_NAMESPACE "UGeometryScriptLibrary_PolyPathFunctions"
+
+
+int UGeometryScriptLibrary_PolyPathFunctions::GetPolyPathNumVertices(FGeometryScriptPolyPath PolyPath)
+{
+	return (PolyPath.Path.IsValid()) ? PolyPath.Path->Num() : 0;
+}
+
+int UGeometryScriptLibrary_PolyPathFunctions::GetPolyPathLastIndex(FGeometryScriptPolyPath PolyPath)
+{
+	return (PolyPath.Path.IsValid()) ? FMath::Max(PolyPath.Path->Num()-1,0) : 0;
+}
+
+FVector UGeometryScriptLibrary_PolyPathFunctions::GetPolyPathVertex(FGeometryScriptPolyPath PolyPath, int Index, bool& bIsValidIndex)
+{
+	bIsValidIndex = false;
+	if (PolyPath.Path.IsValid() && Index >= 0 && Index < PolyPath.Path->Num())
+	{
+		bIsValidIndex = true;
+		return (*PolyPath.Path)[Index];
+	}
+	return FVector::ZeroVector;
+}
+
+FVector UGeometryScriptLibrary_PolyPathFunctions::GetPolyPathTangent(FGeometryScriptPolyPath PolyPath, int Index, bool& bIsValidIndex)
+{
+	bIsValidIndex = false;
+	if (PolyPath.Path.IsValid() && Index >= 0 && Index < PolyPath.Path->Num())
+	{
+		bIsValidIndex = true;
+		FVector Tangent = UE::Geometry::CurveUtil::Tangent<double, FVector>(*PolyPath.Path, Index, PolyPath.bClosedLoop);
+		return Tangent;
+	}
+	return FVector::ZeroVector;
+}
+
+double UGeometryScriptLibrary_PolyPathFunctions::GetPolyPathArcLength(FGeometryScriptPolyPath PolyPath)
+{
+	if (PolyPath.Path.IsValid())
+	{
+		return UE::Geometry::CurveUtil::ArcLength<double, FVector>(*PolyPath.Path, PolyPath.bClosedLoop);
+	}
+	return 0;
+}
+
+int32 UGeometryScriptLibrary_PolyPathFunctions::GetNearestVertexIndex(FGeometryScriptPolyPath PolyPath, FVector Point)
+{
+	if (PolyPath.Path.IsValid())
+	{
+		return UE::Geometry::CurveUtil::FindNearestIndex<double, FVector>(*PolyPath.Path, Point);
+	}
+	return -1;
+}
+
+UFUNCTION(BlueprintCallable, Category = "GeometryScript|PolyPath")
+FGeometryScriptPolyPath UGeometryScriptLibrary_PolyPathFunctions::FlattenTo2DOnAxis(FGeometryScriptPolyPath PolyPath, EGeometryScriptAxis DropAxis)
+{
+	if (PolyPath.Path.IsValid())
+	{
+		int32 Keep0 = int32(DropAxis == EGeometryScriptAxis::X);
+		int32 Keep1 = 1 + int32(DropAxis != EGeometryScriptAxis::Z);
+		for (FVector& V : *PolyPath.Path)
+		{
+			V[0] = V[Keep0];
+			V[1] = V[Keep1];
+			V.Z = 0;
+		}
+	}
+	return PolyPath;
+}
+
+void UGeometryScriptLibrary_PolyPathFunctions::ConvertPolyPathToArray(FGeometryScriptPolyPath PolyPath, TArray<FVector>& PathVertices)
+{
+	PathVertices.Reset();
+	if (PolyPath.Path.IsValid())
+	{
+		PathVertices.Append(*PolyPath.Path);
+	}
+}
+
+void UGeometryScriptLibrary_PolyPathFunctions::ConvertArrayToPolyPath(const TArray<FVector>& PathVertices, FGeometryScriptPolyPath& PolyPath)
+{
+	PolyPath.Reset();
+	PolyPath.Path->Append(PathVertices);
+}
+
+void UGeometryScriptLibrary_PolyPathFunctions::ConvertPolyPathToArrayOfVector2D(FGeometryScriptPolyPath PolyPath, TArray<FVector2D>& PathVertices)
+{
+	PathVertices.Reset();
+	if (PolyPath.Path.IsValid())
+	{
+		PathVertices.Reserve(PolyPath.Path->Num());
+		for (const FVector& V : *PolyPath.Path)
+		{
+			PathVertices.Emplace(V.X, V.Y);
+		}
+	}
+}
+
+void UGeometryScriptLibrary_PolyPathFunctions::ConvertArrayOfVector2DToPolyPath(const TArray<FVector2D>& PathVertices, FGeometryScriptPolyPath& PolyPath)
+{
+	PolyPath.Reset();
+	PolyPath.Path->Reserve(PathVertices.Num());
+	for (const FVector2D& V : PathVertices)
+	{
+		PolyPath.Path->Emplace(V.X, V.Y, 0);
+	}
+}
+
+void UGeometryScriptLibrary_PolyPathFunctions::ConvertSplineToPolyPath(const USplineComponent* Spline, FGeometryScriptPolyPath& PolyPath, FGeometryScriptSplineSamplingOptions SamplingOptions)
+{
+	PolyPath.Reset();
+	if (Spline)
+	{
+		bool bIsLoop = Spline->IsClosedLoop();
+		PolyPath.bClosedLoop = bIsLoop;
+		if (SamplingOptions.SampleSpacing == EGeometryScriptSampleSpacing::ErrorTolerance)
+		{
+			float SquaredErrorTolerance = FMath::Max(KINDA_SMALL_NUMBER, SamplingOptions.ErrorTolerance * SamplingOptions.ErrorTolerance);
+			Spline->ConvertSplineToPolyLine(SamplingOptions.CoordinateSpace, SquaredErrorTolerance, *PolyPath.Path);
+			if (bIsLoop)
+			{
+				PolyPath.Path->Pop(); // delete the duplicate end-point for loops
+			}
+		}
+		else
+		{
+			float Duration = Spline->Duration;
+			
+			bool bUseConstantVelocity = SamplingOptions.SampleSpacing == EGeometryScriptSampleSpacing::UniformDistance;
+			int32 UseSamples = FMath::Max(2, SamplingOptions.NumSamples); // Always use at least 2 samples
+			// In non-loops, we adjust DivNum so we exactly sample the end of the spline
+			// In loops we don't sample the endpoint, by convention, as it's the same as the start
+			float DivNum = float(UseSamples - (int32)!bIsLoop);
+			PolyPath.Path->Reserve(UseSamples);
+			for (int32 Idx = 0; Idx < UseSamples; Idx++)
+			{
+				float Time = Duration * ((float)Idx / DivNum);
+				PolyPath.Path->Add(Spline->GetLocationAtTime(Time, SamplingOptions.CoordinateSpace, bUseConstantVelocity));
+			}
+		}
+	}
+}
+
+TArray<FVector> UGeometryScriptLibrary_PolyPathFunctions::Conv_GeometryScriptPolyPathToArray(FGeometryScriptPolyPath PolyPath)
+{
+	TArray<FVector> PathVertices;
+	ConvertPolyPathToArray(PolyPath, PathVertices);
+	return PathVertices;
+}
+
+TArray<FVector2D> UGeometryScriptLibrary_PolyPathFunctions::Conv_GeometryScriptPolyPathToArrayOfVector2D(FGeometryScriptPolyPath PolyPath)
+{
+	TArray<FVector2D> PathVertices;
+	ConvertPolyPathToArrayOfVector2D(PolyPath, PathVertices);
+	return PathVertices;
+}
+
+FGeometryScriptPolyPath UGeometryScriptLibrary_PolyPathFunctions::Conv_ArrayToGeometryScriptPolyPath(const TArray<FVector>& PathVertices)
+{
+	FGeometryScriptPolyPath PolyPath;
+	ConvertArrayToPolyPath(PathVertices, PolyPath);
+	return PolyPath;
+}
+
+FGeometryScriptPolyPath UGeometryScriptLibrary_PolyPathFunctions::Conv_ArrayOfVector2DToGeometryScriptPolyPath(const TArray<FVector2D>& PathVertices)
+{
+	FGeometryScriptPolyPath PolyPath;
+	ConvertArrayOfVector2DToPolyPath(PathVertices, PolyPath);
+	return PolyPath;
+}
+
+
+#undef LOCTEXT_NAMESPACE
\ No newline at end of file
diff --git a/Engine/Plugins/Experimental/GeometryScripting/Source/GeometryScriptingCore/Public/GeometryScript/GeometryScriptTypes.h b/Engine/Plugins/Experimental/GeometryScripting/Source/GeometryScriptingCore/Public/GeometryScript/GeometryScriptTypes.h
--- a/Engine/Plugins/Experimental/GeometryScripting/Source/GeometryScriptingCore/Public/GeometryScript/GeometryScriptTypes.h
+++ b/Engine/Plugins/Experimental/GeometryScripting/Source/GeometryScriptingCore/Public/GeometryScript/GeometryScriptTypes.h
@@ -44,6 +44,15 @@ enum class EGeometryScriptLODType : uint8
 };
 
 
+UENUM(BlueprintType)
+enum class EGeometryScriptAxis : uint8
+{
+	X = 0,
+	Y = 1,
+	Z = 2
+};
+
+
 USTRUCT(BlueprintType)
 struct GEOMETRYSCRIPTINGCORE_API FGeometryScriptMeshReadLOD
 {
@@ -296,6 +305,25 @@ public:
 };
 
 
+USTRUCT(BlueprintType, meta = (DisplayName = "PolyPath"))
+struct GEOMETRYSCRIPTINGCORE_API FGeometryScriptPolyPath
+{
+	GENERATED_BODY()
+public:
+	TSharedPtr<TArray<FVector>> Path;
+
+	UPROPERTY(EditAnywhere, Category = "Options")
+	bool bClosedLoop = false;
+
+	void Reset()
+	{
+		if (!Path.IsValid())
+		{
+			Path = MakeShared<TArray<FVector>>();
+		}
+		Path->Reset();
+	}
+};
 
 
 //
diff --git a/Engine/Plugins/Experimental/GeometryScripting/Source/GeometryScriptingCore/Public/GeometryScript/PolyPathFunctions.h b/Engine/Plugins/Experimental/GeometryScripting/Source/GeometryScriptingCore/Public/GeometryScript/PolyPathFunctions.h
new file mode 100644
--- /dev/null
+++ b/Engine/Plugins/Experimental/GeometryScripting/Source/GeometryScriptingCore/Public/GeometryScript/PolyPathFunctions.h
@@ -0,0 +1,95 @@
+// Copyright Epic Games, Inc. All Rights Reserved.
+
+#pragma once
+
+#include "CoreMinimal.h"
+#include "Kismet/BlueprintFunctionLibrary.h"
+#include "GeometryScript/GeometryScriptTypes.h"
+#include "Components/SplineComponent.h"
+#include "PolyPathFunctions.generated.h"
+
+
+UENUM(BlueprintType)
+enum class EGeometryScriptSampleSpacing : uint8
+{
+	UniformDistance,
+	UniformTime,
+	ErrorTolerance
+};
+
+
+USTRUCT(BlueprintType)
+struct GEOMETRYSCRIPTINGCORE_API FGeometryScriptSplineSamplingOptions
+{
+	GENERATED_BODY()
+public:
+	UPROPERTY(BlueprintReadWrite, Category = Options, meta = (EditCondition = "SampleSpacing != EGeometryScriptSampleSpacing::ErrorTolerance"))
+	int32 NumSamples = 10;
+
+	UPROPERTY(BlueprintReadWrite, Category = Options, meta = (EditCondition = "SampleSpacing == EGeometryScriptSampleSpacing::ErrorTolerance"))
+	float ErrorTolerance = 1;
+
+	UPROPERTY(BlueprintReadWrite, Category = Options)
+	EGeometryScriptSampleSpacing SampleSpacing = EGeometryScriptSampleSpacing::UniformDistance;
+
+	UPROPERTY(BlueprintReadWrite, Category = Options)
+	TEnumAsByte<ESplineCoordinateSpace::Type> CoordinateSpace = ESplineCoordinateSpace::Type::Local;
+};
+
+
+UCLASS(meta = (ScriptName = "GeometryScript_PolyPath"))
+class GEOMETRYSCRIPTINGCORE_API UGeometryScriptLibrary_PolyPathFunctions : public UBlueprintFunctionLibrary
+{
+	GENERATED_BODY()
+public:
+
+	UFUNCTION(BlueprintPure, Category = "GeometryScript|PolyPath", meta=(ScriptMethod))
+	static int GetPolyPathNumVertices(FGeometryScriptPolyPath PolyPath);
+
+	UFUNCTION(BlueprintPure, Category = "GeometryScript|PolyPath", meta=(ScriptMethod))
+	static int GetPolyPathLastIndex(FGeometryScriptPolyPath PolyPath);
+
+	UFUNCTION(BlueprintPure, Category = "GeometryScript|PolyPath", meta=(ScriptMethod))
+	static FVector GetPolyPathVertex(FGeometryScriptPolyPath PolyPath, int Index, bool& bIsValidIndex);
+
+	UFUNCTION(BlueprintPure, Category = "GeometryScript|PolyPath", meta = (ScriptMethod))
+	static FVector GetPolyPathTangent(FGeometryScriptPolyPath PolyPath, int Index, bool& bIsValidIndex);
+
+	UFUNCTION(BlueprintCallable, Category = "GeometryScript|PolyPath", meta = (ScriptMethod))
+	static double GetPolyPathArcLength(FGeometryScriptPolyPath PolyPath);
+
+	/** Find the index of the vertex closest to a given point.  Returns -1 if path has no vertices. */
+	UFUNCTION(BlueprintCallable, Category = "GeometryScript|PolyPath", meta = (ScriptMethod))
+	static int32 GetNearestVertexIndex(FGeometryScriptPolyPath PolyPath, FVector Point);
+
+	/** Flatten to 2D by dropping the given axis, and using the other two coordinates as the new X, Y coordinates.  Returns the modified path for convenience. */
+	UFUNCTION(BlueprintCallable, Category = "GeometryScript|PolyPath", meta = (ScriptMethod, DisplayName = "Flatten To 2D On Axis"))
+	static UPARAM(DisplayName = "Target Poly Path") FGeometryScriptPolyPath FlattenTo2DOnAxis(FGeometryScriptPolyPath TargetPolyPath, EGeometryScriptAxis DropAxis = EGeometryScriptAxis::Z);
+
+	UFUNCTION(BlueprintCallable, Category = "GeometryScript|PolyPath")
+	static void ConvertSplineToPolyPath(const USplineComponent* Spline, FGeometryScriptPolyPath& PolyPath, FGeometryScriptSplineSamplingOptions SamplingOptions);
+
+	UFUNCTION(BlueprintCallable, Category = "GeometryScript|PolyPath", meta=(ScriptMethod))
+	static void ConvertPolyPathToArray(FGeometryScriptPolyPath PolyPath, TArray<FVector>& VertexArray);
+
+	UFUNCTION(BlueprintCallable, Category = "GeometryScript|PolyPath")
+	static void ConvertArrayToPolyPath(const TArray<FVector>& VertexArray, FGeometryScriptPolyPath& PolyPath);
+
+	UFUNCTION(BlueprintCallable, Category = "GeometryScript|PolyPath", meta = (ScriptMethod))
+	static void ConvertPolyPathToArrayOfVector2D(FGeometryScriptPolyPath PolyPath, TArray<FVector2D>& VertexArray);
+
+	UFUNCTION(BlueprintCallable, Category = "GeometryScript|PolyPath")
+	static void ConvertArrayOfVector2DToPolyPath(const TArray<FVector2D>& VertexArray, FGeometryScriptPolyPath& PolyPath);
+
+	UFUNCTION(BlueprintPure, meta = (DisplayName = "PolyPath To Array Of Vector", CompactNodeTitle = "->", BlueprintAutocast), Category = "GeometryScript|PolyPath")
+	static TArray<FVector> Conv_GeometryScriptPolyPathToArray(FGeometryScriptPolyPath PolyPath);
+
+	UFUNCTION(BlueprintPure, meta = (DisplayName = "PolyPath To Array Of Vector2D", CompactNodeTitle = "->", BlueprintAutocast), Category = "GeometryScript|PolyPath")
+	static TArray<FVector2D> Conv_GeometryScriptPolyPathToArrayOfVector2D(FGeometryScriptPolyPath PolyPath);
+
+	UFUNCTION(BlueprintPure, meta = (DisplayName = "Array Of Vector To PolyPath", CompactNodeTitle = "->", BlueprintAutocast), Category = "GeometryScript|PolyPath")
+	static FGeometryScriptPolyPath Conv_ArrayToGeometryScriptPolyPath(const TArray<FVector>& PathVertices);
+
+	UFUNCTION(BlueprintPure, meta = (DisplayName = "Array Of Vector2D To PolyPath", CompactNodeTitle = "->", BlueprintAutocast), Category = "GeometryScript|PolyPath")
+	static FGeometryScriptPolyPath Conv_ArrayOfVector2DToGeometryScriptPolyPath(const TArray<FVector2D>& PathVertices);
+};
diff --git a/Engine/Source/Runtime/GeometryCore/Public/Curve/CurveUtil.h b/Engine/Source/Runtime/GeometryCore/Public/Curve/CurveUtil.h
--- a/Engine/Source/Runtime/GeometryCore/Public/Curve/CurveUtil.h
+++ b/Engine/Source/Runtime/GeometryCore/Public/Curve/CurveUtil.h
@@ -48,7 +48,7 @@ using namespace UE::Math;
 		{
 			Sum += Distance(Vertices[i], Vertices[i-1]);
 		}
-		if (bLoop)
+		if (bLoop && NV > 1)
 		{
 			Sum += Distance(Vertices[NV-1], Vertices[0]);
 		}
