GeometryScript: allow CopyMeshToStaticMesh function to write to meshes in /Engine/Transient, as this is where temporary StaticMesh Assets created in the Transient package will be located, and such temporary assets are used by various Engine features like datasmith import processing

diff --git a/Engine/Plugins/Experimental/GeometryScripting/Source/GeometryScriptingCore/Private/MeshAssetFunctions.cpp b/Engine/Plugins/Experimental/GeometryScripting/Source/GeometryScriptingCore/Private/MeshAssetFunctions.cpp
--- a/Engine/Plugins/Experimental/GeometryScripting/Source/GeometryScriptingCore/Private/MeshAssetFunctions.cpp
+++ b/Engine/Plugins/Experimental/GeometryScripting/Source/GeometryScriptingCore/Private/MeshAssetFunctions.cpp
@@ -153,8 +153,10 @@ UDynamicMesh*  UGeometryScriptLibrary_StaticMeshFunctions::CopyMeshToStaticMesh(
 		return FromDynamicMesh;
 	}
 
-	// don't allow build-in engine assets to be modified
-	if (ToStaticMeshAsset->GetPathName().StartsWith(TEXT("/Engine/")))
+	// Don't allow built-in engine assets to be modified. However we do allow assets in /Engine/Transient/ to be updated because
+	// this is a location that temporary assets in the Transient package will be created in, and in some cases we want to use
+	// script functions on such asset (Datasmith does this for example)
+	if ( ToStaticMeshAsset->GetPathName().StartsWith(TEXT("/Engine/")) && ToStaticMeshAsset->GetPathName().StartsWith(TEXT("/Engine/Transient")) == false )
 	{
 		UE::Geometry::AppendError(Debug, EGeometryScriptErrorType::InvalidInputs, LOCTEXT("CopyMeshToAsset_EngineAsset", "CopyMeshToStaticMesh: Cannot modify built-in Engine asset"));
 		return FromDynamicMesh;
