GeometryScript: handle cases where we are running game w/ editor (eg -game, Standalone Game in PIE menu, etc) in collision generation nodes. In this case WITH_EDITOR define is true but GEditor is nullptr.

diff --git a/Engine/Plugins/Experimental/GeometryScripting/Source/GeometryScriptingCore/Private/CollisionFunctions.cpp b/Engine/Plugins/Experimental/GeometryScripting/Source/GeometryScriptingCore/Private/CollisionFunctions.cpp
--- a/Engine/Plugins/Experimental/GeometryScripting/Source/GeometryScriptingCore/Private/CollisionFunctions.cpp
+++ b/Engine/Plugins/Experimental/GeometryScripting/Source/GeometryScriptingCore/Private/CollisionFunctions.cpp
@@ -134,7 +134,7 @@ void ComputeCollisionFromMesh(
 static void SetStaticMeshSimpleCollision(UStaticMesh* StaticMeshAsset, const FKAggregateGeom& NewSimpleCollision, bool bEmitTransaction)
 {
 #if WITH_EDITOR
-	if (bEmitTransaction)
+	if (bEmitTransaction && GEditor)
 	{
 		GEditor->BeginTransaction(LOCTEXT("UpdateStaticMesh", "Set Simple Collision"));
 
@@ -193,7 +193,7 @@ static void SetStaticMeshSimpleCollision(UStaticMesh* StaticMeshAsset, const FKA
 	}
 
 #if WITH_EDITOR
-	if (bEmitTransaction)
+	if (bEmitTransaction && GEditor)
 	{
 		GEditor->EndTransaction();
 	}
@@ -292,7 +292,7 @@ UDynamicMesh* UGeometryScriptLibrary_CollisionFunctions::SetDynamicMeshCollision
 	});
 
 #if WITH_EDITOR
-	if (Options.bEmitTransaction)
+	if (Options.bEmitTransaction && GEditor)
 	{
 		GEditor->BeginTransaction(LOCTEXT("UpdateDynamicMesh", "Set Simple Collision"));
 
@@ -317,7 +317,7 @@ UDynamicMesh* UGeometryScriptLibrary_CollisionFunctions::SetDynamicMeshCollision
 	// do we need to do a post edit change here??
 
 #if WITH_EDITOR
-	if (Options.bEmitTransaction)
+	if (Options.bEmitTransaction && GEditor)
 	{
 		GEditor->EndTransaction();
 	}
@@ -339,7 +339,7 @@ void UGeometryScriptLibrary_CollisionFunctions::ResetDynamicMeshCollision(
 	}
 
 #if WITH_EDITOR
-	if (bEmitTransaction)
+	if (bEmitTransaction && GEditor)
 	{
 		GEditor->BeginTransaction(LOCTEXT("ResetDynamicMeshCollisionTransaction", "Clear Simple Collision"));
 		DynamicMeshComponent->Modify();
@@ -364,7 +364,7 @@ void UGeometryScriptLibrary_CollisionFunctions::ResetDynamicMeshCollision(
 	// do we need to do a post edit change here??
 
 #if WITH_EDITOR
-	if (bEmitTransaction)
+	if (bEmitTransaction && GEditor)
 	{
 		GEditor->EndTransaction();
 	}
diff --git a/Engine/Plugins/Experimental/GeometryScripting/Source/GeometryScriptingCore/Private/MeshAssetFunctions.cpp b/Engine/Plugins/Experimental/GeometryScripting/Source/GeometryScriptingCore/Private/MeshAssetFunctions.cpp
--- a/Engine/Plugins/Experimental/GeometryScripting/Source/GeometryScriptingCore/Private/MeshAssetFunctions.cpp
+++ b/Engine/Plugins/Experimental/GeometryScripting/Source/GeometryScriptingCore/Private/MeshAssetFunctions.cpp
@@ -165,7 +165,7 @@ UDynamicMesh*  UGeometryScriptLibrary_StaticMeshFunctions::CopyMeshToStaticMesh(
 	// flush any pending rendering commands, which might want to touch this StaticMesh while we are rebuilding it
 	FlushRenderingCommands();
 
-	if (Options.bEmitTransaction)
+	if (Options.bEmitTransaction && GEditor)
 	{
 		GEditor->BeginTransaction(LOCTEXT("UpdateStaticMesh", "Update Static Mesh"));
 	}
@@ -256,7 +256,7 @@ UDynamicMesh*  UGeometryScriptLibrary_StaticMeshFunctions::CopyMeshToStaticMesh(
 		ToStaticMeshAsset->PostEditChange();
 	}
 
-	if (Options.bEmitTransaction)
+	if (Options.bEmitTransaction && GEditor)
 	{
 		GEditor->EndTransaction();
 	}
