GeometryScript: add support for reading StaticMesh LODResources mesh (render mesh) in CopyMeshFromStaticMesh

diff --git a/Engine/Plugins/Experimental/GeometryScripting/Source/GeometryScriptingCore/GeometryScriptingCore.Build.cs b/Engine/Plugins/Experimental/GeometryScripting/Source/GeometryScriptingCore/GeometryScriptingCore.Build.cs
--- a/Engine/Plugins/Experimental/GeometryScripting/Source/GeometryScriptingCore/GeometryScriptingCore.Build.cs
+++ b/Engine/Plugins/Experimental/GeometryScripting/Source/GeometryScriptingCore/GeometryScriptingCore.Build.cs
@@ -44,6 +44,7 @@ public class GeometryScriptingCore : ModuleRules
 				"StaticMeshDescription",
 				"SkeletalMeshDescription",
 				"MeshConversion",
+				"MeshConversionEngineTypes",
 				"GeometryAlgorithms",
 				"ModelingOperators",
 				"ModelingComponents"
diff --git a/Engine/Plugins/Experimental/GeometryScripting/Source/GeometryScriptingCore/Private/MeshAssetFunctions.cpp b/Engine/Plugins/Experimental/GeometryScripting/Source/GeometryScriptingCore/Private/MeshAssetFunctions.cpp
--- a/Engine/Plugins/Experimental/GeometryScripting/Source/GeometryScriptingCore/Private/MeshAssetFunctions.cpp
+++ b/Engine/Plugins/Experimental/GeometryScripting/Source/GeometryScriptingCore/Private/MeshAssetFunctions.cpp
@@ -18,6 +18,7 @@
 
 #include "MeshDescriptionToDynamicMesh.h"
 #include "DynamicMeshToMeshDescription.h"
+#include "StaticMeshLODResourcesToDynamicMesh.h"
 #include "AssetUtils/StaticMeshMaterialUtil.h"
 
 
@@ -33,36 +34,32 @@ using namespace UE::Geometry;
 
 
 
-UDynamicMesh*  UGeometryScriptLibrary_StaticMeshFunctions::CopyMeshFromStaticMesh(
+static UDynamicMesh* CopyMeshFromStaticMesh_SourceData(	
 	UStaticMesh* FromStaticMeshAsset, 
 	UDynamicMesh* ToDynamicMesh, 
 	FGeometryScriptCopyMeshFromAssetOptions AssetOptions,
 	FGeometryScriptMeshReadLOD RequestedLOD,
 	TEnumAsByte<EGeometryScriptOutcomePins>& Outcome,
-	UGeometryScriptDebug* Debug)
-{
-	Outcome = EGeometryScriptOutcomePins::Failure;
+	UGeometryScriptDebug* Debug
 
-	if (FromStaticMeshAsset == nullptr)
-	{
-		UE::Geometry::AppendError(Debug, EGeometryScriptErrorType::InvalidInputs, LOCTEXT("CopyMeshFromAsset_InvalidInput1", "CopyMeshFromStaticMesh: FromStaticMeshAsset is Null"));
-		return ToDynamicMesh;
-	}
-	if (ToDynamicMesh == nullptr)
-	{
-		UE::Geometry::AppendError(Debug, EGeometryScriptErrorType::InvalidInputs, LOCTEXT("CopyMeshFromAsset_InvalidInput2", "CopyMeshFromStaticMesh: ToDynamicMesh is Null"));
-		return ToDynamicMesh;
-	}
+)
+{
 	if (RequestedLOD.LODType != EGeometryScriptLODType::MaxAvailable && RequestedLOD.LODType != EGeometryScriptLODType::SourceModel)
 	{
-		UE::Geometry::AppendError(Debug, EGeometryScriptErrorType::InvalidInputs, LOCTEXT("CopyMeshFromAsset_LODNotAvailable", "CopyMeshFromStaticMesh: Requested LOD is not available"));
+		UE::Geometry::AppendError(Debug, EGeometryScriptErrorType::InvalidInputs, LOCTEXT("CopyMeshFromStaticMesh_LODNotAvailable", "CopyMeshFromStaticMesh: Requested LOD Type is not available"));
 		return ToDynamicMesh;
 	}
 
 #if WITH_EDITOR
-	int32 UseLODIndex = FMath::Clamp(RequestedLOD.LODIndex, 0, FromStaticMeshAsset->GetNumLODs() - 1);
+	int32 UseLODIndex = FMath::Clamp(RequestedLOD.LODIndex, 0, FromStaticMeshAsset->GetNumSourceModels() - 1);
 
 	const FMeshDescription* SourceMesh = FromStaticMeshAsset->GetMeshDescription(UseLODIndex);
+	if (SourceMesh == nullptr)
+	{
+		UE::Geometry::AppendError(Debug, EGeometryScriptErrorType::InvalidInputs, LOCTEXT("CopyMeshFromStaticMesh_SourceLODIsNull", "CopyMeshFromStaticMesh: Requested SourceModel LOD is null, only RenderData Mesh is available"));
+		return ToDynamicMesh;
+	}
+
 	const FStaticMeshSourceModel& SourceModel = FromStaticMeshAsset->GetSourceModel(UseLODIndex);
 	const FMeshBuildSettings& BuildSettings = SourceModel.BuildSettings;
 
@@ -106,11 +103,89 @@ UDynamicMesh*  UGeometryScriptLibrary_StaticMeshFunctions::CopyMeshFromStaticMes
 	ToDynamicMesh->SetMesh(MoveTemp(NewMesh));
 
 	Outcome = EGeometryScriptOutcomePins::Success;
+#else
+	UE::Geometry::AppendError(Debug, EGeometryScriptErrorType::InvalidInputs, LOCTEXT("CopyMeshFromAsset_EditorOnly", "CopyMeshFromStaticMesh: Source Models are not available at Runtime"));
+#endif
+
+	return ToDynamicMesh;
+}
+
+
+
+static UDynamicMesh* CopyMeshFromStaticMesh_RenderData(	
+	UStaticMesh* FromStaticMeshAsset, 
+	UDynamicMesh* ToDynamicMesh, 
+	FGeometryScriptCopyMeshFromAssetOptions AssetOptions,
+	FGeometryScriptMeshReadLOD RequestedLOD,
+	TEnumAsByte<EGeometryScriptOutcomePins>& Outcome,
+	UGeometryScriptDebug* Debug
+
+)
+{
+	if (RequestedLOD.LODType != EGeometryScriptLODType::RenderData)
+	{
+		UE::Geometry::AppendError(Debug, EGeometryScriptErrorType::InvalidInputs, LOCTEXT("CopyMeshFromStaticMesh_LODNotAvailable", "CopyMeshFromStaticMesh: Requested LOD Type is not available"));
+		return ToDynamicMesh;
+	}
+
+	int32 UseLODIndex = FMath::Clamp(RequestedLOD.LODIndex, 0, FromStaticMeshAsset->GetNumLODs() - 1);
+
+
+	const FStaticMeshLODResources* LODResources = &FromStaticMeshAsset->GetRenderData()->LODResources[UseLODIndex];
+
+	FStaticMeshLODResourcesToDynamicMesh::ConversionOptions ConvertOptions;
+#if WITH_EDITOR
+	// respect BuildScale build setting
+	const FMeshBuildSettings& LODBuildSettings = FromStaticMeshAsset->GetSourceModel(UseLODIndex).BuildSettings;
+	ConvertOptions.BuildScale = (FVector3d)LODBuildSettings.BuildScale3D;
+#endif
+
+	FDynamicMesh3 NewMesh;
+	FStaticMeshLODResourcesToDynamicMesh Converter;
+	Converter.Convert(LODResources, ConvertOptions, NewMesh);
+
+	ToDynamicMesh->SetMesh(MoveTemp(NewMesh));
+	Outcome = EGeometryScriptOutcomePins::Success;
+	return ToDynamicMesh;
+}
 
+
+
+UDynamicMesh*  UGeometryScriptLibrary_StaticMeshFunctions::CopyMeshFromStaticMesh(
+	UStaticMesh* FromStaticMeshAsset, 
+	UDynamicMesh* ToDynamicMesh, 
+	FGeometryScriptCopyMeshFromAssetOptions AssetOptions,
+	FGeometryScriptMeshReadLOD RequestedLOD,
+	TEnumAsByte<EGeometryScriptOutcomePins>& Outcome,
+	UGeometryScriptDebug* Debug)
+{
+	Outcome = EGeometryScriptOutcomePins::Failure;
+
+	if (FromStaticMeshAsset == nullptr)
+	{
+		UE::Geometry::AppendError(Debug, EGeometryScriptErrorType::InvalidInputs, LOCTEXT("CopyMeshFromAsset_InvalidInput1", "CopyMeshFromStaticMesh: FromStaticMeshAsset is Null"));
+		return ToDynamicMesh;
+	}
+	if (ToDynamicMesh == nullptr)
+	{
+		UE::Geometry::AppendError(Debug, EGeometryScriptErrorType::InvalidInputs, LOCTEXT("CopyMeshFromAsset_InvalidInput2", "CopyMeshFromStaticMesh: ToDynamicMesh is Null"));
+		return ToDynamicMesh;
+	}
+
+#if WITH_EDITOR
+	if (RequestedLOD.LODType == EGeometryScriptLODType::RenderData)
+	{
+		return CopyMeshFromStaticMesh_RenderData(FromStaticMeshAsset, ToDynamicMesh, AssetOptions, RequestedLOD, Outcome, Debug);
+	}
+	else
+	{
+		return CopyMeshFromStaticMesh_SourceData(FromStaticMeshAsset, ToDynamicMesh, AssetOptions, RequestedLOD, Outcome, Debug);
+	}
 #else
-	UE::Geometry::AppendError(Debug, EGeometryScriptErrorType::InvalidInputs, LOCTEXT("CopyMeshFromAsset_EditorOnly", "CopyMeshFromStaticMesh: Not currently supported at Runtime"));
+	return CopyMeshFromStaticMesh_RenderData(FromStaticMeshAsset, ToDynamicMesh, AssetOptions, RequestedLOD, Outcome, Debug);	
 #endif
 
+
 	return ToDynamicMesh;
 }
 
