GeometryScript: add new ApplyDisplaceFromPerVertexVectors to displace mesh from set of per-vertex vectors

diff --git a/Engine/Plugins/Experimental/GeometryScripting/Source/GeometryScriptingCore/Private/MeshDeformFunctions.cpp b/Engine/Plugins/Experimental/GeometryScripting/Source/GeometryScriptingCore/Private/MeshDeformFunctions.cpp
--- a/Engine/Plugins/Experimental/GeometryScripting/Source/GeometryScriptingCore/Private/MeshDeformFunctions.cpp
+++ b/Engine/Plugins/Experimental/GeometryScripting/Source/GeometryScriptingCore/Private/MeshDeformFunctions.cpp
@@ -422,10 +422,6 @@ UDynamicMesh* UGeometryScriptLibrary_MeshDeformFunctions::ApplyDisplaceFromTextu
 				UV.Y *= VHeight;
 
 				double Offset = DisplaceField.BilinearSampleClamped(UV);
-				//if (AdjustmentCurve)
-				//{
-				//	Offset = AdjustmentCurve->Eval(Offset);
-				//}
 				Offset -= Options.Center;
 
 				FVector3f Normal = Normals->GetElement(NormalTri[j]);
@@ -451,4 +447,46 @@ UDynamicMesh* UGeometryScriptLibrary_MeshDeformFunctions::ApplyDisplaceFromTextu
 }
 
 
+
+
+
+
+UDynamicMesh* UGeometryScriptLibrary_MeshDeformFunctions::ApplyDisplaceFromPerVertexVectors(
+	UDynamicMesh* TargetMesh,
+	const FGeometryScriptVectorList& VectorList, 
+	float Magnitude,
+	UGeometryScriptDebug* Debug)
+{
+	if (TargetMesh == nullptr)
+	{
+		UE::Geometry::AppendError(Debug, EGeometryScriptErrorType::InvalidInputs, LOCTEXT("ApplyDisplaceFromPerVertexVectors_InvalidInput", "ApplyDisplaceFromTextureMap: TargetMesh is Null"));
+		return TargetMesh;
+	}
+
+	const TArray<FVector>& Vectors = *VectorList.List;
+
+	TargetMesh->EditMesh([&](FDynamicMesh3& EditMesh)
+	{
+		if (Vectors.Num() < EditMesh.MaxVertexID())
+		{
+			UE::Geometry::AppendError(Debug, EGeometryScriptErrorType::InvalidInputs, LOCTEXT("ApplyDisplaceFromPerVertexVectors_InvalidVectorLength", "ApplyDisplaceFromTextureMap: VectorList Length is less than TargetMesh MaxVertexID"));
+			return;
+		}
+
+		for (int32 vid : EditMesh.VertexIndicesItr())
+		{
+			FVector Position = EditMesh.GetVertex(vid);
+			FVector Vector = Magnitude * Vectors[vid];
+			if (Vector.Length() < HALF_WORLD_MAX)		// this is a bit arbitrary but should avoid disasters
+			{
+				EditMesh.SetVertex(vid, Position + Vector);
+			}
+		}
+
+	}, EDynamicMeshChangeType::GeneralEdit, EDynamicMeshAttributeChangeFlags::Unknown, false);
+
+	return TargetMesh;
+}
+
+
 #undef LOCTEXT_NAMESPACE
\ No newline at end of file
diff --git a/Engine/Plugins/Experimental/GeometryScripting/Source/GeometryScriptingCore/Public/GeometryScript/MeshDeformFunctions.h b/Engine/Plugins/Experimental/GeometryScripting/Source/GeometryScriptingCore/Public/GeometryScript/MeshDeformFunctions.h
--- a/Engine/Plugins/Experimental/GeometryScripting/Source/GeometryScriptingCore/Public/GeometryScript/MeshDeformFunctions.h
+++ b/Engine/Plugins/Experimental/GeometryScripting/Source/GeometryScriptingCore/Public/GeometryScript/MeshDeformFunctions.h
@@ -254,4 +254,12 @@ public:
 		UGeometryScriptDebug* Debug = nullptr);
 
 	
+	UFUNCTION(BlueprintCallable, Category = "GeometryScript|Deformations", meta=(ScriptMethod))
+	static UPARAM(DisplayName = "Target Mesh") UDynamicMesh* 
+	ApplyDisplaceFromPerVertexVectors(  
+		UDynamicMesh* TargetMesh,
+		const FGeometryScriptVectorList& VectorList, 
+		float Magnitude = 5.0,
+		UGeometryScriptDebug* Debug = nullptr);
+
 };
\ No newline at end of file
