GeometryScript: new functions for extracing per-vertex mesh Normals, UVs, and Vertex Colors

diff --git a/Engine/Plugins/Experimental/GeometryScripting/Source/GeometryScriptingCore/Private/MeshNormalsFunctions.cpp b/Engine/Plugins/Experimental/GeometryScripting/Source/GeometryScriptingCore/Private/MeshNormalsFunctions.cpp
--- a/Engine/Plugins/Experimental/GeometryScripting/Source/GeometryScriptingCore/Private/MeshNormalsFunctions.cpp
+++ b/Engine/Plugins/Experimental/GeometryScripting/Source/GeometryScriptingCore/Private/MeshNormalsFunctions.cpp
@@ -281,4 +281,78 @@ UDynamicMesh* UGeometryScriptLibrary_MeshNormalsFunctions::SetMeshTriangleNormal
 
 
 
+
+UDynamicMesh* UGeometryScriptLibrary_MeshNormalsFunctions::GetMeshPerVertexNormals(
+	UDynamicMesh* TargetMesh, 
+	FGeometryScriptVectorList& NormalList, 
+	bool& bIsValidNormalSet,
+	bool& bHasVertexIDGaps,
+	bool bAverageSplitVertexValues)
+{
+	NormalList.Reset();
+	TArray<FVector>& Normals = *NormalList.List;
+	bHasVertexIDGaps = false;
+	bIsValidNormalSet = false;
+	if (TargetMesh)
+	{
+		TargetMesh->ProcessMesh([&](const FDynamicMesh3& ReadMesh)
+		{
+			bHasVertexIDGaps = ! ReadMesh.IsCompactV();
+
+			if (ReadMesh.HasAttributes() && ReadMesh.Attributes()->NumNormalLayers() > 0 )
+			{
+				const FDynamicMeshNormalOverlay* NormalOverlay = ReadMesh.Attributes()->PrimaryNormals();
+
+				if (bAverageSplitVertexValues)
+				{
+					Normals.Init(FVector::Zero(), ReadMesh.MaxVertexID());
+					for (int32 tid : ReadMesh.TriangleIndicesItr())
+					{
+						if (NormalOverlay->IsSetTriangle(tid))
+						{
+							FIndex3i TriV = ReadMesh.GetTriangle(tid);
+							FVector3f A, B, C;
+							NormalOverlay->GetTriElements(tid, A, B, C);
+							Normals[TriV.A] += (FVector)A;
+							Normals[TriV.B] += (FVector)B;
+							Normals[TriV.C] += (FVector)C;
+						}
+					}
+
+					for (int32 k = 0; k < Normals.Num(); ++k)
+					{
+						if (Normals[k].SquaredLength() > 0)
+						{
+							Normalize(Normals[k]);
+						}
+					}
+				}
+				else
+				{
+					Normals.Init(FVector::UnitZ(), ReadMesh.MaxVertexID());
+					for (int32 tid : ReadMesh.TriangleIndicesItr())
+					{
+						if (NormalOverlay->IsSetTriangle(tid))
+						{
+							FIndex3i TriV = ReadMesh.GetTriangle(tid);
+							FVector3f A, B, C;
+							NormalOverlay->GetTriElements(tid, A, B, C);
+							Normals[TriV.A] = (FVector)A;
+							Normals[TriV.B] = (FVector)B;
+							Normals[TriV.C] = (FVector)C;
+						}
+					}
+				}
+
+				bIsValidNormalSet = true;
+			}
+		});
+	}
+
+	return TargetMesh;
+}
+
+
+
+
 #undef LOCTEXT_NAMESPACE
diff --git a/Engine/Plugins/Experimental/GeometryScripting/Source/GeometryScriptingCore/Private/MeshUVFunctions.cpp b/Engine/Plugins/Experimental/GeometryScripting/Source/GeometryScriptingCore/Private/MeshUVFunctions.cpp
--- a/Engine/Plugins/Experimental/GeometryScripting/Source/GeometryScriptingCore/Private/MeshUVFunctions.cpp
+++ b/Engine/Plugins/Experimental/GeometryScripting/Source/GeometryScriptingCore/Private/MeshUVFunctions.cpp
@@ -714,4 +714,68 @@ UDynamicMesh* UGeometryScriptLibrary_MeshUVFunctions::AutoGenerateXAtlasMeshUVs(
 
 
 
+
+
+UDynamicMesh* UGeometryScriptLibrary_MeshUVFunctions::GetMeshPerVertexUVs(
+	UDynamicMesh* TargetMesh, 
+	int UVSetIndex,
+	FGeometryScriptUVList& UVList, 
+	bool& bIsValidUVSet,
+	bool& bHasVertexIDGaps,
+	bool& bHasSplitUVs,
+	UGeometryScriptDebug* Debug)
+{
+	UVList.Reset();
+	TArray<FVector2D>& UVs = *UVList.List;
+	bHasVertexIDGaps = false;
+	bIsValidUVSet = false;
+	bHasSplitUVs = false;
+	if (TargetMesh)
+	{
+		TargetMesh->ProcessMesh([&](const FDynamicMesh3& ReadMesh)
+		{
+			const FDynamicMeshUVOverlay* UVOverlay = (ReadMesh.HasAttributes() && UVSetIndex < ReadMesh.Attributes()->NumUVLayers()) ?
+				ReadMesh.Attributes()->GetUVLayer(UVSetIndex) : nullptr;
+			if (UVOverlay == nullptr)
+			{
+				UE::Geometry::AppendError(Debug, EGeometryScriptErrorType::InvalidInputs, LOCTEXT("GetMeshPerVertexUVs_InvalidUVSet", "GetMeshPerVertexUVs: UVSetIndex does not exist on TargetMesh"));
+				return;
+			}
+
+			bHasVertexIDGaps = ! ReadMesh.IsCompactV();
+
+			UVs.Init(FVector2D::Zero(), ReadMesh.MaxVertexID());
+			TArray<int32> ElemIndex;		// set to elementID of first element seen at each vertex, if we see a second element ID, it is a split vertex
+			ElemIndex.Init(-1, UVs.Num());
+
+			for (int32 tid : ReadMesh.TriangleIndicesItr())
+			{
+				if (UVOverlay->IsSetTriangle(tid))
+				{
+					FIndex3i TriV = ReadMesh.GetTriangle(tid);
+					FIndex3i TriE = UVOverlay->GetTriangle(tid);
+					for (int j = 0; j < 3; ++j)
+					{
+						if (ElemIndex[TriV[j]] == -1)
+						{
+							UVs[TriV[j]] = (FVector2D)UVOverlay->GetElement(TriE[j]);
+							ElemIndex[TriV[j]] = TriE[j];
+						}
+						else if (ElemIndex[TriV[j]] != TriE[j])
+						{
+							bHasSplitUVs = true;
+						}
+					}
+				}
+			}
+
+			bIsValidUVSet = true;
+		});
+	}
+
+	return TargetMesh;
+}
+
+
+
 #undef LOCTEXT_NAMESPACE
\ No newline at end of file
diff --git a/Engine/Plugins/Experimental/GeometryScripting/Source/GeometryScriptingCore/Private/MeshVertexColorFunctions.cpp b/Engine/Plugins/Experimental/GeometryScripting/Source/GeometryScriptingCore/Private/MeshVertexColorFunctions.cpp
--- a/Engine/Plugins/Experimental/GeometryScripting/Source/GeometryScriptingCore/Private/MeshVertexColorFunctions.cpp
+++ b/Engine/Plugins/Experimental/GeometryScripting/Source/GeometryScriptingCore/Private/MeshVertexColorFunctions.cpp
@@ -138,4 +138,88 @@ UDynamicMesh* UGeometryScriptLibrary_MeshVertexColorFunctions::SetMeshPerVertexC
 	return TargetMesh;
 }
 
+
+
+
+
+
+
+
+UDynamicMesh* UGeometryScriptLibrary_MeshVertexColorFunctions::GetMeshPerVertexColors(
+	UDynamicMesh* TargetMesh, 
+	FGeometryScriptColorList& ColorList, 
+	bool& bIsValidColorSet,
+	bool& bHasVertexIDGaps,
+	bool bBlendSplitVertexValues)
+{
+	ColorList.Reset();
+	TArray<FLinearColor>& Colors = *ColorList.List;
+	bHasVertexIDGaps = false;
+	bIsValidColorSet = false;
+	if (TargetMesh)
+	{
+		TargetMesh->ProcessMesh([&](const FDynamicMesh3& ReadMesh)
+		{
+			Colors.Init(FLinearColor::Black, ReadMesh.MaxVertexID());
+			bHasVertexIDGaps = ! ReadMesh.IsCompactV();
+
+			if (ReadMesh.HasAttributes() && ReadMesh.Attributes()->HasPrimaryColors() )
+			{
+				const FDynamicMeshColorOverlay* ColorOverlay = ReadMesh.Attributes()->PrimaryColors();
+
+				if (bBlendSplitVertexValues)
+				{
+					TArray<int32> ColorCounts;
+					ColorCounts.Init(0, ReadMesh.MaxVertexID());
+					for (int32 tid : ReadMesh.TriangleIndicesItr())
+					{
+						if (ColorOverlay->IsSetTriangle(tid))
+						{
+							FIndex3i TriV = ReadMesh.GetTriangle(tid);
+							FVector4f A, B, C;
+							ColorOverlay->GetTriElements(tid, A, B, C);
+							Colors[TriV.A] += ToLinearColor(A);
+							ColorCounts[TriV.A]++;
+							Colors[TriV.B] += ToLinearColor(B);
+							ColorCounts[TriV.B]++;
+							Colors[TriV.C] += ToLinearColor(B);
+							ColorCounts[TriV.C]++;
+						}
+					}
+
+					for (int32 k = 0; k < ColorCounts.Num(); ++k)
+					{
+						if (ColorCounts[k] > 1)
+						{
+							Colors[k] *= 1.0f / (float)ColorCounts[k];
+						}
+					}
+				}
+				else
+				{
+					for (int32 tid : ReadMesh.TriangleIndicesItr())
+					{
+						if (ColorOverlay->IsSetTriangle(tid))
+						{
+							FIndex3i TriV = ReadMesh.GetTriangle(tid);
+							FVector4f A, B, C;
+							ColorOverlay->GetTriElements(tid, A, B, C);
+							Colors[TriV.A] = ToLinearColor(A);
+							Colors[TriV.B] = ToLinearColor(B);
+							Colors[TriV.C] = ToLinearColor(B);
+						}
+					}
+				}
+				
+				bIsValidColorSet = true;
+			}
+		});
+	}
+
+	return TargetMesh;
+}
+
+
+
+
 #undef LOCTEXT_NAMESPACE
\ No newline at end of file
diff --git a/Engine/Plugins/Experimental/GeometryScripting/Source/GeometryScriptingCore/Public/GeometryScript/MeshNormalsFunctions.h b/Engine/Plugins/Experimental/GeometryScripting/Source/GeometryScriptingCore/Public/GeometryScript/MeshNormalsFunctions.h
--- a/Engine/Plugins/Experimental/GeometryScripting/Source/GeometryScriptingCore/Public/GeometryScript/MeshNormalsFunctions.h
+++ b/Engine/Plugins/Experimental/GeometryScripting/Source/GeometryScriptingCore/Public/GeometryScript/MeshNormalsFunctions.h
@@ -126,4 +126,24 @@ public:
 		bool& bIsValidTriangle, 
 		bool bDeferChangeNotifications = false );
 
+
+
+	/**
+	 * Get a list of single normal vectors for each mesh vertex in the TargetMesh, derived from the Normals Overlay.
+	 * The Normals Overlay may store multiple normals for a single vertex (ie split normals)
+	 * In such cases the normals can either be averaged, or the last normal seen will be used, depending on the bAverageSplitVertexValues parameter.
+	 * @param NormalList output normal list will be stored here. Size will be equal to the MaxVertexID of TargetMesh  (not the VertexCount!)
+	 * @param bIsValidNormalSet will be set to true if the Normal Overlay was valid
+	 * @param bHasVertexIDGaps will be set to true if some vertex indices in TargetMesh were invalid, ie MaxVertexID > VertexCount 
+	 * @param bAverageSplitVertexValues control how multiple normals at the same vertex should be interpreted
+	 */
+	UFUNCTION(BlueprintCallable, Category = "GeometryScript|Normals", meta=(ScriptMethod))
+	static UPARAM(DisplayName = "Target Mesh") UDynamicMesh* 
+	GetMeshPerVertexNormals( 
+		UDynamicMesh* TargetMesh, 
+		FGeometryScriptVectorList& NormalList, 
+		bool& bIsValidNormalSet,
+		bool& bHasVertexIDGaps,
+		bool bAverageSplitVertexValues = true);
+
 };
\ No newline at end of file
diff --git a/Engine/Plugins/Experimental/GeometryScripting/Source/GeometryScriptingCore/Public/GeometryScript/MeshUVFunctions.h b/Engine/Plugins/Experimental/GeometryScripting/Source/GeometryScriptingCore/Public/GeometryScript/MeshUVFunctions.h
--- a/Engine/Plugins/Experimental/GeometryScripting/Source/GeometryScriptingCore/Public/GeometryScript/MeshUVFunctions.h
+++ b/Engine/Plugins/Experimental/GeometryScripting/Source/GeometryScriptingCore/Public/GeometryScript/MeshUVFunctions.h
@@ -264,4 +264,26 @@ public:
 		FGeometryScriptXAtlasOptions Options,
 		UGeometryScriptDebug* Debug = nullptr );
 
+
+	/**
+	 * Get a list of single vertex UVs for each mesh vertex in the TargetMesh, derived from the specified UV Overlay.
+	 * The UV Overlay may store multiple UVs for a single vertex (along UV seams)
+	 * In such cases an arbitrary UV will be stored for that vertex, and bHasSplitUVs will be returned as true
+	 * @param UVSetIndex index of UV Set to read
+	 * @param UVList output UV list will be stored here. Size will be equal to the MaxVertexID of TargetMesh  (not the VertexCount!)
+	 * @param bIsValidUVSet will be set to true if the UV Overlay was valid
+	 * @param bHasVertexIDGaps will be set to true if some vertex indices in TargetMesh were invalid, ie MaxVertexID > VertexCount 
+	 * @param bHasSplitUVs will be set to true if there were split UVs in the UV overlay
+	 */
+	UFUNCTION(BlueprintCallable, Category = "GeometryScript|UVs", meta=(ScriptMethod))
+	static UPARAM(DisplayName = "Target Mesh") UDynamicMesh* 
+	GetMeshPerVertexUVs( 
+		UDynamicMesh* TargetMesh, 
+		int UVSetIndex,
+		FGeometryScriptUVList& UVList, 
+		bool& bIsValidUVSet,
+		bool& bHasVertexIDGaps,
+		bool& bHasSplitUVs,
+		UGeometryScriptDebug* Debug = nullptr);
+
 };
\ No newline at end of file
diff --git a/Engine/Plugins/Experimental/GeometryScripting/Source/GeometryScriptingCore/Public/GeometryScript/MeshVertexColorFunctions.h b/Engine/Plugins/Experimental/GeometryScripting/Source/GeometryScriptingCore/Public/GeometryScript/MeshVertexColorFunctions.h
--- a/Engine/Plugins/Experimental/GeometryScripting/Source/GeometryScriptingCore/Public/GeometryScript/MeshVertexColorFunctions.h
+++ b/Engine/Plugins/Experimental/GeometryScripting/Source/GeometryScriptingCore/Public/GeometryScript/MeshVertexColorFunctions.h
@@ -17,6 +17,11 @@ class GEOMETRYSCRIPTINGCORE_API UGeometryScriptLibrary_MeshVertexColorFunctions
 	GENERATED_BODY()
 public:
 
+	/**
+	 * Set all vertex colors (optionally specific channels) in the TargetMesh VertexColor Overlay to a constant value
+	 * @param Color the constant color to set
+	 * @param Flags specify which RGBA channels to set (default all channels)
+	 */
 	UFUNCTION(BlueprintCallable, Category = "GeometryScript|VertexColor", meta=(ScriptMethod))
 	static UPARAM(DisplayName = "Target Mesh") UDynamicMesh*
 	SetMeshConstantVertexColor(
@@ -26,6 +31,10 @@ public:
 		bool bClearExisting = false,
 		UGeometryScriptDebug* Debug = nullptr);
 
+	/**
+	 * Set all vertex colors in the TargetMesh VertexColor Overlay to the specified per-vertex colors
+	 * @param VertexColorList per-vertex colors. Size must be less than or equal to the MaxVertexID of TargetMesh  (ie gaps are supported)
+	 */
 	UFUNCTION(BlueprintCallable, Category = "GeometryScript|VertexColor", meta=(ScriptMethod))
 	static UPARAM(DisplayName = "Target Mesh") UDynamicMesh*
 	SetMeshPerVertexColors(
@@ -33,6 +42,22 @@ public:
 		FGeometryScriptColorList VertexColorList,
 		UGeometryScriptDebug* Debug = nullptr);
 
-
+	/**
+	 * Get a list of single vertex colors for each mesh vertex in the TargetMesh, derived from the VertexColor Overlay.
+	 * The VertexColor Overlay may store multiple colors for a single vertex (ie different colors for that vertex on different triangles)
+	 * In such cases the colors can either be averaged, or the last color seen will be used, depending on the bBlendSplitVertexValues parameter.
+	 * @param ColorList output color list will be stored here. Size will be equal to the MaxVertexID of TargetMesh  (not the VertexCount!)
+	 * @param bIsValidColorSet will be set to true if the VertexColor Overlay was valid
+	 * @param bHasVertexIDGaps will be set to true if some vertex indices in TargetMesh were invalid, ie MaxVertexID > VertexCount 
+	 * @param bBlendSplitVertexValues control how multiple colors at the same vertex should be interpreted
+	 */
+	UFUNCTION(BlueprintCallable, Category = "GeometryScript|VertexColor", meta=(ScriptMethod))
+	static UPARAM(DisplayName = "Target Mesh") UDynamicMesh* 
+	GetMeshPerVertexColors( 
+		UDynamicMesh* TargetMesh, 
+		FGeometryScriptColorList& ColorList, 
+		bool& bIsValidColorSet,
+		bool& bHasVertexIDGaps,
+		bool bBlendSplitVertexValues = true);
 };
 
