GeometryScript: add SampleTexture2DAtUVPositions function that returns texture colors (nearest or bilinear sampled) at UV positions.
GeometryProcessing: add TImageBuilder::NearestSampleUV() function to simplify nearest sampling

diff --git a/Engine/Plugins/Experimental/GeometryScripting/Source/GeometryScriptingCore/Private/MeshDeformFunctions.cpp b/Engine/Plugins/Experimental/GeometryScripting/Source/GeometryScriptingCore/Private/MeshDeformFunctions.cpp
--- a/Engine/Plugins/Experimental/GeometryScripting/Source/GeometryScriptingCore/Private/MeshDeformFunctions.cpp
+++ b/Engine/Plugins/Experimental/GeometryScripting/Source/GeometryScriptingCore/Private/MeshDeformFunctions.cpp
@@ -459,7 +459,7 @@ UDynamicMesh* UGeometryScriptLibrary_MeshDeformFunctions::ApplyDisplaceFromPerVe
 {
 	if (TargetMesh == nullptr)
 	{
-		UE::Geometry::AppendError(Debug, EGeometryScriptErrorType::InvalidInputs, LOCTEXT("ApplyDisplaceFromPerVertexVectors_InvalidInput", "ApplyDisplaceFromTextureMap: TargetMesh is Null"));
+		UE::Geometry::AppendError(Debug, EGeometryScriptErrorType::InvalidInputs, LOCTEXT("ApplyDisplaceFromPerVertexVectors_InvalidInput", "ApplyDisplaceFromPerVertexVectors: TargetMesh is Null"));
 		return TargetMesh;
 	}
 
@@ -469,7 +469,7 @@ UDynamicMesh* UGeometryScriptLibrary_MeshDeformFunctions::ApplyDisplaceFromPerVe
 	{
 		if (Vectors.Num() < EditMesh.MaxVertexID())
 		{
-			UE::Geometry::AppendError(Debug, EGeometryScriptErrorType::InvalidInputs, LOCTEXT("ApplyDisplaceFromPerVertexVectors_InvalidVectorLength", "ApplyDisplaceFromTextureMap: VectorList Length is less than TargetMesh MaxVertexID"));
+			UE::Geometry::AppendError(Debug, EGeometryScriptErrorType::InvalidInputs, LOCTEXT("ApplyDisplaceFromPerVertexVectors_InvalidVectorLength", "ApplyDisplaceFromPerVertexVectors: VectorList Length is less than TargetMesh MaxVertexID"));
 			return;
 		}
 
diff --git a/Engine/Plugins/Experimental/GeometryScripting/Source/GeometryScriptingCore/Private/MeshVertexColorFunctions.cpp b/Engine/Plugins/Experimental/GeometryScripting/Source/GeometryScriptingCore/Private/MeshVertexColorFunctions.cpp
--- a/Engine/Plugins/Experimental/GeometryScripting/Source/GeometryScriptingCore/Private/MeshVertexColorFunctions.cpp
+++ b/Engine/Plugins/Experimental/GeometryScripting/Source/GeometryScriptingCore/Private/MeshVertexColorFunctions.cpp
@@ -182,7 +182,7 @@ UDynamicMesh* UGeometryScriptLibrary_MeshVertexColorFunctions::GetMeshPerVertexC
 							ColorCounts[TriV.A]++;
 							Colors[TriV.B] += ToLinearColor(B);
 							ColorCounts[TriV.B]++;
-							Colors[TriV.C] += ToLinearColor(B);
+							Colors[TriV.C] += ToLinearColor(C);
 							ColorCounts[TriV.C]++;
 						}
 					}
@@ -206,7 +206,7 @@ UDynamicMesh* UGeometryScriptLibrary_MeshVertexColorFunctions::GetMeshPerVertexC
 							ColorOverlay->GetTriElements(tid, A, B, C);
 							Colors[TriV.A] = ToLinearColor(A);
 							Colors[TriV.B] = ToLinearColor(B);
-							Colors[TriV.C] = ToLinearColor(B);
+							Colors[TriV.C] = ToLinearColor(C);
 						}
 					}
 				}
diff --git a/Engine/Plugins/Experimental/GeometryScripting/Source/GeometryScriptingCore/Private/TextureMapFunctions.cpp b/Engine/Plugins/Experimental/GeometryScripting/Source/GeometryScriptingCore/Private/TextureMapFunctions.cpp
new file mode 100644
--- /dev/null
+++ b/Engine/Plugins/Experimental/GeometryScripting/Source/GeometryScriptingCore/Private/TextureMapFunctions.cpp
@@ -0,0 +1,66 @@
+// Copyright Epic Games, Inc. All Rights Reserved.
+
+#include "GeometryScript/TextureMapFunctions.h"
+
+#include "Async/ParallelFor.h"
+#include "AssetUtils/Texture2DUtil.h"
+#include "Spatial/SampledScalarField2.h"
+
+using namespace UE::Geometry;
+
+#define LOCTEXT_NAMESPACE "UGeometryScriptLibrary_TextureMapFunctions"
+
+void UGeometryScriptLibrary_TextureMapFunctions::SampleTexture2DAtUVPositions(
+	FGeometryScriptUVList UVList,
+	UTexture2D* TextureAsset,
+	FGeometryScriptSampleTextureOptions SampleOptions,
+	FGeometryScriptColorList& ColorList,
+	UGeometryScriptDebug* Debug)
+{
+	if (TextureAsset == nullptr)
+	{
+		UE::Geometry::AppendError(Debug, EGeometryScriptErrorType::InvalidInputs, LOCTEXT("SampleTexture2DAtUVPositions_InvalidInput2", "SampleTexture2DAtUVPositions: Texture is Null"));
+		return;
+	}
+
+	TImageBuilder<FVector4f> ImageData;
+	if (UE::AssetUtils::ReadTexture(TextureAsset, ImageData, false) == false)
+	{
+		UE::Geometry::AppendError(Debug, EGeometryScriptErrorType::OperationFailed, LOCTEXT("SampleTexture2DAtUVPositions_TexReadFailed", "SampleTexture2DAtUVPositions: Error reading source texture data"));
+		return;
+	}
+
+	const TArray<FVector2D>& UVs = *UVList.List;
+	int32 NumUVs = UVs.Num();
+
+	ColorList.Reset();
+	TArray<FLinearColor>& Colors = *ColorList.List;
+	Colors.SetNumUninitialized(NumUVs);
+
+	for (int32 k = 0; k < NumUVs; ++k)
+	{
+		FVector2d UV = UVs[k];
+
+		// Adjust UV value and tile it. 
+		UV = UV * SampleOptions.UVScale + SampleOptions.UVOffset;
+		if (SampleOptions.bWrap)
+		{
+			UV = UV - FVector2d(FMathd::Floor(UV.X), FMathd::Floor(UV.Y));
+		}
+		else
+		{
+			UV.X = FMathd::Clamp(UV.X, 0.0, 1.0);
+			UV.Y = FMathd::Clamp(UV.Y, 0.0, 1.0);
+		}
+
+		FVector4f InterpValue = (SampleOptions.SamplingMethod == EGeometryScriptPixelSamplingMethod::Bilinear) ?
+			ImageData.BilinearSampleUV<double>(UV, FVector4f::Zero()) : ImageData.NearestSampleUV<double>(UV, FVector4f::Zero());
+
+		Colors[k] = (FLinearColor)InterpValue;
+	}
+
+
+}
+
+
+#undef LOCTEXT_NAMESPACE
\ No newline at end of file
diff --git a/Engine/Plugins/Experimental/GeometryScripting/Source/GeometryScriptingCore/Public/GeometryScript/TextureMapFunctions.h b/Engine/Plugins/Experimental/GeometryScripting/Source/GeometryScriptingCore/Public/GeometryScript/TextureMapFunctions.h
new file mode 100644
--- /dev/null
+++ b/Engine/Plugins/Experimental/GeometryScripting/Source/GeometryScriptingCore/Public/GeometryScript/TextureMapFunctions.h
@@ -0,0 +1,61 @@
+// Copyright Epic Games, Inc. All Rights Reserved.
+
+#pragma once
+
+#include "CoreMinimal.h"
+#include "Kismet/BlueprintFunctionLibrary.h"
+#include "GeometryScript/GeometryScriptTypes.h"
+#include "TextureMapFunctions.generated.h"
+
+
+UENUM(BlueprintType)
+enum class EGeometryScriptPixelSamplingMethod : uint8
+{
+	Bilinear = 0,
+	Nearest = 1
+};
+
+
+
+
+USTRUCT(BlueprintType)
+struct GEOMETRYSCRIPTINGCORE_API FGeometryScriptSampleTextureOptions
+{
+	GENERATED_BODY()
+public:
+	UPROPERTY(BlueprintReadWrite, Category = Options)
+	EGeometryScriptPixelSamplingMethod SamplingMethod = EGeometryScriptPixelSamplingMethod::Bilinear;
+
+	UPROPERTY(BlueprintReadWrite, Category = Options)
+	bool bWrap = true;
+
+	UPROPERTY(BlueprintReadWrite, Category = Options)
+	FVector2D UVScale = FVector2D(1,1);
+
+	UPROPERTY(BlueprintReadWrite, Category = Options)
+	FVector2D UVOffset = FVector2D(0,0);
+};
+
+
+
+
+UCLASS(meta = (ScriptName = "GeometryScript_TextureUtils"))
+class GEOMETRYSCRIPTINGCORE_API UGeometryScriptLibrary_TextureMapFunctions : public UBlueprintFunctionLibrary
+{
+	GENERATED_BODY()
+public:
+
+	/**
+	 * Sample the the given TextureMap at the list of UV positions and return the color at each position in ColorList output
+	 */
+	UFUNCTION(BlueprintCallable, Category = "GeometryScript|TextureUtils")
+	static void 
+	SampleTexture2DAtUVPositions(  
+		FGeometryScriptUVList UVList, 
+		UTexture2D* Texture,
+		FGeometryScriptSampleTextureOptions SampleOptions,
+		FGeometryScriptColorList& ColorList,
+		UGeometryScriptDebug* Debug = nullptr);
+
+
+};
\ No newline at end of file
diff --git a/Engine/Source/Runtime/GeometryCore/Public/Image/ImageBuilder.h b/Engine/Source/Runtime/GeometryCore/Public/Image/ImageBuilder.h
--- a/Engine/Source/Runtime/GeometryCore/Public/Image/ImageBuilder.h
+++ b/Engine/Source/Runtime/GeometryCore/Public/Image/ImageBuilder.h
@@ -217,6 +217,24 @@ public:
 		return BilinearSample<ScalarType>(PixelCoords, InvalidValue);
 	}
 
+	/**
+	* Sample the image value at floating-point UV coords with Nearest sampling.
+	* The UV coords are assumed to be in range [0,1]x[0,1], and that this maps to the [0,Width]x[0,Height] image pixel rectangle.
+	*/
+	template<typename ScalarType>
+	PixelType NearestSampleUV(const FVector2d& UVCoords, const PixelType& InvalidValue) const
+	{
+		int X = (int)(UVCoords.X * (double)Dimensions.GetWidth());
+		int Y = (int)(UVCoords.Y * (double)Dimensions.GetHeight());
+		if (X < 0 || X >= Dimensions.GetWidth() ||
+			Y < 0 || Y >= Dimensions.GetHeight())
+		{
+			return InvalidValue;
+		}
+		return GetPixel(X, Y);
+	}
+
+
 	/**
 	 * @return true if all pixels have the same value
 	 */
