GeometryScript: add new list utility functions for scalar lists and manipulating vector lists. Add new library of vector math functions for lists of vectors/scalars.

diff --git a/Engine/Plugins/Experimental/GeometryScripting/Source/GeometryScriptingCore/Private/ListUtilityFunctions.cpp b/Engine/Plugins/Experimental/GeometryScripting/Source/GeometryScriptingCore/Private/ListUtilityFunctions.cpp
--- a/Engine/Plugins/Experimental/GeometryScripting/Source/GeometryScriptingCore/Private/ListUtilityFunctions.cpp
+++ b/Engine/Plugins/Experimental/GeometryScripting/Source/GeometryScriptingCore/Private/ListUtilityFunctions.cpp
@@ -1,7 +1,9 @@
 // Copyright Epic Games, Inc. All Rights Reserved.
 
 #include "GeometryScript/ListUtilityFunctions.h"
+#include "VectorTypes.h"
 
+using namespace UE::Geometry;
 
 #define LOCTEXT_NAMESPACE "UGeometryScriptLibrary_ListUtilityFunctions"
 
@@ -86,6 +88,69 @@ void UGeometryScriptLibrary_ListUtilityFunctions::ConvertArrayToTriangleList(con
 
 
 
+int UGeometryScriptLibrary_ListUtilityFunctions::GetScalarListLength(FGeometryScriptScalarList ScalarList)
+{
+	return (ScalarList.List.IsValid()) ? ScalarList.List->Num() : 0;
+}
+
+int UGeometryScriptLibrary_ListUtilityFunctions::GetScalarListLastIndex(FGeometryScriptScalarList ScalarList)
+{
+	return (ScalarList.List.IsValid()) ? FMath::Max(ScalarList.List->Num()-1,0) : 0;
+}
+
+double UGeometryScriptLibrary_ListUtilityFunctions::GetScalarListItem(FGeometryScriptScalarList ScalarList, int Index, bool& bIsValidIndex)
+{
+	bIsValidIndex = false;
+	if (ScalarList.List.IsValid() && Index >= 0 && Index < ScalarList.List->Num())
+	{
+		bIsValidIndex = true;
+		return (*ScalarList.List)[Index];
+	}
+	return 0.0;
+}
+
+void UGeometryScriptLibrary_ListUtilityFunctions::SetScalarListItem(FGeometryScriptScalarList& ScalarList, int Index, double NewValue, bool& bIsValidIndex)
+{
+	bIsValidIndex = false;
+	if (ScalarList.List.IsValid() && Index >= 0 && Index < ScalarList.List->Num())
+	{
+		bIsValidIndex = true;
+		(*ScalarList.List)[Index] = NewValue;
+	}
+}
+
+void UGeometryScriptLibrary_ListUtilityFunctions::ConvertScalarListToArray(FGeometryScriptScalarList ScalarList, TArray<double>& ScalarArray)
+{
+	ScalarArray.Reset();
+	if (ScalarList.List.IsValid())
+	{
+		ScalarArray.Append(*ScalarList.List);
+	}
+}
+
+void UGeometryScriptLibrary_ListUtilityFunctions::ConvertArrayToScalarList(const TArray<double>& ScalarArray, FGeometryScriptScalarList& ScalarList)
+{
+	ScalarList.Reset();
+	ScalarList.List->Append(ScalarArray);
+}
+
+void UGeometryScriptLibrary_ListUtilityFunctions::DuplicateScalarList(FGeometryScriptScalarList ScalarList, FGeometryScriptScalarList& DuplicateList)
+{
+	DuplicateList.Reset();
+	*DuplicateList.List = *ScalarList.List;
+}
+
+void UGeometryScriptLibrary_ListUtilityFunctions::ClearScalarList(FGeometryScriptScalarList& ScalarList, double ClearValue)
+{
+	int32 Num = ScalarList.List->Num();
+	ScalarList.Reset();
+	ScalarList.List->Init(ClearValue, Num);
+}
+
+
+
+
+
 int UGeometryScriptLibrary_ListUtilityFunctions::GetVectorListLength(FGeometryScriptVectorList VectorList)
 {
 	return (VectorList.List.IsValid()) ? VectorList.List->Num() : 0;
@@ -96,6 +161,7 @@ int UGeometryScriptLibrary_ListUtilityFunctions::GetVectorListLastIndex(FGeometr
 	return (VectorList.List.IsValid()) ? FMath::Max(VectorList.List->Num()-1,0) : 0;
 }
 
+
 FVector UGeometryScriptLibrary_ListUtilityFunctions::GetVectorListItem(FGeometryScriptVectorList VectorList, int Index, bool& bIsValidIndex)
 {
 	bIsValidIndex = false;
@@ -107,6 +173,16 @@ FVector UGeometryScriptLibrary_ListUtilityFunctions::GetVectorListItem(FGeometry
 	return FVector::Zero();
 }
 
+void UGeometryScriptLibrary_ListUtilityFunctions::SetVectorListItem(FGeometryScriptVectorList& VectorList, int Index, FVector NewValue, bool& bIsValidIndex)
+{
+	bIsValidIndex = false;
+	if (VectorList.List.IsValid() && Index >= 0 && Index < VectorList.List->Num())
+	{
+		bIsValidIndex = true;
+		(*VectorList.List)[Index] = NewValue;
+	}
+}
+
 void UGeometryScriptLibrary_ListUtilityFunctions::ConvertVectorListToArray(FGeometryScriptVectorList VectorList, TArray<FVector>& VectorArray)
 {
 	VectorArray.Reset();
@@ -122,7 +198,18 @@ void UGeometryScriptLibrary_ListUtilityFunctions::ConvertArrayToVectorList(const
 	VectorList.List->Append(VectorArray);
 }
 
+void UGeometryScriptLibrary_ListUtilityFunctions::DuplicateVectorList(FGeometryScriptVectorList VectorList, FGeometryScriptVectorList& DuplicateList)
+{
+	DuplicateList.Reset();
+	*DuplicateList.List = *VectorList.List;
+}
 
+void UGeometryScriptLibrary_ListUtilityFunctions::ClearVectorList(FGeometryScriptVectorList& VectorList, FVector ClearValue)
+{
+	int32 Num = VectorList.List->Num();
+	VectorList.Reset();
+	VectorList.List->Init(ClearValue, Num);
+}
 
 
 
@@ -202,6 +289,64 @@ void UGeometryScriptLibrary_ListUtilityFunctions::ConvertArrayToColorList(const
 	ColorList.List->Append(ColorArray);
 }
 
+void UGeometryScriptLibrary_ListUtilityFunctions::ExtractColorListChannel(FGeometryScriptColorList ColorList, FGeometryScriptScalarList& ScalarList, int32 ChannelIndex)
+{
+	if (ChannelIndex < 0 || ChannelIndex > 3)
+	{
+		UE_LOG(LogGeometry, Warning, TEXT("ExtractColorListChannel: ChannelIndex is not in valid range 0-3"));
+	}
+	ChannelIndex = FMath::Clamp(ChannelIndex, 0, 3);
+
+	ScalarList.Reset();
+	if (ColorList.List.IsValid())
+	{
+		TArray<FLinearColor>& Colors = *ColorList.List;
+		TArray<double>& Scalars = *ScalarList.List;
+		int32 N = Colors.Num();
+		Scalars.SetNumUninitialized(N);
+		for (int32 k = 0; k < N; ++k)
+		{
+			FVector4d Vector4 = (FVector4d)(Colors[k]);
+			Scalars[k] = Vector4[ChannelIndex];
+		}
+	}
+}
+
+
+void UGeometryScriptLibrary_ListUtilityFunctions::ExtractColorListChannels(
+	FGeometryScriptColorList ColorList, FGeometryScriptVectorList& VectorList, 
+	int32 XChannelIndex, int32 YChannelIndex, int32 ZChannelIndex)
+{
+	if (XChannelIndex < 0 || XChannelIndex > 3)
+	{
+		UE_LOG(LogGeometry, Warning, TEXT("ExtractColorListChannels: XChannelIndex is not in valid range 0-3"));
+	}
+	if (YChannelIndex < 0 || YChannelIndex > 3)
+	{
+		UE_LOG(LogGeometry, Warning, TEXT("ExtractColorListChannels: YChannelIndex is not in valid range 0-3"));
+	}
+	if (ZChannelIndex < 0 || ZChannelIndex > 3)
+	{
+		UE_LOG(LogGeometry, Warning, TEXT("ExtractColorListChannels: ZChannelIndex is not in valid range 0-3"));
+	}
+	XChannelIndex = FMath::Clamp(XChannelIndex, 0, 3);
+	YChannelIndex = FMath::Clamp(YChannelIndex, 0, 3);
+	ZChannelIndex = FMath::Clamp(ZChannelIndex, 0, 3);
+
+	VectorList.Reset();
+	if (ColorList.List.IsValid())
+	{
+		TArray<FLinearColor>& Colors = *ColorList.List;
+		TArray<FVector>& Vectors = *VectorList.List;
+		int32 N = Colors.Num();
+		Vectors.SetNumUninitialized(N);
+		for (int32 k = 0; k < N; ++k)
+		{
+			FVector4d Vector4 = (FVector4d)(Colors[k]);
+			Vectors[k] = FVector(Vector4[XChannelIndex], Vector4[YChannelIndex], Vector4[ZChannelIndex]);
+		}
+	}
+}
 
 
 
diff --git a/Engine/Plugins/Experimental/GeometryScripting/Source/GeometryScriptingCore/Private/VectorMathFunctions.cpp b/Engine/Plugins/Experimental/GeometryScripting/Source/GeometryScriptingCore/Private/VectorMathFunctions.cpp
new file mode 100644
--- /dev/null
+++ b/Engine/Plugins/Experimental/GeometryScripting/Source/GeometryScriptingCore/Private/VectorMathFunctions.cpp
@@ -0,0 +1,340 @@
+// Copyright Epic Games, Inc. All Rights Reserved.
+
+#include "GeometryScript/VectorMathFunctions.h"
+#include "VectorTypes.h"
+
+using namespace UE::Geometry;
+
+#define LOCTEXT_NAMESPACE "UGeometryScriptLibrary_VectorMathFunctions"
+
+
+namespace VectorMathFunctionsLocal
+{
+
+enum class EListOpResultCode
+{
+	Success = 0,
+	DifferentSizes = 1,
+};
+
+
+
+template<typename ListType1, typename ListTypeOut, typename Func>
+EListOpResultCode UnaryListToListOp(const ListType1& List1, ListTypeOut& OutputList, Func OperationFunc)
+{
+	if (List1.List.IsValid() == false)
+	{
+		return EListOpResultCode::DifferentSizes;
+	}
+
+	const auto& Array1 = *List1.List;
+	int32 N = Array1.Num();
+
+	// warning: OutputList may be List1 !!
+	if (OutputList.List.IsValid() == false)
+	{
+		OutputList.Reset();
+	}
+	auto& OutputArray = *OutputList.List;
+	if ( OutputArray.Num() != N )
+	{
+		OutputArray.SetNumUninitialized(N);
+	}
+
+	for (int32 k = 0; k < N; ++k)
+	{
+		OutputArray[k] = OperationFunc(Array1[k]);
+	}
+
+	return EListOpResultCode::Success;
+}
+
+
+
+template<typename ListType1, typename ListType2, typename ListType3, typename Func>
+EListOpResultCode BinaryListToListOp(const ListType1& List1, const ListType2& List2, ListType3& OutputList, Func OperationFunc)
+{
+	if (List1.List.IsValid() == false || List1.List.IsValid() == false || List1.List->Num() != List2.List->Num() )
+	{
+		return EListOpResultCode::DifferentSizes;
+	}
+
+	const auto& Array1 = *List1.List;
+	const auto& Array2 = *List2.List;
+	int32 N = Array1.Num();
+
+	// warning: OutputList may be List1 or List2 !!
+	if (OutputList.List.IsValid() == false)
+	{
+		OutputList.Reset();
+	}
+	auto& OutputArray = *OutputList.List;
+	if ( OutputArray.Num() != N )
+	{
+		OutputArray.SetNumUninitialized(N);
+	}
+	
+	// ParallelFor ?
+	for (int32 k = 0; k < N; ++k)
+	{
+		OutputArray[k] = OperationFunc(Array1[k], Array2[k]);
+	}
+
+	return EListOpResultCode::Success;
+}
+
+
+
+}
+
+
+FGeometryScriptScalarList UGeometryScriptLibrary_VectorMathFunctions::VectorLength(FGeometryScriptVectorList VectorList)
+{
+	FGeometryScriptScalarList ResultList;
+	VectorMathFunctionsLocal::EListOpResultCode Result = VectorMathFunctionsLocal::UnaryListToListOp(
+		VectorList, ResultList, 
+		[&](const FVector& Vector) { return Vector.Length(); } );
+	if ( Result == VectorMathFunctionsLocal::EListOpResultCode::DifferentSizes )
+	{
+		UE_LOG(LogGeometry, Warning, TEXT("VectorLength: VectorList is empty"));
+	}
+	return ResultList;
+}
+
+
+FGeometryScriptScalarList UGeometryScriptLibrary_VectorMathFunctions::VectorDot(FGeometryScriptVectorList VectorListA, FGeometryScriptVectorList VectorListB)
+{
+	FGeometryScriptScalarList ResultList;
+	VectorMathFunctionsLocal::EListOpResultCode Result = VectorMathFunctionsLocal::BinaryListToListOp(
+		VectorListA, VectorListB, ResultList, 
+		[&](const FVector& VectorA, const FVector& VectorB) { return FVector::DotProduct(VectorA, VectorB); } );
+	if ( Result == VectorMathFunctionsLocal::EListOpResultCode::DifferentSizes )
+	{
+		UE_LOG(LogGeometry, Warning, TEXT("VectorDot: Vector lists have different sizes"));
+	}
+	return ResultList;
+}
+
+
+FGeometryScriptVectorList UGeometryScriptLibrary_VectorMathFunctions::VectorCross(FGeometryScriptVectorList VectorListA, FGeometryScriptVectorList VectorListB)
+{
+	FGeometryScriptVectorList ResultList;
+	VectorMathFunctionsLocal::EListOpResultCode Result = VectorMathFunctionsLocal::BinaryListToListOp(
+		VectorListA, VectorListB, ResultList, 
+		[&](const FVector& VectorA, const FVector& VectorB) { return FVector::CrossProduct(VectorA, VectorB); } );
+	if ( Result == VectorMathFunctionsLocal::EListOpResultCode::DifferentSizes )
+	{
+		UE_LOG(LogGeometry, Warning, TEXT("VectorCross: Vector lists have different sizes"));
+	}
+	return ResultList;
+}
+
+
+
+void UGeometryScriptLibrary_VectorMathFunctions::VectorNormalizeInPlace(FGeometryScriptVectorList& VectorList, FVector SetOnFailure)
+{
+	VectorMathFunctionsLocal::EListOpResultCode Result = VectorMathFunctionsLocal::UnaryListToListOp( 
+		VectorList, VectorList, [&](const FVector& Vec) { 
+		return Vec.GetSafeNormal(UE_SMALL_NUMBER, SetOnFailure);
+	} );
+	if ( Result == VectorMathFunctionsLocal::EListOpResultCode::DifferentSizes )
+	{
+		UE_LOG(LogGeometry, Warning, TEXT("VectorNormalizeInPlace: VectorList is empty"));
+	}
+}
+
+
+
+FGeometryScriptVectorList UGeometryScriptLibrary_VectorMathFunctions::VectorBlend(FGeometryScriptVectorList VectorListA, FGeometryScriptVectorList VectorListB, double ConstantA, double ConstantB)
+{
+	FGeometryScriptVectorList ResultList;
+	VectorMathFunctionsLocal::EListOpResultCode Result = VectorMathFunctionsLocal::BinaryListToListOp(
+		VectorListA, VectorListB, ResultList, 
+		[&](const FVector& VectorA, const FVector& VectorB) { return ConstantA*VectorA + ConstantB*VectorB; } );
+	if ( Result == VectorMathFunctionsLocal::EListOpResultCode::DifferentSizes )
+	{
+		UE_LOG(LogGeometry, Warning, TEXT("VectorBlend: Vector lists have different sizes"));
+	}
+	return ResultList;
+}
+
+void UGeometryScriptLibrary_VectorMathFunctions::VectorBlendInPlace(FGeometryScriptVectorList VectorListA, FGeometryScriptVectorList& VectorListB, double ConstantA, double ConstantB)
+{
+	VectorMathFunctionsLocal::EListOpResultCode Result = VectorMathFunctionsLocal::BinaryListToListOp(
+		VectorListA, VectorListB, VectorListB, 
+		[&](const FVector& VectorA, const FVector& VectorB) { return ConstantA*VectorA + ConstantB*VectorB; } );
+	if ( Result == VectorMathFunctionsLocal::EListOpResultCode::DifferentSizes )
+	{
+		UE_LOG(LogGeometry, Warning, TEXT("VectorBlendInPlace: Vector lists have different sizes"));
+	}
+}
+
+
+
+FGeometryScriptVectorList UGeometryScriptLibrary_VectorMathFunctions::ScalarVectorMultiply(FGeometryScriptScalarList ScalarList, FGeometryScriptVectorList VectorList, double ScalarMultiplier)
+{
+	FGeometryScriptVectorList ResultList;
+	VectorMathFunctionsLocal::EListOpResultCode Result = VectorMathFunctionsLocal::BinaryListToListOp(
+		ScalarList, VectorList, ResultList, 
+		[&](double Scalar, const FVector& Vector) { return (FVector::FReal)(Scalar * ScalarMultiplier) * Vector; } );
+	if ( Result == VectorMathFunctionsLocal::EListOpResultCode::DifferentSizes )
+	{
+		UE_LOG(LogGeometry, Warning, TEXT("ScalarVectorMultiply: Scalar and Vector lists have different sizes"));
+	}
+	return ResultList;
+}
+
+void UGeometryScriptLibrary_VectorMathFunctions::ScalarVectorMultiplyInPlace(FGeometryScriptScalarList ScalarList, FGeometryScriptVectorList& VectorList, double ScalarMultiplier)
+{
+	VectorMathFunctionsLocal::EListOpResultCode Result = VectorMathFunctionsLocal::BinaryListToListOp(
+		ScalarList, VectorList, VectorList, 
+		[&](double Scalar, const FVector& Vector) { return (FVector::FReal)(Scalar * ScalarMultiplier) * Vector; } );
+	if ( Result == VectorMathFunctionsLocal::EListOpResultCode::DifferentSizes )
+	{
+		UE_LOG(LogGeometry, Warning, TEXT("ScalarVectorMultiplyInPlace: Scalar and Vector lists have different sizes"));
+	}
+}
+
+
+FGeometryScriptVectorList UGeometryScriptLibrary_VectorMathFunctions::ConstantVectorMultiply(double Constant, FGeometryScriptVectorList VectorList)
+{
+	FGeometryScriptVectorList ResultList;
+	VectorMathFunctionsLocal::EListOpResultCode Result = VectorMathFunctionsLocal::UnaryListToListOp(
+		VectorList, ResultList, 
+		[&](const FVector& Vector) { return (FVector::FReal)(Constant) * Vector; } );
+	if ( Result == VectorMathFunctionsLocal::EListOpResultCode::DifferentSizes )
+	{
+		UE_LOG(LogGeometry, Warning, TEXT("ConstantVectorMultiply: VectorList is empty"));
+	}
+	return ResultList;
+}
+
+void UGeometryScriptLibrary_VectorMathFunctions::ConstantVectorMultiplyInPlace(double Constant, FGeometryScriptVectorList& VectorList)
+{
+	VectorMathFunctionsLocal::EListOpResultCode Result = VectorMathFunctionsLocal::UnaryListToListOp(
+		VectorList, VectorList, 
+		[&](const FVector& Vector) { return (FVector::FReal)(Constant) * Vector; } );
+	if ( Result == VectorMathFunctionsLocal::EListOpResultCode::DifferentSizes )
+	{
+		UE_LOG(LogGeometry, Warning, TEXT("ConstantVectorMultiplyInPlace: VectorList is empty"));
+	}
+}
+
+
+FGeometryScriptScalarList UGeometryScriptLibrary_VectorMathFunctions::VectorToScalar(FGeometryScriptVectorList VectorList, double ConstantX, double ConstantY, double ConstantZ)
+{
+	FGeometryScriptScalarList ResultList;
+	VectorMathFunctionsLocal::EListOpResultCode Result = VectorMathFunctionsLocal::UnaryListToListOp(
+		VectorList, ResultList, 
+		[&](const FVector& Vector) { return ConstantX*Vector.X + ConstantY*Vector.Y + ConstantZ*Vector.Z; } );
+	if ( Result == VectorMathFunctionsLocal::EListOpResultCode::DifferentSizes )
+	{
+		UE_LOG(LogGeometry, Warning, TEXT("VectorLength: VectorList is empty"));
+	}
+	return ResultList;
+}
+
+
+
+
+FGeometryScriptScalarList UGeometryScriptLibrary_VectorMathFunctions::ScalarInvert(FGeometryScriptScalarList ScalarList, double Numerator, double SetOnFailure, double Epsilon)
+{
+	FGeometryScriptScalarList ResultList;
+	VectorMathFunctionsLocal::EListOpResultCode Result = VectorMathFunctionsLocal::UnaryListToListOp(
+		ScalarList, ResultList, 
+		[&](const double& Scalar) { return (FMathd::Abs(Scalar) > Epsilon) ? (Numerator / Scalar) : SetOnFailure; } );
+	if ( Result == VectorMathFunctionsLocal::EListOpResultCode::DifferentSizes )
+	{
+		UE_LOG(LogGeometry, Warning, TEXT("ScalarInvert: Scalar lists have different sizes"));
+	}
+	return ResultList;
+}
+
+void UGeometryScriptLibrary_VectorMathFunctions::ScalarInvertInPlace(FGeometryScriptScalarList& ScalarList, double Numerator, double SetOnFailure, double Epsilon)
+{
+	VectorMathFunctionsLocal::EListOpResultCode Result = VectorMathFunctionsLocal::UnaryListToListOp(
+		ScalarList, ScalarList, 
+		[&](const double& Scalar) { return (FMathd::Abs(Scalar) > Epsilon) ? (Numerator / Scalar) : SetOnFailure; } );
+	if ( Result == VectorMathFunctionsLocal::EListOpResultCode::DifferentSizes )
+	{
+		UE_LOG(LogGeometry, Warning, TEXT("ScalarInvertInPlace: ScalarList is empty"));
+	}
+}
+
+
+FGeometryScriptScalarList UGeometryScriptLibrary_VectorMathFunctions::ScalarBlend(FGeometryScriptScalarList ScalarListA, FGeometryScriptScalarList ScalarListB, double ConstantA, double ConstantB)
+{
+	FGeometryScriptScalarList ResultList;
+	VectorMathFunctionsLocal::EListOpResultCode Result = VectorMathFunctionsLocal::BinaryListToListOp(
+		ScalarListA, ScalarListB, ResultList, 
+		[&](const double& ScalarA, const double& ScalarB) { return ConstantA*ScalarA + ConstantB*ScalarB; } );
+	if ( Result == VectorMathFunctionsLocal::EListOpResultCode::DifferentSizes )
+	{
+		UE_LOG(LogGeometry, Warning, TEXT("ScalarBlend: Scalar lists have different sizes"));
+	}
+	return ResultList;
+}
+
+void UGeometryScriptLibrary_VectorMathFunctions::ScalarBlendInPlace(FGeometryScriptScalarList ScalarListA, FGeometryScriptScalarList& ScalarListB, double ConstantA, double ConstantB)
+{
+	VectorMathFunctionsLocal::EListOpResultCode Result = VectorMathFunctionsLocal::BinaryListToListOp(
+		ScalarListA, ScalarListB, ScalarListB, 
+		[&](const double& ScalarA, const double& ScalarB) { return ConstantA*ScalarA + ConstantB*ScalarB; } );
+	if ( Result == VectorMathFunctionsLocal::EListOpResultCode::DifferentSizes )
+	{
+		UE_LOG(LogGeometry, Warning, TEXT("ScalarBlendInPlace: Scalar lists have different sizes"));
+	}
+}
+
+
+FGeometryScriptScalarList UGeometryScriptLibrary_VectorMathFunctions::ScalarMultiply(FGeometryScriptScalarList ScalarListA, FGeometryScriptScalarList ScalarListB, double ScalarMultiplier)
+{
+	FGeometryScriptScalarList ResultList;
+	VectorMathFunctionsLocal::EListOpResultCode Result = VectorMathFunctionsLocal::BinaryListToListOp(
+		ScalarListA, ScalarListB, ResultList, 
+		[&](const double& ScalarA, const double& ScalarB) { return ScalarMultiplier*ScalarA*ScalarB; } );
+	if ( Result == VectorMathFunctionsLocal::EListOpResultCode::DifferentSizes )
+	{
+		UE_LOG(LogGeometry, Warning, TEXT("ScalarMultiply: Scalar lists have different sizes"));
+	}
+	return ResultList;
+}
+
+void UGeometryScriptLibrary_VectorMathFunctions::ScalarMultiplyInPlace(FGeometryScriptScalarList ScalarListA, FGeometryScriptScalarList& ScalarListB, double ScalarMultiplier)
+{
+	VectorMathFunctionsLocal::EListOpResultCode Result = VectorMathFunctionsLocal::BinaryListToListOp(
+		ScalarListA, ScalarListB, ScalarListB, 
+		[&](const double& ScalarA, const double& ScalarB) { return ScalarMultiplier*ScalarA*ScalarB; } );
+	if ( Result == VectorMathFunctionsLocal::EListOpResultCode::DifferentSizes )
+	{
+		UE_LOG(LogGeometry, Warning, TEXT("ScalarMultiplyInPlace: Scalar lists have different sizes"));
+	}
+}
+
+
+FGeometryScriptScalarList UGeometryScriptLibrary_VectorMathFunctions::ConstantScalarMultiply(double Constant, FGeometryScriptScalarList ScalarList)
+{
+	FGeometryScriptScalarList ResultList;
+	VectorMathFunctionsLocal::EListOpResultCode Result = VectorMathFunctionsLocal::UnaryListToListOp(
+		ScalarList, ResultList, 
+		[&](const double& Scalar) { return Constant*Scalar; } );
+	if ( Result == VectorMathFunctionsLocal::EListOpResultCode::DifferentSizes )
+	{
+		UE_LOG(LogGeometry, Warning, TEXT("ConstantScalarMultiply: Scalar lists have different sizes"));
+	}
+	return ResultList;
+}
+
+void UGeometryScriptLibrary_VectorMathFunctions::ConstantScalarMultiplyInPlace(double Constant, FGeometryScriptScalarList& ScalarList)
+{
+	VectorMathFunctionsLocal::EListOpResultCode Result = VectorMathFunctionsLocal::UnaryListToListOp(
+		ScalarList, ScalarList, 
+		[&](const double& Scalar) { return Constant*Scalar; } );
+	if ( Result == VectorMathFunctionsLocal::EListOpResultCode::DifferentSizes )
+	{
+		UE_LOG(LogGeometry, Warning, TEXT("ConstantScalarMultiplyInPlace: ScalarList is empty"));
+	}
+}
+
+
+#undef LOCTEXT_NAMESPACE
diff --git a/Engine/Plugins/Experimental/GeometryScripting/Source/GeometryScriptingCore/Public/GeometryScript/GeometryScriptTypes.h b/Engine/Plugins/Experimental/GeometryScripting/Source/GeometryScriptingCore/Public/GeometryScript/GeometryScriptTypes.h
--- a/Engine/Plugins/Experimental/GeometryScripting/Source/GeometryScriptingCore/Public/GeometryScript/GeometryScriptTypes.h
+++ b/Engine/Plugins/Experimental/GeometryScripting/Source/GeometryScriptingCore/Public/GeometryScript/GeometryScriptTypes.h
@@ -248,6 +248,24 @@ public:
 };
 
 
+USTRUCT(BlueprintType, meta = (DisplayName = "Scalar List"))
+struct GEOMETRYSCRIPTINGCORE_API FGeometryScriptScalarList
+{
+	GENERATED_BODY()
+public:
+	TSharedPtr<TArray<double>> List;
+
+	void Reset()
+	{
+		if (List.IsValid() == false)
+		{
+			List = MakeShared<TArray<double>>();
+		}
+		List->Reset();
+	}
+};
+
+
 
 USTRUCT(BlueprintType, meta = (DisplayName = "Vector List"))
 struct GEOMETRYSCRIPTINGCORE_API FGeometryScriptVectorList
diff --git a/Engine/Plugins/Experimental/GeometryScripting/Source/GeometryScriptingCore/Public/GeometryScript/ListUtilityFunctions.h b/Engine/Plugins/Experimental/GeometryScripting/Source/GeometryScriptingCore/Public/GeometryScript/ListUtilityFunctions.h
--- a/Engine/Plugins/Experimental/GeometryScripting/Source/GeometryScriptingCore/Public/GeometryScript/ListUtilityFunctions.h
+++ b/Engine/Plugins/Experimental/GeometryScripting/Source/GeometryScriptingCore/Public/GeometryScript/ListUtilityFunctions.h
@@ -50,6 +50,33 @@ public:
 
 
 
+	UFUNCTION(BlueprintPure, Category = "GeometryScript|ListUtils", meta=(ScriptMethod))
+	static int GetScalarListLength(FGeometryScriptScalarList ScalarList);
+
+	UFUNCTION(BlueprintPure, Category = "GeometryScript|ListUtils", meta=(ScriptMethod))
+	static int GetScalarListLastIndex(FGeometryScriptScalarList ScalarList);
+
+	UFUNCTION(BlueprintPure, Category = "GeometryScript|ListUtils", meta=(ScriptMethod))
+	static double GetScalarListItem(FGeometryScriptScalarList ScalarList, int Index, bool& bIsValidIndex);
+
+	UFUNCTION(BlueprintCallable, Category = "GeometryScript|ListUtils", meta=(ScriptMethod))
+	static void SetScalarListItem(UPARAM(ref) FGeometryScriptScalarList& ScalarList, int Index, double NewValue, bool& bIsValidIndex);
+
+	UFUNCTION(BlueprintCallable, Category = "GeometryScript|ListUtils", meta=(ScriptMethod))
+	static void ConvertScalarListToArray(FGeometryScriptScalarList ScalarList, TArray<double>& ScalarArray);
+
+	UFUNCTION(BlueprintCallable, Category = "GeometryScript|ListUtils")
+	static void ConvertArrayToScalarList(const TArray<double>& VectorArray, FGeometryScriptScalarList& ScalarList);
+
+	UFUNCTION(BlueprintCallable, Category = "GeometryScript|ListUtils", meta=(ScriptMethod))
+	static void DuplicateScalarList(FGeometryScriptScalarList ScalarList, FGeometryScriptScalarList& DuplicateList);
+
+	UFUNCTION(BlueprintCallable, Category = "GeometryScript|ListUtils", meta=(ScriptMethod))
+	static void ClearScalarList(UPARAM(ref) FGeometryScriptScalarList& ScalarList, double ClearValue = 0.0);
+
+
+
+
 	UFUNCTION(BlueprintPure, Category = "GeometryScript|ListUtils", meta=(ScriptMethod))
 	static int GetVectorListLength(FGeometryScriptVectorList VectorList);
 
@@ -59,12 +86,22 @@ public:
 	UFUNCTION(BlueprintPure, Category = "GeometryScript|ListUtils", meta=(ScriptMethod))
 	static FVector GetVectorListItem(FGeometryScriptVectorList VectorList, int Index, bool& bIsValidIndex);
 
+	UFUNCTION(BlueprintCallable, Category = "GeometryScript|ListUtils", meta=(ScriptMethod))
+	static void SetVectorListItem(UPARAM(ref) FGeometryScriptVectorList& VectorList, int Index, FVector NewValue, bool& bIsValidIndex);
+
 	UFUNCTION(BlueprintCallable, Category = "GeometryScript|ListUtils", meta=(ScriptMethod))
 	static void ConvertVectorListToArray(FGeometryScriptVectorList VectorList, TArray<FVector>& VectorArray);
 
 	UFUNCTION(BlueprintCallable, Category = "GeometryScript|ListUtils")
 	static void ConvertArrayToVectorList(const TArray<FVector>& VectorArray, FGeometryScriptVectorList& VectorList);
 
+	UFUNCTION(BlueprintCallable, Category = "GeometryScript|ListUtils", meta=(ScriptMethod))
+	static void DuplicateVectorList(FGeometryScriptVectorList VectorList, FGeometryScriptVectorList& DuplicateList);
+
+	UFUNCTION(BlueprintCallable, Category = "GeometryScript|ListUtils", meta=(ScriptMethod))
+	static void ClearVectorList(UPARAM(ref) FGeometryScriptVectorList& VectorList, FVector ClearValue = FVector::ZeroVector);
+
+
 
 
 	UFUNCTION(BlueprintPure, Category = "GeometryScript|ListUtils", meta=(ScriptMethod))
@@ -84,6 +121,7 @@ public:
 
 
 
+
 	UFUNCTION(BlueprintPure, Category = "GeometryScript|ListUtils", meta=(ScriptMethod))
 	static int GetColorListLength(FGeometryScriptColorList ColorList);
 
@@ -98,4 +136,11 @@ public:
 
 	UFUNCTION(BlueprintCallable, Category = "GeometryScript|ListUtils")
 	static void ConvertArrayToColorList(const TArray<FLinearColor>& ColorArray, FGeometryScriptColorList& ColorList);
+
+	UFUNCTION(BlueprintCallable, Category = "GeometryScript|ListUtils", meta=(ScriptMethod))
+	static void ExtractColorListChannel(FGeometryScriptColorList ColorList, FGeometryScriptScalarList& ScalarList, int32 ChannelIndex = 0);
+
+	UFUNCTION(BlueprintCallable, Category = "GeometryScript|ListUtils", meta=(ScriptMethod))
+	static void ExtractColorListChannels(FGeometryScriptColorList ColorList, FGeometryScriptVectorList& VectorList, int32 XChannelIndex = 0, int32 YChannelIndex = 1, int32 ZChannelIndex = 2);
+
 };
\ No newline at end of file
diff --git a/Engine/Plugins/Experimental/GeometryScripting/Source/GeometryScriptingCore/Public/GeometryScript/VectorMathFunctions.h b/Engine/Plugins/Experimental/GeometryScripting/Source/GeometryScriptingCore/Public/GeometryScript/VectorMathFunctions.h
new file mode 100644
--- /dev/null
+++ b/Engine/Plugins/Experimental/GeometryScripting/Source/GeometryScriptingCore/Public/GeometryScript/VectorMathFunctions.h
@@ -0,0 +1,158 @@
+// Copyright Epic Games, Inc. All Rights Reserved.
+
+#pragma once
+
+#include "CoreMinimal.h"
+#include "Kismet/BlueprintFunctionLibrary.h"
+#include "GeometryScript/GeometryScriptTypes.h"
+#include "VectorMathFunctions.generated.h"
+
+class UDynamicMesh;
+
+
+UCLASS(meta = (ScriptName = "GeometryScript_VectorMath"))
+class GEOMETRYSCRIPTINGCORE_API UGeometryScriptLibrary_VectorMathFunctions : public UBlueprintFunctionLibrary
+{
+	GENERATED_BODY()
+public:
+
+	/**
+	 * Compute the length/magnitude of each vector in VectorListA and return in new ScalarList.
+	 * Note that SquaredLength can be computed using VectorDot(A,A).
+	 */
+	UFUNCTION(BlueprintCallable, Category = "GeometryScript|Math|VectorMath", meta=(ScriptMethod))
+	static FGeometryScriptScalarList 
+	VectorLength(FGeometryScriptVectorList VectorList);
+
+	/**
+	 * Compute the dot-product between each pair of vectors in VectorListA and VectorListB and return in new ScalarList
+	 */
+	UFUNCTION(BlueprintCallable, Category = "GeometryScript|Math|VectorMath", meta=(ScriptMethod))
+	static FGeometryScriptScalarList 
+	VectorDot(FGeometryScriptVectorList VectorListA, FGeometryScriptVectorList VectorListB);
+
+	/**
+	* Compute the cross-product between each pair of vectors in VectorListA and VectorListB and return in new VectorList
+	*/
+	UFUNCTION(BlueprintCallable, Category = "GeometryScript|Math|VectorMath", meta=(ScriptMethod))
+	static FGeometryScriptVectorList 
+	VectorCross(FGeometryScriptVectorList VectorListA, FGeometryScriptVectorList VectorListB);
+
+	/**
+	 * Normalize each vector in VectorList, and store in VectorList. 
+	 * If a vector is degenerate, set the normal to the SetOnFailure vector.
+	 */
+	UFUNCTION(BlueprintCallable, Category = "GeometryScript|Math|VectorMath", meta=(ScriptMethod))
+	static void 
+	VectorNormalizeInPlace(UPARAM(ref) FGeometryScriptVectorList& VectorList, FVector SetOnFailure = FVector::ZeroVector);
+
+	/**
+	 * Compute (ConstantA * A) + (ConstantB * B) for each pair of vectors in VectorListA and VectorListB and return in new VectorList.
+	 * By default (constants = 1) this just adds the two vectors. Set ConstantB = -1 to subtract B from A. 
+	 * Can also be used to Linear Interpolate, by setting ConstantB = (1-ConstantA)
+	 */
+	UFUNCTION(BlueprintCallable, Category = "GeometryScript|Math|VectorMath", meta=(ScriptMethod))
+	static FGeometryScriptVectorList 
+	VectorBlend(FGeometryScriptVectorList VectorListA, FGeometryScriptVectorList VectorListB, double ConstantA = 1.0, double ConstantB = 1.0);
+
+	/**
+	* Compute (ConstantA * A) + (ConstantB * B) for each pair of vectors in VectorListA and VectorListB, and store in VectorListB
+	* By default (constants = 1) this just adds the two vectors. Set ConstantB = -1 to subtract B from A. 
+	* Can also be used to Linear Interpolate, by setting ConstantB = (1-ConstantA)
+	*/
+	UFUNCTION(BlueprintCallable, Category = "GeometryScript|Math|VectorMath", meta=(ScriptMethod))
+	static void 
+	VectorBlendInPlace(FGeometryScriptVectorList VectorListA, UPARAM(ref) FGeometryScriptVectorList& VectorListB, double ConstantA = 1.0, double ConstantB = 1.0);
+
+	/**
+	 * Compute (ScalarMultiplier * Scalar * Vector) for each scalar/vector pair in the two input lists, and return in a new VectorList.
+	 */
+	UFUNCTION(BlueprintCallable, Category = "GeometryScript|Math|VectorMath", meta=(ScriptMethod))
+	static FGeometryScriptVectorList 
+	ScalarVectorMultiply(FGeometryScriptScalarList ScalarList, FGeometryScriptVectorList VectorList, double ScalarMultiplier = 1.0);
+
+	/**
+	 * Compute (ScalarMultiplier * Scalar * Vector) for each scalar/vector pair in the two input lists, and store in the input VectorList
+	 */
+	UFUNCTION(BlueprintCallable, Category = "GeometryScript|Math|VectorMath", meta=(ScriptMethod))
+	static void 
+	ScalarVectorMultiplyInPlace(FGeometryScriptScalarList ScalarList, UPARAM(ref) FGeometryScriptVectorList& VectorList, double ScalarMultiplier = 1.0);
+
+	/**
+	 * Compute (Constant * Vector) for each element in VectorList, and return in a new list
+	 */
+	UFUNCTION(BlueprintCallable, Category = "GeometryScript|Math|VectorMath")
+	static FGeometryScriptVectorList 
+	ConstantVectorMultiply(double Constant, FGeometryScriptVectorList VectorList);
+
+	/**
+	 * Compute (Constant * Vector) for each element in VectorList, and store in VectorList
+	 */
+	UFUNCTION(BlueprintCallable, Category = "GeometryScript|Math|VectorMath")
+	static void 
+	ConstantVectorMultiplyInPlace(double Constant, UPARAM(ref) FGeometryScriptVectorList& VectorList);
+
+
+	/**
+	 * Convert each Vector in VectorList to a Scalar by computing (ConstantX*Vector.X + ConstantY*Vector.Y + ConstantZ*Vector.Z), and return in a new ScalarList.
+	 * This can be used to extract the X/Y/Z values from a Vector, or other component-wise math
+	 */
+	UFUNCTION(BlueprintCallable, Category = "GeometryScript|Math|VectorMath", meta=(ScriptMethod))
+	static FGeometryScriptScalarList 
+	VectorToScalar(FGeometryScriptVectorList VectorList, double ConstantX = 1.0, double ConstantY = 0.0, double ConstantZ = 0.0);
+
+
+
+	/**
+	 * Compute (Numerator / Scalar) for each element of ScalarList and return in a new ScalarList. 
+	 * If Abs(Scalar) < Epsilon, set to SetOnFailure value.
+	 */
+	UFUNCTION(BlueprintCallable, Category = "GeometryScript|Math|VectorMath", meta=(ScriptMethod))
+	static FGeometryScriptScalarList 
+	ScalarInvert(FGeometryScriptScalarList ScalarList, double Numerator = 1.0, double SetOnFailure = 0.0, double Epsilon = 0.0);
+
+	/**
+	* Compute (Numerator / Scalar) for each element of ScalarList and store in input ScalarList
+	* If Abs(Scalar) < Epsilon, set to SetOnFailure value.
+	*/
+	UFUNCTION(BlueprintCallable, Category = "GeometryScript|Math|VectorMath", meta=(ScriptMethod))
+	static void 
+	ScalarInvertInPlace(UPARAM(ref) FGeometryScriptScalarList& ScalarList, double Numerator = 1.0, double SetOnFailure = 0.0, double Epsilon = 0.0);
+
+	/**
+	 * Compute (ConstantA * A) + (ConstantB * B) for each pair of values in ScalarListA and ScalarListB and return in new ScalarList.
+	 * By default (constants = 1) this just adds the two values. Set ConstantB = -1 to subtract B from A. 
+	 * Can also be used to Linear Interpolate, by setting ConstantB = (1-ConstantA)
+	 */
+	UFUNCTION(BlueprintCallable, Category = "GeometryScript|VectorMath", meta=(ScriptMethod))
+	static FGeometryScriptScalarList 
+	ScalarBlend(FGeometryScriptScalarList ScalarListA, FGeometryScriptScalarList ScalarListB, double ConstantA = 1.0, double ConstantB = 1.0);
+
+	/**
+	 * Compute (ConstantA * A) + (ConstantB * B) for each pair of values in ScalarListA and ScalarListB and return in ScalarListB.
+	 * By default (constants = 1) this just adds the two values. Set ConstantB = -1 to subtract B from A. 
+	 * Can also be used to Linear Interpolate, by setting ConstantB = (1-ConstantA)
+	 */
+	UFUNCTION(BlueprintCallable, Category = "GeometryScript|Math|VectorMath", meta=(ScriptMethod))
+	static void 
+	ScalarBlendInPlace(FGeometryScriptScalarList ScalarListA, UPARAM(ref) FGeometryScriptScalarList& ScalarListB, double ConstantA = 1.0, double ConstantB = 1.0);
+
+
+	UFUNCTION(BlueprintCallable, Category = "GeometryScript|Math|VectorMath", meta=(ScriptMethod))
+	static FGeometryScriptScalarList 
+	ScalarMultiply(FGeometryScriptScalarList ScalarListA, FGeometryScriptScalarList ScalarListB, double ConstantMultiplier = 1.0);
+
+	UFUNCTION(BlueprintCallable, Category = "GeometryScript|Math|VectorMath", meta=(ScriptMethod))
+	static void 
+	ScalarMultiplyInPlace(FGeometryScriptScalarList ScalarListA, UPARAM(ref) FGeometryScriptScalarList& ScalarListB, double ConstantMultiplier = 1.0);
+
+
+	UFUNCTION(BlueprintCallable, Category = "GeometryScript|Math|VectorMath")
+	static FGeometryScriptScalarList 
+	ConstantScalarMultiply(double Constant, FGeometryScriptScalarList ScalarList);
+
+	UFUNCTION(BlueprintCallable, Category = "GeometryScript|Math|VectorMath")
+	static void 
+	ConstantScalarMultiplyInPlace(double Constant, UPARAM(ref) FGeometryScriptScalarList& ScalarList);
+
+};
\ No newline at end of file
