guard against empty path for append sweep geometry script function

diff --git a/Engine/Plugins/Experimental/GeometryScripting/Source/GeometryScriptingCore/Private/MeshPrimitiveFunctions.cpp b/Engine/Plugins/Experimental/GeometryScripting/Source/GeometryScriptingCore/Private/MeshPrimitiveFunctions.cpp
--- a/Engine/Plugins/Experimental/GeometryScripting/Source/GeometryScriptingCore/Private/MeshPrimitiveFunctions.cpp
+++ b/Engine/Plugins/Experimental/GeometryScripting/Source/GeometryScriptingCore/Private/MeshPrimitiveFunctions.cpp
@@ -620,6 +620,11 @@ UDynamicMesh* UGeometryScriptLibrary_MeshPrimitiveFunctions::AppendSimpleSweptPo
 		UE::Geometry::AppendError(Debug, EGeometryScriptErrorType::InvalidInputs, LOCTEXT("AppendSimpleSweptPolygon_InvalidPolygon", "AppendSimpleSweptPolygon: PolygonVertices array requires at least 3 positions"));
 		return TargetMesh;
 	}
+	if (SweepPath.Num() < 2)
+	{
+		UE::Geometry::AppendError(Debug, EGeometryScriptErrorType::InvalidInputs, LOCTEXT("AppendSimpleSweptPolygon_InvalidSweepPath", "AppendSimpleSweptPolygon: SweepPath array requires at least 2 positions"));
+		return TargetMesh;
+	}
 
 	FGeneralizedCylinderGenerator SweepGen;
 	for (FVector2D Point : PolygonVertices)
