Make PolyPath FlattenTo2DOnAxis function create a flattened copy rather than modifying the passed-in path

diff --git a/Engine/Plugins/Experimental/GeometryScripting/Source/GeometryScriptingCore/Private/PolyPathFunctions.cpp b/Engine/Plugins/Experimental/GeometryScripting/Source/GeometryScriptingCore/Private/PolyPathFunctions.cpp
--- a/Engine/Plugins/Experimental/GeometryScripting/Source/GeometryScriptingCore/Private/PolyPathFunctions.cpp
+++ b/Engine/Plugins/Experimental/GeometryScripting/Source/GeometryScriptingCore/Private/PolyPathFunctions.cpp
@@ -58,21 +58,25 @@ int32 UGeometryScriptLibrary_PolyPathFunctions::GetNearestVertexIndex(FGeometryS
 	return -1;
 }
 
-UFUNCTION(BlueprintCallable, Category = "GeometryScript|PolyPath")
 FGeometryScriptPolyPath UGeometryScriptLibrary_PolyPathFunctions::FlattenTo2DOnAxis(FGeometryScriptPolyPath PolyPath, EGeometryScriptAxis DropAxis)
 {
+	FGeometryScriptPolyPath ToReturn;
 	if (PolyPath.Path.IsValid())
 	{
+		ToReturn.Reset();
+		ToReturn.bClosedLoop = PolyPath.bClosedLoop;
+		ToReturn.Path = PolyPath.Path;
+
 		int32 Keep0 = int32(DropAxis == EGeometryScriptAxis::X);
 		int32 Keep1 = 1 + int32(DropAxis != EGeometryScriptAxis::Z);
-		for (FVector& V : *PolyPath.Path)
+		for (FVector& V : *ToReturn.Path)
 		{
 			V[0] = V[Keep0];
 			V[1] = V[Keep1];
 			V.Z = 0;
 		}
 	}
-	return PolyPath;
+	return ToReturn;
 }
 
 void UGeometryScriptLibrary_PolyPathFunctions::ConvertPolyPathToArray(FGeometryScriptPolyPath PolyPath, TArray<FVector>& PathVertices)
diff --git a/Engine/Plugins/Experimental/GeometryScripting/Source/GeometryScriptingCore/Public/GeometryScript/PolyPathFunctions.h b/Engine/Plugins/Experimental/GeometryScripting/Source/GeometryScriptingCore/Public/GeometryScript/PolyPathFunctions.h
--- a/Engine/Plugins/Experimental/GeometryScripting/Source/GeometryScriptingCore/Public/GeometryScript/PolyPathFunctions.h
+++ b/Engine/Plugins/Experimental/GeometryScripting/Source/GeometryScriptingCore/Public/GeometryScript/PolyPathFunctions.h
@@ -62,9 +62,9 @@ public:
 	UFUNCTION(BlueprintCallable, Category = "GeometryScript|PolyPath", meta = (ScriptMethod))
 	static int32 GetNearestVertexIndex(FGeometryScriptPolyPath PolyPath, FVector Point);
 
-	/** Flatten to 2D by dropping the given axis, and using the other two coordinates as the new X, Y coordinates.  Returns the modified path for convenience. */
+	/** Create a 2D, flattened copy of the path by dropping the given axis, and using the other two coordinates as the new X, Y coordinates. */
 	UFUNCTION(BlueprintCallable, Category = "GeometryScript|PolyPath", meta = (ScriptMethod, DisplayName = "Flatten To 2D On Axis"))
-	static UPARAM(DisplayName = "Target Poly Path") FGeometryScriptPolyPath FlattenTo2DOnAxis(FGeometryScriptPolyPath TargetPolyPath, EGeometryScriptAxis DropAxis = EGeometryScriptAxis::Z);
+	static UPARAM(DisplayName = "Poly Path") FGeometryScriptPolyPath FlattenTo2DOnAxis(FGeometryScriptPolyPath PolyPath, EGeometryScriptAxis DropAxis = EGeometryScriptAxis::Z);
 
 	UFUNCTION(BlueprintCallable, Category = "GeometryScript|PolyPath")
 	static void ConvertSplineToPolyPath(const USplineComponent* Spline, FGeometryScriptPolyPath& PolyPath, FGeometryScriptSplineSamplingOptions SamplingOptions);
