GeometryScript: fix copy-paste-o in GetVertexPosition, it is validating against Triangle IDs

diff --git a/Engine/Plugins/Experimental/GeometryScripting/Source/GeometryScriptingCore/Private/MeshQueryFunctions.cpp b/Engine/Plugins/Experimental/GeometryScripting/Source/GeometryScriptingCore/Private/MeshQueryFunctions.cpp
--- a/Engine/Plugins/Experimental/GeometryScripting/Source/GeometryScriptingCore/Private/MeshQueryFunctions.cpp
+++ b/Engine/Plugins/Experimental/GeometryScripting/Source/GeometryScriptingCore/Private/MeshQueryFunctions.cpp
@@ -270,7 +270,7 @@ UDynamicMesh* UGeometryScriptLibrary_MeshQueryFunctions::GetAllVertexIDs(UDynami
 FVector UGeometryScriptLibrary_MeshQueryFunctions::GetVertexPosition(UDynamicMesh* TargetMesh, int32 VertexID, bool& bIsValidVertex)
 {
 	return SimpleMeshQuery<FVector>(TargetMesh, FVector::ZeroVector, [&](const FDynamicMesh3& Mesh) {
-		bIsValidVertex = Mesh.IsTriangle(VertexID);
+		bIsValidVertex = Mesh.IsVertex(VertexID);
 		return (bIsValidVertex) ? (FVector)Mesh.GetVertex(VertexID) : FVector::ZeroVector;
 	});
 }
