Fix bug revealed by static analysis warning

diff --git a/Engine/Plugins/Experimental/GeometryScripting/Source/GeometryScriptingCore/Private/VectorMathFunctions.cpp b/Engine/Plugins/Experimental/GeometryScripting/Source/GeometryScriptingCore/Private/VectorMathFunctions.cpp
--- a/Engine/Plugins/Experimental/GeometryScripting/Source/GeometryScriptingCore/Private/VectorMathFunctions.cpp
+++ b/Engine/Plugins/Experimental/GeometryScripting/Source/GeometryScriptingCore/Private/VectorMathFunctions.cpp
@@ -54,7 +54,7 @@ EListOpResultCode UnaryListToListOp(const ListType1& List1, ListTypeOut& OutputL
 template<typename ListType1, typename ListType2, typename ListType3, typename Func>
 EListOpResultCode BinaryListToListOp(const ListType1& List1, const ListType2& List2, ListType3& OutputList, Func OperationFunc)
 {
-	if (List1.List.IsValid() == false || List1.List.IsValid() == false || List1.List->Num() != List2.List->Num() )
+	if (List1.List.IsValid() == false || List2.List.IsValid() == false || List1.List->Num() != List2.List->Num() )
 	{
 		return EListOpResultCode::DifferentSizes;
 	}
