Make geometry script functions use ApplyTransformInverse and InverseTransformPosition rather than trying to invert a transform and then use it

diff --git a/Engine/Plugins/Experimental/GeometryScripting/Source/GeometryScriptingCore/Private/MeshBakeFunctions.cpp b/Engine/Plugins/Experimental/GeometryScripting/Source/GeometryScriptingCore/Private/MeshBakeFunctions.cpp
--- a/Engine/Plugins/Experimental/GeometryScripting/Source/GeometryScriptingCore/Private/MeshBakeFunctions.cpp
+++ b/Engine/Plugins/Experimental/GeometryScripting/Source/GeometryScriptingCore/Private/MeshBakeFunctions.cpp
@@ -497,7 +497,7 @@ namespace GeometryScriptBakeLocals
 			const FTransformSRT3d SourceToWorld = SourceTransform;
 			MeshTransforms::ApplyTransform(SourceMeshCopy, SourceToWorld);
 			const FTransformSRT3d TargetToWorld = TargetTransform;
-			MeshTransforms::ApplyTransform(SourceMeshCopy, TargetToWorld.Inverse());
+			MeshTransforms::ApplyTransformInverse(SourceMeshCopy, TargetToWorld);
 			SourceMeshToUse = &SourceMeshCopy;
 		}
 
@@ -662,7 +662,7 @@ namespace GeometryScriptBakeLocals
 			const FTransformSRT3d SourceToWorld = SourceTransform;
 			MeshTransforms::ApplyTransform(SourceMeshCopy, SourceToWorld);
 			const FTransformSRT3d TargetToWorld = TargetTransform;
-			MeshTransforms::ApplyTransform(SourceMeshCopy, TargetToWorld.Inverse());
+			MeshTransforms::ApplyTransformInverse(SourceMeshCopy, TargetToWorld);
 			SourceMeshToUse = &SourceMeshCopy;
 		}
 
diff --git a/Engine/Plugins/Experimental/GeometryScripting/Source/GeometryScriptingCore/Private/MeshBooleanFunctions.cpp b/Engine/Plugins/Experimental/GeometryScripting/Source/GeometryScriptingCore/Private/MeshBooleanFunctions.cpp
--- a/Engine/Plugins/Experimental/GeometryScripting/Source/GeometryScriptingCore/Private/MeshBooleanFunctions.cpp
+++ b/Engine/Plugins/Experimental/GeometryScripting/Source/GeometryScriptingCore/Private/MeshBooleanFunctions.cpp
@@ -85,7 +85,7 @@ UDynamicMesh* UGeometryScriptLibrary_MeshBooleanFunctions::ApplyMeshBoolean(
 	}
 
 	// Boolean result is in the space of TargetTransform, so invert that
-	MeshTransforms::ApplyTransform(NewResultMesh, (FTransformSRT3d)TargetTransform.Inverse());
+	MeshTransforms::ApplyTransformInverse(NewResultMesh, (FTransformSRT3d)TargetTransform);
 
 	if (NewBoundaryEdges.Num() > 0 && Options.bFillHoles)
 	{
diff --git a/Engine/Plugins/Experimental/GeometryScripting/Source/GeometryScriptingCore/Private/MeshComparisonFunctions.cpp b/Engine/Plugins/Experimental/GeometryScripting/Source/GeometryScriptingCore/Private/MeshComparisonFunctions.cpp
--- a/Engine/Plugins/Experimental/GeometryScripting/Source/GeometryScriptingCore/Private/MeshComparisonFunctions.cpp
+++ b/Engine/Plugins/Experimental/GeometryScripting/Source/GeometryScriptingCore/Private/MeshComparisonFunctions.cpp
@@ -101,9 +101,8 @@ UDynamicMesh* UGeometryScriptLibrary_MeshComparisonFunctions::IsIntersectingMesh
 			}
 			else
 			{
-				Transform1 = Transform1.Inverse();
 				bIsIntersecting = Spatials[0].TestIntersection(Spatials[1],
-					[&](const FVector3d& Pos) { return Transform1.TransformPosition(Transform2.TransformPosition(Pos)); });
+					[&](const FVector3d& Pos) { return Transform1.InverseTransformPosition(Transform2.TransformPosition(Pos)); });
 			}
 		});
 	});
diff --git a/Engine/Plugins/Experimental/GeometryScripting/Source/GeometryScriptingCore/Private/MeshModelingFunctions.cpp b/Engine/Plugins/Experimental/GeometryScripting/Source/GeometryScriptingCore/Private/MeshModelingFunctions.cpp
--- a/Engine/Plugins/Experimental/GeometryScripting/Source/GeometryScriptingCore/Private/MeshModelingFunctions.cpp
+++ b/Engine/Plugins/Experimental/GeometryScripting/Source/GeometryScriptingCore/Private/MeshModelingFunctions.cpp
@@ -447,7 +447,6 @@ UDynamicMesh* UGeometryScriptLibrary_MeshModelingFunctions::ApplyMeshPolygroupBe
 		if (Options.bApplyFilterBox)
 		{
 			FAxisAlignedBox3d QueryBox(Options.FilterBox);
-			FTransformSRT3d InvTransform(Options.FilterBoxTransform.Inverse());
 			TSet<int32> FoundEdges;
 
 			// find all mesh edges inside filter shape
@@ -455,7 +454,7 @@ UDynamicMesh* UGeometryScriptLibrary_MeshModelingFunctions::ApplyMeshPolygroupBe
 			FDynamicMeshAABBTree3 Spatial(&EditMesh, true);
 			FDynamicMeshAABBTree3::FTreeTraversal EdgeTraversal;
 			EdgeTraversal.NextBoxF = [&QueryBox](const FAxisAlignedBox3d& Box, int Depth) { return Box.Intersects(QueryBox); };
-			EdgeTraversal.NextTriangleF = [&QueryBox, &InvTransform, &FoundEdges, &EditMesh, &Options](int TriangleID)
+			EdgeTraversal.NextTriangleF = [&QueryBox, &FoundEdges, &EditMesh, &Options](int TriangleID)
 			{
 				FIndex3i Edges = EditMesh.GetTriEdges(TriangleID);
 				for (int32 j = 0; j < 3; ++j)
@@ -464,8 +463,8 @@ UDynamicMesh* UGeometryScriptLibrary_MeshModelingFunctions::ApplyMeshPolygroupBe
 					{
 						FVector3d A, B;
 						EditMesh.GetEdgeV(Edges[j], A, B);
-						A = InvTransform.TransformPosition(A);
-						B = InvTransform.TransformPosition(B);
+						A = Options.FilterBoxTransform.InverseTransformPosition(A);
+						B = Options.FilterBoxTransform.InverseTransformPosition(B);
 						bool bIntersects = (Options.bFullyContained) ?
 							(QueryBox.Contains(A) && QueryBox.Contains(B)) :
 							(QueryBox.Contains(A) || QueryBox.Contains(B));
