add CreateCircle and CreateArc functions for geometry script PolyPath

diff --git a/Engine/Plugins/Experimental/GeometryScripting/Source/GeometryScriptingCore/Private/PolyPathFunctions.cpp b/Engine/Plugins/Experimental/GeometryScripting/Source/GeometryScriptingCore/Private/PolyPathFunctions.cpp
--- a/Engine/Plugins/Experimental/GeometryScripting/Source/GeometryScriptingCore/Private/PolyPathFunctions.cpp
+++ b/Engine/Plugins/Experimental/GeometryScripting/Source/GeometryScriptingCore/Private/PolyPathFunctions.cpp
@@ -79,6 +79,78 @@ FGeometryScriptPolyPath UGeometryScriptLibrary_PolyPathFunctions::FlattenTo2DOnA
 	return ToReturn;
 }
 
+namespace PolyPathInternal
+{
+	
+TArray<FVector> CreateArcPoints(FTransform Transform, float Radius, int NumPoints, float StartAngle, float EndAngle, bool bIncludeEndPoint)
+{
+	TArray<FVector> Points;
+	NumPoints = FMath::Max(2, NumPoints);
+	Points.Reserve(NumPoints);
+
+	double AngleStartRad = FMathd::DegToRad * double(StartAngle);
+	double AngleRangeRad = FMathd::DegToRad * double(EndAngle) - AngleStartRad;
+	double ToAngle = AngleRangeRad / double(NumPoints - bIncludeEndPoint); // if bIncludeEndPoint, adjust fraction to reach end at PtIdx = NumPoints-1
+	for (int PtIdx = 0; PtIdx < NumPoints; PtIdx++)
+	{
+		double Angle = AngleStartRad + PtIdx * ToAngle;
+		Points.Emplace(Transform.TransformPosition(FVector3d(FMathd::Cos(Angle) * (double)Radius, FMathd::Sin(Angle) * (double)Radius, 0.0)));
+	}
+
+	return Points;
+}
+
+}
+
+FGeometryScriptPolyPath UGeometryScriptLibrary_PolyPathFunctions::CreateCirclePath3D(FTransform Transform, float Radius, int NumPoints)
+{
+	FGeometryScriptPolyPath PolyPath;
+	PolyPath.Reset();
+	PolyPath.bClosedLoop = true;
+
+	*PolyPath.Path = PolyPathInternal::CreateArcPoints(Transform, Radius, NumPoints, 0, 360, false);
+
+	return PolyPath;
+}
+
+FGeometryScriptPolyPath UGeometryScriptLibrary_PolyPathFunctions::CreateArcPath3D(FTransform Transform, float Radius, int NumPoints, float StartAngle, float EndAngle)
+{
+	FGeometryScriptPolyPath PolyPath;
+	PolyPath.Reset();
+
+	PolyPath.bClosedLoop = false;
+
+	*PolyPath.Path = PolyPathInternal::CreateArcPoints(Transform, Radius, NumPoints, StartAngle, EndAngle, true);
+
+	return PolyPath;
+}
+
+FGeometryScriptPolyPath UGeometryScriptLibrary_PolyPathFunctions::CreateCirclePath2D(FVector2D Center, float Radius, int NumPoints)
+{
+	FGeometryScriptPolyPath PolyPath;
+	PolyPath.Reset();
+
+	PolyPath.bClosedLoop = true;
+
+	FTransform TransformCenter(FVector(Center.X, Center.Y, 0));
+	*PolyPath.Path = PolyPathInternal::CreateArcPoints(TransformCenter, Radius, NumPoints, 0, 360, false);
+
+	return PolyPath;
+}
+
+FGeometryScriptPolyPath UGeometryScriptLibrary_PolyPathFunctions::CreateArcPath2D(FVector2D Center, float Radius, int NumPoints, float StartAngle, float EndAngle)
+{
+	FGeometryScriptPolyPath PolyPath;
+	PolyPath.Reset();
+
+	PolyPath.bClosedLoop = false;
+
+	FTransform TransformCenter(FVector(Center.X, Center.Y, 0));
+	*PolyPath.Path = PolyPathInternal::CreateArcPoints(TransformCenter, Radius, NumPoints, StartAngle, EndAngle, true);
+
+	return PolyPath;
+}
+
 void UGeometryScriptLibrary_PolyPathFunctions::ConvertPolyPathToArray(FGeometryScriptPolyPath PolyPath, TArray<FVector>& PathVertices)
 {
 	PathVertices.Reset();
diff --git a/Engine/Plugins/Experimental/GeometryScripting/Source/GeometryScriptingCore/Public/GeometryScript/PolyPathFunctions.h b/Engine/Plugins/Experimental/GeometryScripting/Source/GeometryScriptingCore/Public/GeometryScript/PolyPathFunctions.h
--- a/Engine/Plugins/Experimental/GeometryScripting/Source/GeometryScriptingCore/Public/GeometryScript/PolyPathFunctions.h
+++ b/Engine/Plugins/Experimental/GeometryScripting/Source/GeometryScriptingCore/Public/GeometryScript/PolyPathFunctions.h
@@ -66,6 +66,34 @@ public:
 	UFUNCTION(BlueprintCallable, Category = "GeometryScript|PolyPath", meta = (ScriptMethod, DisplayName = "Flatten To 2D On Axis"))
 	static UPARAM(DisplayName = "Poly Path") FGeometryScriptPolyPath FlattenTo2DOnAxis(FGeometryScriptPolyPath PolyPath, EGeometryScriptAxis DropAxis = EGeometryScriptAxis::Z);
 
+	/**
+	 * Create a closed circle around the origin on the XY plane, then transformed by Transform.
+	 * By our convention for closed paths, the end vertex is *not* a duplicate of the start vertex.
+	 */
+	UFUNCTION(BlueprintCallable, Category = "GeometryScript|PolyPath")
+	static UPARAM(DisplayName = "Poly Path") FGeometryScriptPolyPath CreateCirclePath3D(FTransform Transform, float Radius = 10, int NumPoints = 10);
+
+	/**
+	 * Create an open arc around the origin on the XY plane, then transformed by Transform.
+	 * As it is an open path, the end vertex exactly hits the target EndAngle (so will be positioned on the start vertex if the end aligns to the start)
+	 */
+	UFUNCTION(BlueprintCallable, Category = "GeometryScript|PolyPath")
+	static UPARAM(DisplayName = "Poly Path") FGeometryScriptPolyPath CreateArcPath3D(FTransform Transform, float Radius = 10, int NumPoints = 10, float StartAngle = 0, float EndAngle = 90);
+	
+	/**
+	 * Create a closed circle on the XY plane around the given Center.
+	 * By our convention for closed paths, the end vertex is *not* a duplicate of the start vertex.
+	 */
+	UFUNCTION(BlueprintCallable, Category = "GeometryScript|PolyPath")
+	static UPARAM(DisplayName = "Poly Path") FGeometryScriptPolyPath CreateCirclePath2D(FVector2D Center = FVector2D(0, 0), float Radius = 10, int NumPoints = 10);
+
+	/**
+	 * Create an open arc on the XY plane around the given Center.
+	 * As it is an open path, the end vertex exactly hits the target EndAngle (so will be positioned on the start vertex if the end aligns to the start)
+	 */
+	UFUNCTION(BlueprintCallable, Category = "GeometryScript|PolyPath")
+	static UPARAM(DisplayName = "Poly Path") FGeometryScriptPolyPath CreateArcPath2D(FVector2D Center = FVector2D(0, 0), float Radius = 10, int NumPoints = 10, float StartAngle = 0, float EndAngle = 90);
+
 	UFUNCTION(BlueprintCallable, Category = "GeometryScript|PolyPath")
 	static void ConvertSplineToPolyPath(const USplineComponent* Spline, FGeometryScriptPolyPath& PolyPath, FGeometryScriptSplineSamplingOptions SamplingOptions);
 
