CIS fix static analysis warning on polypath by explicitly casting bool

diff --git a/Engine/Plugins/Experimental/GeometryScripting/Source/GeometryScriptingCore/Private/PolyPathFunctions.cpp b/Engine/Plugins/Experimental/GeometryScripting/Source/GeometryScriptingCore/Private/PolyPathFunctions.cpp
--- a/Engine/Plugins/Experimental/GeometryScripting/Source/GeometryScriptingCore/Private/PolyPathFunctions.cpp
+++ b/Engine/Plugins/Experimental/GeometryScripting/Source/GeometryScriptingCore/Private/PolyPathFunctions.cpp
@@ -90,7 +90,7 @@ TArray<FVector> CreateArcPoints(FTransform Transform, float Radius, int NumPoint
 
 	double AngleStartRad = FMathd::DegToRad * double(StartAngle);
 	double AngleRangeRad = FMathd::DegToRad * double(EndAngle) - AngleStartRad;
-	double ToAngle = AngleRangeRad / double(NumPoints - bIncludeEndPoint); // if bIncludeEndPoint, adjust fraction to reach end at PtIdx = NumPoints-1
+	double ToAngle = AngleRangeRad / double(NumPoints - (int)bIncludeEndPoint); // if bIncludeEndPoint, adjust fraction to reach end at PtIdx = NumPoints-1
 	for (int PtIdx = 0; PtIdx < NumPoints; PtIdx++)
 	{
 		double Angle = AngleStartRad + PtIdx * ToAngle;
