Geometry: Expose Remesh function in Geometry Script, and add basic convergence check to queue remesher
- add FQueueRemesher::MinActiveEdgeFraction. If fraction of modified/total edges in remesh pass is below this parameter (default 1%), consider result converged.
- Expose as option in FRemeshMeshOp, and add sane defaults for all parameters of Op
- move CalculateTargetEdgeLength function from RemeshMeshTool to static function in FRemeshMeshOp, update Tool
- add Geometry Script function ApplyUniformRemesh

diff --git a/Engine/Plugins/Experimental/GeometryScripting/Source/GeometryScriptingCore/Private/MeshRemeshFunctions.cpp b/Engine/Plugins/Experimental/GeometryScripting/Source/GeometryScriptingCore/Private/MeshRemeshFunctions.cpp
new file mode 100644
--- /dev/null
+++ b/Engine/Plugins/Experimental/GeometryScripting/Source/GeometryScriptingCore/Private/MeshRemeshFunctions.cpp
@@ -0,0 +1,148 @@
+// Copyright Epic Games, Inc. All Rights Reserved.
+
+#include "GeometryScript/MeshRemeshFunctions.h"
+
+#include "DynamicMesh/DynamicMesh3.h"
+#include "DynamicMesh/DynamicMeshAABBTree3.h"
+#include "DynamicMesh/DynamicMeshAttributeSet.h"
+#include "DynamicMesh/MeshAttributeUtil.h"
+#include "MeshConstraintsUtil.h"
+#include "ProjectionTargets.h"
+#include "DynamicMesh/MeshNormals.h"
+#include "CleaningOps/RemeshMeshOp.h"
+
+#include "UDynamicMesh.h"
+
+using namespace UE::Geometry;
+
+#define LOCTEXT_NAMESPACE "UGeometryScriptLibrary_MeshRemeshFunctions"
+
+
+
+static EEdgeRefineFlags MakeEdgeRefineFlagsFromConstraintType(EGeometryScriptRemeshEdgeConstraintType ConstraintType)
+{
+	switch (ConstraintType)
+	{
+		case EGeometryScriptRemeshEdgeConstraintType::Fixed: return EEdgeRefineFlags::FullyConstrained;
+		case EGeometryScriptRemeshEdgeConstraintType::Refine: return EEdgeRefineFlags::SplitsOnly;
+		case EGeometryScriptRemeshEdgeConstraintType::Free: return EEdgeRefineFlags::NoFlip;
+		case EGeometryScriptRemeshEdgeConstraintType::Ignore: return EEdgeRefineFlags::NoConstraint;
+	}
+	return EEdgeRefineFlags::NoFlip;
+}
+
+
+UDynamicMesh* UGeometryScriptLibrary_RemeshingFunctions::ApplyUniformRemesh(
+	UDynamicMesh* TargetMesh, 
+	FGeometryScriptRemeshOptions RemeshOptions,
+	FGeometryScriptUniformRemeshOptions UniformOptions,
+	UGeometryScriptDebug* Debug)
+{
+	if (TargetMesh == nullptr)
+	{
+		UE::Geometry::AppendError(Debug, EGeometryScriptErrorType::InvalidInputs, LOCTEXT("ApplyUniformRemesh_InvalidInput", "ApplyUniformRemesh: TargetMesh is Null"));
+		return TargetMesh;
+	}
+
+	TargetMesh->EditMesh([&](FDynamicMesh3& EditMesh) 
+	{
+		TSharedPtr<FDynamicMesh3> SourceMesh = MakeShared<FDynamicMesh3>(MoveTemp(EditMesh));
+
+		TSharedPtr<FDynamicMeshAABBTree3> SourceSpatial;
+		if (RemeshOptions.bReprojectToInputMesh)
+		{
+			SourceSpatial = MakeShared<FDynamicMeshAABBTree3>(SourceMesh.Get(), true);
+		}
+
+		FRemeshMeshOp RemeshOp;
+		RemeshOp.OriginalMesh = SourceMesh;
+		RemeshOp.OriginalMeshSpatial = SourceSpatial;
+
+		RemeshOp.bDiscardAttributes = RemeshOptions.bDiscardAttributes;
+
+		RemeshOp.RemeshType = (RemeshOptions.bUseFullRemeshPasses) ? ERemeshType::FullPass : ERemeshType::Standard;
+
+		RemeshOp.RemeshIterations = RemeshOptions.RemeshIterations;
+		RemeshOp.MaxRemeshIterations = RemeshOptions.RemeshIterations;
+		RemeshOp.ExtraProjectionIterations = 0;		// unused for regular remeshing
+		RemeshOp.TriangleCountHint = 0;				// unused for regular remeshing
+
+		// smoothing options
+		RemeshOp.SmoothingStrength = FMath::Clamp(RemeshOptions.SmoothingRate, 0.0f, 1.0f);
+		switch (RemeshOptions.SmoothingType)
+		{
+		case EGeometryScriptRemeshSmoothingType::Mixed:
+			RemeshOp.SmoothingType = ERemeshSmoothingType::MeanValue;
+			break;
+		case EGeometryScriptRemeshSmoothingType::UVPreserving:
+			RemeshOp.SmoothingType = ERemeshSmoothingType::Cotangent;
+			break;
+		case EGeometryScriptRemeshSmoothingType::Uniform:
+		default:
+			RemeshOp.SmoothingType = ERemeshSmoothingType::Uniform;
+			break;
+		}
+
+		if (UniformOptions.TargetType == EGeometryScriptUniformRemeshTargetType::TriangleCount)
+		{
+			RemeshOp.TargetEdgeLength = FRemeshMeshOp::CalculateTargetEdgeLength(SourceMesh.Get(), UniformOptions.TargetTriangleCount);
+		}
+		else
+		{
+			RemeshOp.TargetEdgeLength = UniformOptions.TargetEdgeLength;
+		}
+
+		// currently not exposing this option. It seems to control multiple things that should be independent...
+		RemeshOp.bPreserveSharpEdges = (RemeshOp.bDiscardAttributes == false);
+		RemeshOp.bFlips = RemeshOptions.bAllowFlips; 
+		RemeshOp.bSplits = RemeshOptions.bAllowSplits; 
+		RemeshOp.bCollapses = RemeshOptions.bAllowCollapses; 
+		RemeshOp.bPreventNormalFlips = RemeshOptions.bPreventNormalFlips;
+		RemeshOp.bPreventTinyTriangles = RemeshOptions.bPreventTinyTriangles;
+
+		RemeshOp.MeshBoundaryConstraint = MakeEdgeRefineFlagsFromConstraintType(RemeshOptions.MeshBoundaryConstraint);
+		RemeshOp.GroupBoundaryConstraint = MakeEdgeRefineFlagsFromConstraintType(RemeshOptions.GroupBoundaryConstraint);
+		RemeshOp.MaterialBoundaryConstraint = MakeEdgeRefineFlagsFromConstraintType(RemeshOptions.MaterialBoundaryConstraint);
+
+		RemeshOp.bReproject = RemeshOptions.bReprojectToInputMesh; 
+		RemeshOp.ProjectionTarget = nullptr;
+		RemeshOp.ProjectionTargetSpatial = nullptr;		
+
+
+		RemeshOp.bReprojectConstraints = false;
+		RemeshOp.BoundaryCornerAngleThreshold = 45.0;
+
+
+		RemeshOp.TargetMeshLocalToWorld = FTransformSRT3d::Identity();
+		RemeshOp.ToolMeshLocalToWorld = FTransformSRT3d::Identity();
+		RemeshOp.bUseWorldSpace = false;
+		RemeshOp.bParallel = true;
+
+		RemeshOp.CalculateResult(nullptr);
+		if (RemeshOp.GetResultInfo().Result == EGeometryResultType::Success)
+		{
+			TUniquePtr<FDynamicMesh3> ResultMesh = RemeshOp.ExtractResult();
+			EditMesh = MoveTemp(*ResultMesh);
+		}
+		else
+		{
+			EditMesh = MoveTemp(*SourceMesh);
+			UE::Geometry::AppendError(Debug, EGeometryScriptErrorType::InvalidInputs, LOCTEXT("ApplyUniformRemesh_ComputeError", "ApplyUniformRemesh: Error computing result, returning input mesh"));
+		}
+
+		// if we discarded attributes, re-enable the standard attributes
+		if (RemeshOptions.bDiscardAttributes)
+		{
+			EditMesh.EnableTriangleGroups();
+			EditMesh.EnableAttributes();
+			EditMesh.Attributes()->EnableMaterialID();
+		}
+
+	}, EDynamicMeshChangeType::GeneralEdit, EDynamicMeshAttributeChangeFlags::Unknown, false);
+
+	return TargetMesh;
+}
+
+
+
+#undef LOCTEXT_NAMESPACE
\ No newline at end of file
diff --git a/Engine/Plugins/Experimental/GeometryScripting/Source/GeometryScriptingCore/Public/GeometryScript/MeshRemeshFunctions.h b/Engine/Plugins/Experimental/GeometryScripting/Source/GeometryScriptingCore/Public/GeometryScript/MeshRemeshFunctions.h
new file mode 100644
--- /dev/null
+++ b/Engine/Plugins/Experimental/GeometryScripting/Source/GeometryScriptingCore/Public/GeometryScript/MeshRemeshFunctions.h
@@ -0,0 +1,157 @@
+// Copyright Epic Games, Inc. All Rights Reserved.
+
+#pragma once
+
+#include "CoreMinimal.h"
+#include "Kismet/BlueprintFunctionLibrary.h"
+#include "GeometryScript/GeometryScriptTypes.h"
+#include "MeshRemeshFunctions.generated.h"
+
+class UDynamicMesh;
+
+
+/** Goal types for Uniform Remeshing */
+UENUM(BlueprintType)
+enum class EGeometryScriptUniformRemeshTargetType : uint8
+{
+	/** Approximate Desired Triangle Count. This is used to compute a Target Edge Length, and is not an explicit target */
+	TriangleCount = 0,
+	/** Attempt to Remesh such that all edges have approximately this length */
+	TargetEdgeLength = 1
+};
+
+/** Types of edge constraints, specified for different mesh attributes */
+UENUM(BlueprintType)
+enum class EGeometryScriptRemeshEdgeConstraintType : uint8
+{
+	/** Constrained edges cannot be flipped, split or collapsed, and vertices will not move */
+	Fixed = 0,
+	/** Constrained edges can be split, but not flipped or collapsed. Vertices will not move. */
+	Refine = 1,
+	/** Constrained edges cannot be flipped, but otherwise are free to move */
+	Free = 2,
+	/** Edges are not constrained, ie the Attribute used to derive the Constraints will not be considered */
+	Ignore = 3
+};
+
+/** The Vertex Smoothing strategy used in a Remeshing operation */
+UENUM(BlueprintType)
+enum class EGeometryScriptRemeshSmoothingType : uint8
+{
+	/** Vertices move towards their 3D one-ring centroids, UVs are ignored. This produces the most regular mesh possible. */
+	Uniform = 0,
+	/** Vertices move towards the projection of their one-ring centroids onto their normal vectors, preserving UVs  */
+	UVPreserving = 1,
+	/** Similar to UV Preserving, but allows some tangential drift (causing UV distortion) when vertices would otherwise be "stuck" */
+	Mixed = 2
+};
+
+
+/**
+ * Standard Remeshing Options
+ */
+USTRUCT(Blueprintable)
+struct GEOMETRYSCRIPTINGCORE_API FGeometryScriptRemeshOptions
+{
+	GENERATED_BODY()
+public:
+
+	/** When enabled, all mesh attributes are discarded, so UV and Normal Seams can be freely */
+	UPROPERTY(BlueprintReadWrite, Category = Options)
+	bool bDiscardAttributes = false;
+
+	/** When enabled, mesh vertices are projected back onto the input mesh surface during Remeshing, preserving the shape */
+	UPROPERTY(BlueprintReadWrite, Category = Options)
+	bool bReprojectToInputMesh = true;
+
+	/** Type of 3D Mesh Smoothing to apply during Remeshing. Disable by setting SmoothingRate = 0 */
+	UPROPERTY(BlueprintReadWrite, Category = Options)
+	EGeometryScriptRemeshSmoothingType SmoothingType = EGeometryScriptRemeshSmoothingType::Mixed;
+
+	/** Smoothing Rate/Speed. Faster Smoothing results in a more regular mesh, but also more potential for undesirable 3D shape change and UV distortion */
+	UPROPERTY(BlueprintReadWrite, Category = Options, meta = (UIMin = 0, UIMax = 1, ClampMin = 0, ClampMax = 1))
+	float SmoothingRate = 0.25f;
+
+	/** Constraints on the open mesh boundary/border edges */
+	UPROPERTY(BlueprintReadWrite, Category = Options)
+	EGeometryScriptRemeshEdgeConstraintType MeshBoundaryConstraint = EGeometryScriptRemeshEdgeConstraintType::Free;
+
+	/** Constraints on the mesh boundary/border edges between different PolyGroups of the Mesh */
+	UPROPERTY(BlueprintReadWrite, Category = Options)
+	EGeometryScriptRemeshEdgeConstraintType GroupBoundaryConstraint = EGeometryScriptRemeshEdgeConstraintType::Free;
+
+	/** Constraints on the mesh boundary/border edges between different Material Results of the Mesh */
+	UPROPERTY(BlueprintReadWrite, Category = Options)
+	EGeometryScriptRemeshEdgeConstraintType MaterialBoundaryConstraint = EGeometryScriptRemeshEdgeConstraintType::Free;
+
+	/** Enable/Disable Edge Flips during Remeshing. Disabling flips will significantly reduce the output mesh quality */
+	UPROPERTY(BlueprintReadWrite, Category = Options)
+	bool bAllowFlips = true;
+
+	/** Enable/Disable Edge Splits during Remeshing. Disabling Splits will prevent the mesh density from increasing. */
+	UPROPERTY(BlueprintReadWrite, Category = Options)
+	bool bAllowSplits = true;
+
+	/** Enable/Disable Edge Collapses during Remeshing. Disabling Collapses will prevent the mesh density from decreasing. */
+	UPROPERTY(BlueprintReadWrite, Category = Options)
+	bool bAllowCollapses = true;
+
+	/** When Enabled, Flips and Collapses will be skipped if they would flip any triangle face normals */
+	UPROPERTY(BlueprintReadWrite, Category = Options)
+	bool bPreventNormalFlips = true;
+
+	/** When Enabled, Flips and Collapses will be skipped if they would create tiny degenerate triangles */
+	UPROPERTY(BlueprintReadWrite, Category = Options)
+	bool bPreventTinyTriangles = true;
+
+	/** By default, remeshing is accelerated by tracking a queue of edges that need to be processed. This is signficantly faster but can produce a lower quality output. Enable this option to use a more expensive strategy that guarantees maximum quality.  */
+	UPROPERTY(BlueprintReadWrite, Category = Options)
+	bool bUseFullRemeshPasses = false;
+
+	/** Maximum Number of iterations of the Remeshing Strategy to apply to the Mesh. More iterations are generally more expensive (much moreso with bUseFullRemeshPasses = true) */
+	UPROPERTY(BlueprintReadWrite, Category = Options, meta = (UIMin = 0, ClampMin = 0))
+	int32 RemeshIterations = 20;
+};
+
+
+/**
+ * Uniform Remeshing Options
+ */
+USTRUCT(Blueprintable)
+struct GEOMETRYSCRIPTINGCORE_API FGeometryScriptUniformRemeshOptions
+{
+	GENERATED_BODY()
+public:
+	/** Method used to define target/goal of Uniform Remeshing */
+	UPROPERTY(BlueprintReadWrite, Category = Options)
+	EGeometryScriptUniformRemeshTargetType TargetType = EGeometryScriptUniformRemeshTargetType::TriangleCount;
+
+	/** Approximate Target Triangle Count, combined with mesh surface area to derive a TargetEdgeLength */
+	UPROPERTY(BlueprintReadWrite, Category = Options, meta = (UIMin = 0, ClampMin = 0, EditCondition = "TargetType == EGeometryScriptUniformRemeshTargetType::TriangleCount"))
+	int32 TargetTriangleCount = 5000;
+
+	/** Explicit Target Edge Length that is desired in the output uniform mesh */
+	UPROPERTY(BlueprintReadWrite, Category = Options, meta = (UIMin = 0, ClampMin = 0, EditCondition = "TargetType == EGeometryScriptUniformRemeshTargetType::TargetEdgeLength"))
+	float TargetEdgeLength = 1.0f;
+};
+
+
+
+UCLASS(meta = (ScriptName = "GeometryScript_Remeshing"))
+class GEOMETRYSCRIPTINGCORE_API UGeometryScriptLibrary_RemeshingFunctions : public UBlueprintFunctionLibrary
+{
+	GENERATED_BODY()
+public:
+
+	/**
+	 * Apply Uniform Remeshing to the TargetMesh. 
+	 */
+	UFUNCTION(BlueprintCallable, Category = "GeometryScript|Simplification", meta=(ScriptMethod))
+	static UPARAM(DisplayName = "Target Mesh") UDynamicMesh* 
+	ApplyUniformRemesh(  
+		UDynamicMesh* TargetMesh, 
+		FGeometryScriptRemeshOptions RemeshOptions,
+		FGeometryScriptUniformRemeshOptions UniformOptions,
+		UGeometryScriptDebug* Debug = nullptr);
+
+};
\ No newline at end of file
diff --git a/Engine/Plugins/Experimental/MeshModelingToolsetExp/Source/MeshModelingToolsExp/Private/ProjectToTargetTool.cpp b/Engine/Plugins/Experimental/MeshModelingToolsetExp/Source/MeshModelingToolsExp/Private/ProjectToTargetTool.cpp
--- a/Engine/Plugins/Experimental/MeshModelingToolsetExp/Source/MeshModelingToolsExp/Private/ProjectToTargetTool.cpp
+++ b/Engine/Plugins/Experimental/MeshModelingToolsetExp/Source/MeshModelingToolsExp/Private/ProjectToTargetTool.cpp
@@ -74,6 +74,9 @@ TUniquePtr<FDynamicMeshOperator> UProjectToTargetTool::MakeNewOperator()
 		RemeshOp->bSmoothInFillAreas = ProjectProperties->bSmoothInFillAreas;
 		RemeshOp->FillAreaDistanceMultiplier = ProjectProperties->FillAreaDistanceMultiplier;
 		RemeshOp->FillAreaSmoothMultiplier = ProjectProperties->FillAreaSmoothMultiplier;
+
+		// disable convergence check
+		RemeshOp->MinActiveEdgeFraction = 0.0;
 	}
 
 	return Op;
diff --git a/Engine/Plugins/Experimental/MeshModelingToolsetExp/Source/MeshModelingToolsExp/Private/RemeshMeshTool.cpp b/Engine/Plugins/Experimental/MeshModelingToolsetExp/Source/MeshModelingToolsExp/Private/RemeshMeshTool.cpp
--- a/Engine/Plugins/Experimental/MeshModelingToolsetExp/Source/MeshModelingToolsExp/Private/RemeshMeshTool.cpp
+++ b/Engine/Plugins/Experimental/MeshModelingToolsetExp/Source/MeshModelingToolsExp/Private/RemeshMeshTool.cpp
@@ -113,19 +113,14 @@ void URemeshMeshTool::Setup()
 
 	OriginalMeshSpatial = MakeShared<FDynamicMeshAABBTree3, ESPMode::ThreadSafe>(OriginalMesh.Get(), true);
 
-	// calculate initial mesh area (no utility fn yet)
-	// TODO: will need to change to account for component transform's Scale3D
-	InitialMeshArea = 0;
-	for (int tid : OriginalMesh->TriangleIndicesItr())
-	{
-		InitialMeshArea += OriginalMesh->GetTriArea(tid);
-	}
+	// calculate initial mesh area
+	InitialMeshArea = TMeshQueries<FDynamicMesh3>::GetVolumeArea(*OriginalMesh).Y;
 
 	// set properties defaults
 
 	// arbitrary threshold of 5000 tris seems reasonable?
 	BasicProperties->TargetTriangleCount = (OriginalMesh->TriangleCount() < 5000) ? 5000 : OriginalMesh->TriangleCount();
-	BasicProperties->TargetEdgeLength = CalculateTargetEdgeLength(BasicProperties->TargetTriangleCount);
+	BasicProperties->TargetEdgeLength = FRemeshMeshOp::CalculateTargetEdgeLength(OriginalMesh.Get(), BasicProperties->TargetTriangleCount, InitialMeshArea);
 
 	// add properties to GUI
 	AddToolPropertySource(BasicProperties);
@@ -201,7 +196,7 @@ TUniquePtr<FDynamicMeshOperator> URemeshMeshTool::MakeNewOperator()
 
 	if (!BasicProperties->bUseTargetEdgeLength)
 	{
-		Op->TargetEdgeLength = CalculateTargetEdgeLength(BasicProperties->TargetTriangleCount);
+		Op->TargetEdgeLength = FRemeshMeshOp::CalculateTargetEdgeLength(OriginalMesh.Get(), BasicProperties->TargetTriangleCount, InitialMeshArea);
 		Op->TriangleCountHint = 2.0 * BasicProperties->TargetTriangleCount;
 	}
 	else
@@ -275,13 +270,6 @@ void URemeshMeshTool::UpdateVisualization()
 								ToolSetupUtil::GetDefaultWorkingMaterial(GetToolManager()));
 }
 
-double URemeshMeshTool::CalculateTargetEdgeLength(int TargetTriCount)
-{
-	double TargetTriArea = InitialMeshArea / (double)TargetTriCount;
-	double EdgeLen = TriangleUtil::EquilateralEdgeLengthForArea(TargetTriArea);
-	return (double)FMath::RoundToInt(EdgeLen*100.0) / 100.0;
-}
-
 bool URemeshMeshTool::CanAccept() const
 {
 	return Super::CanAccept() && Preview->HaveValidResult();
diff --git a/Engine/Plugins/Experimental/MeshModelingToolsetExp/Source/MeshModelingToolsExp/Public/RemeshMeshTool.h b/Engine/Plugins/Experimental/MeshModelingToolsetExp/Source/MeshModelingToolsExp/Public/RemeshMeshTool.h
--- a/Engine/Plugins/Experimental/MeshModelingToolsetExp/Source/MeshModelingToolsExp/Public/RemeshMeshTool.h
+++ b/Engine/Plugins/Experimental/MeshModelingToolsetExp/Source/MeshModelingToolsExp/Public/RemeshMeshTool.h
@@ -147,6 +147,5 @@ protected:
 	TSharedPtr<UE::Geometry::FDynamicMeshAABBTree3, ESPMode::ThreadSafe> OriginalMeshSpatial;
 	double InitialMeshArea;
 
-	double CalculateTargetEdgeLength(int TargetTriCount);
 	void UpdateVisualization();
 };
diff --git a/Engine/Plugins/Runtime/GeometryProcessing/Source/DynamicMesh/Private/NormalFlowRemesher.cpp b/Engine/Plugins/Runtime/GeometryProcessing/Source/DynamicMesh/Private/NormalFlowRemesher.cpp
--- a/Engine/Plugins/Runtime/GeometryProcessing/Source/DynamicMesh/Private/NormalFlowRemesher.cpp
+++ b/Engine/Plugins/Runtime/GeometryProcessing/Source/DynamicMesh/Private/NormalFlowRemesher.cpp
@@ -54,7 +54,8 @@ void FNormalFlowRemesher::RemeshWithFaceProjection()
 			break;
 		}
 
-		RemeshIteration();
+		// currently disabling converge check
+		RemeshIteration( []() {return false; } );
 
 		if (Iterations > MaxRemeshIterations / 2)
 		{
diff --git a/Engine/Plugins/Runtime/GeometryProcessing/Source/DynamicMesh/Private/QueueRemesher.cpp b/Engine/Plugins/Runtime/GeometryProcessing/Source/DynamicMesh/Private/QueueRemesher.cpp
--- a/Engine/Plugins/Runtime/GeometryProcessing/Source/DynamicMesh/Private/QueueRemesher.cpp
+++ b/Engine/Plugins/Runtime/GeometryProcessing/Source/DynamicMesh/Private/QueueRemesher.cpp
@@ -13,8 +13,6 @@ void FQueueRemesher::FastestRemesh()
 		ensureMsgf(false, TEXT("Input Mesh has Attribute overlays but no Constraints are configured. Use FMeshConstraintsUtil::ConstrainAllBoundariesAndSeams() to create a Constraint Set for Attribute seams."));
 	}
 
-	ModifiedEdgesLastPass = 0;
-
 	ResetQueue();
 
 	// First we do fast splits to hit edge length target
@@ -27,6 +25,7 @@ void FQueueRemesher::FastestRemesh()
 		}
 
 		int nSplits = FastSplitIteration();
+		ModifiedEdgesLastPass = nSplits;
 
 		if ((double)nSplits / (double)Mesh->EdgeCount() < 0.01)
 		{
@@ -47,8 +46,20 @@ void FQueueRemesher::FastestRemesh()
 			break;
 		}
 
+		// if fraction of modified edges falls below threshold, consider the result converged and terminate
+		auto EarlyTerminationCheck = [k, this]() {
+			double ModifiedEdgeFraction = (double)ModifiedEdgesLastPass / (double)Mesh->EdgeCount();
+			return (k > 10 && ModifiedEdgeFraction < MinActiveEdgeFraction);
+		};
+
 		ProjectionMode = (k % 2 == 0) ? ETargetProjectionMode::NoProjection : OriginalProjectionMode;
-		RemeshIteration();
+		RemeshIteration(EarlyTerminationCheck);
+
+		// If fraction of active edges falls below threshold, terminate early
+		if (EarlyTerminationCheck())
+		{
+			break;
+		}
 	}
 
 	if (Cancelled())
@@ -58,7 +69,7 @@ void FQueueRemesher::FastestRemesh()
 
 	// Final pass with full projection
 	ProjectionMode = OriginalProjectionMode;
-	RemeshIteration();
+	RemeshIteration( []() { return false; } );
 }
 
 int FQueueRemesher::FastSplitIteration()
@@ -164,8 +175,10 @@ int FQueueRemesher::FastSplitIteration()
 }
 
 
-void FQueueRemesher::RemeshIteration()
+void FQueueRemesher::RemeshIteration(TFunctionRef<bool(void)> EarlyTerminationCheck)
 {
+	ModifiedEdgesLastPass = 0;
+
 	if (Mesh->TriangleCount() == 0)
 	{
 		return;
@@ -192,7 +205,7 @@ void FQueueRemesher::RemeshIteration()
 		}
 	};
 
-
+	int32 ProcessedEdges = 0;
 	if (!ModifiedEdges)
 	{
 		// FIRST ITERATION
@@ -226,7 +239,7 @@ void FQueueRemesher::RemeshIteration()
 		EdgeBuffer = MoveTemp(*ModifiedEdges);
 		ModifiedEdges->Reset();
 
-		for (auto CurrentEdgeID : EdgeBuffer)
+		for (int CurrentEdgeID : EdgeBuffer)
 		{
 			if (Cancelled())
 			{
@@ -244,6 +257,10 @@ void FQueueRemesher::RemeshIteration()
 
 	ProfileEndOps();
 
+	if (EarlyTerminationCheck())
+	{
+		return;
+	}
 	if (Cancelled())
 	{
 		return;
diff --git a/Engine/Plugins/Runtime/GeometryProcessing/Source/DynamicMesh/Public/QueueRemesher.h b/Engine/Plugins/Runtime/GeometryProcessing/Source/DynamicMesh/Public/QueueRemesher.h
--- a/Engine/Plugins/Runtime/GeometryProcessing/Source/DynamicMesh/Public/QueueRemesher.h
+++ b/Engine/Plugins/Runtime/GeometryProcessing/Source/DynamicMesh/Public/QueueRemesher.h
@@ -34,6 +34,9 @@ public:
 	/// Max number of passes doing full remeshing operations
 	int MaxRemeshIterations = 20;
 
+	/// If fraction of active edges falls below this threshold, consider result converged and terminate before MaxRemeshIterations
+	double MinActiveEdgeFraction = 0.01;
+
 	/// Converge on remeshed result as quickly as possible
 	void FastestRemesh();
 
@@ -139,7 +142,10 @@ protected:
 	int FastSplitIteration();
 
 	/// This pass does all edge operations plug smoothing and reprojection.
-	void RemeshIteration();
+	/// EarlyTerminationCheck will be called after remesh operations, before smoothing/reprojection. Caller can return true to abort further computation.
+	void RemeshIteration(
+		TFunctionRef<bool(void)> EarlyTerminationCheck
+	);
 
 	// Whether an edge should be added to ModifiedEdges (used in async tasks)
 	TArray<bool> EdgeShouldBeQueuedBuffer;	
diff --git a/Engine/Plugins/Runtime/MeshModelingToolset/Source/ModelingOperators/Private/CleaningOps/RemeshMeshOp.cpp b/Engine/Plugins/Runtime/MeshModelingToolset/Source/ModelingOperators/Private/CleaningOps/RemeshMeshOp.cpp
--- a/Engine/Plugins/Runtime/MeshModelingToolset/Source/ModelingOperators/Private/CleaningOps/RemeshMeshOp.cpp
+++ b/Engine/Plugins/Runtime/MeshModelingToolset/Source/ModelingOperators/Private/CleaningOps/RemeshMeshOp.cpp
@@ -23,6 +23,7 @@ TUniquePtr<FRemesher> FRemeshMeshOp::CreateRemesher(ERemeshType Type, FDynamicMe
 	{
 		TUniquePtr<FQueueRemesher> QueueRemesher = MakeUnique<FQueueRemesher>(TargetMesh);
 		QueueRemesher->MaxRemeshIterations = MaxRemeshIterations;
+		QueueRemesher->MinActiveEdgeFraction = MinActiveEdgeFraction;
 		return QueueRemesher;
 	}
 	case ERemeshType::FullPass:
@@ -33,6 +34,7 @@ TUniquePtr<FRemesher> FRemeshMeshOp::CreateRemesher(ERemeshType Type, FDynamicMe
 	{
 		TUniquePtr<FNormalFlowRemesher> NormalFlowRemesher = MakeUnique<FNormalFlowRemesher>(TargetMesh);
 		NormalFlowRemesher->MaxRemeshIterations = MaxRemeshIterations;
+		NormalFlowRemesher->MinActiveEdgeFraction = 0;		// disable convergence check for NormalFlow remeshing
 		NormalFlowRemesher->NumExtraProjectionIterations = ExtraProjectionIterations;
 		NormalFlowRemesher->MaxTriangleCount = FMath::Max(0, TriangleCountHint);
 		NormalFlowRemesher->FaceProjectionPassesPerRemeshIteration = FaceProjectionPassesPerRemeshIteration;
@@ -237,3 +239,18 @@ void FRemeshMeshOp::CalculateResult(FProgressCancel* Progress)
 		}
 	}
 }
+
+
+
+double FRemeshMeshOp::CalculateTargetEdgeLength(const FDynamicMesh3* Mesh, int TargetTriCount, double PrecomputedMeshArea)
+{
+	double InitialMeshArea = PrecomputedMeshArea;
+	if (InitialMeshArea <= 0)
+	{
+		InitialMeshArea = TMeshQueries<FDynamicMesh3>::GetVolumeArea(*Mesh).Y;
+	}
+
+	double TargetTriArea = InitialMeshArea / (double)TargetTriCount;
+	double EdgeLen = TriangleUtil::EquilateralEdgeLengthForArea(TargetTriArea);
+	return (double)FMath::RoundToInt(EdgeLen*100.0) / 100.0;
+}
\ No newline at end of file
diff --git a/Engine/Plugins/Runtime/MeshModelingToolset/Source/ModelingOperators/Public/CleaningOps/RemeshMeshOp.h b/Engine/Plugins/Runtime/MeshModelingToolset/Source/ModelingOperators/Public/CleaningOps/RemeshMeshOp.h
--- a/Engine/Plugins/Runtime/MeshModelingToolset/Source/ModelingOperators/Public/CleaningOps/RemeshMeshOp.h
+++ b/Engine/Plugins/Runtime/MeshModelingToolset/Source/ModelingOperators/Public/CleaningOps/RemeshMeshOp.h
@@ -58,13 +58,22 @@ public:
 
 	ERemeshType RemeshType = ERemeshType::Standard;
 
-	int RemeshIterations;
-	int MaxRemeshIterations;
+	int RemeshIterations = 20;
+	int MaxRemeshIterations = 20;
+	double MinActiveEdgeFraction = 0.01;		// terminate remeshing if modified edge count in last pass is below this fraction of total edges
 	int ExtraProjectionIterations = 5;
 	int TriangleCountHint = 0;
-	float SmoothingStrength, TargetEdgeLength;
+	float SmoothingStrength = 0.25f;
+	double TargetEdgeLength = 1.0f;
 	ERemeshSmoothingType SmoothingType;
-	bool bDiscardAttributes, bPreserveSharpEdges, bFlips, bSplits, bCollapses, bReproject, bPreventNormalFlips, bPreventTinyTriangles;
+	bool bDiscardAttributes = false;
+	bool bPreserveSharpEdges = true;
+	bool bFlips = true;
+	bool bSplits = true;
+	bool bCollapses = true;
+	bool bReproject = true;
+	bool bPreventNormalFlips = true;
+	bool bPreventTinyTriangles = true;
 	bool bReprojectConstraints = false;
 	double BoundaryCornerAngleThreshold = 45.0;
 
@@ -76,8 +85,8 @@ public:
 	FDynamicMesh3* ProjectionTarget = nullptr;
 	FDynamicMeshAABBTree3* ProjectionTargetSpatial = nullptr;
 
-	FTransformSRT3d TargetMeshLocalToWorld;
-	FTransformSRT3d ToolMeshLocalToWorld;
+	FTransformSRT3d TargetMeshLocalToWorld = FTransformSRT3d::Identity();
+	FTransformSRT3d ToolMeshLocalToWorld = FTransformSRT3d::Identity();
 	bool bUseWorldSpace = false;
 	bool bParallel = true;
 
@@ -112,6 +121,11 @@ public:
 	// 
 	virtual void CalculateResult(FProgressCancel* Progress) override;
 
+
+public:
+	static double CalculateTargetEdgeLength(const FDynamicMesh3* Mesh, int TargetTriCount, double PrecomputedMeshArea = 0.0);
+
+
 protected:
 
 	TUniquePtr<FRemesher> CreateRemesher(ERemeshType Type, FDynamicMesh3* TargetMesh);
