GeometryScriptingCore:  MeshBasicEditFunctions.
Fix a crash when creating a mesh from scratch and adding UVs.

diff --git a/Engine/Plugins/Experimental/GeometryScripting/Source/GeometryScriptingCore/Private/MeshBasicEditFunctions.cpp b/Engine/Plugins/Experimental/GeometryScripting/Source/GeometryScriptingCore/Private/MeshBasicEditFunctions.cpp
--- a/Engine/Plugins/Experimental/GeometryScripting/Source/GeometryScriptingCore/Private/MeshBasicEditFunctions.cpp
+++ b/Engine/Plugins/Experimental/GeometryScripting/Source/GeometryScriptingCore/Private/MeshBasicEditFunctions.cpp
@@ -476,7 +476,7 @@ UDynamicMesh* UGeometryScriptLibrary_MeshBasicEditFunctions::AppendBuffersToMesh
 		if (Buffers.Normals.Num() == NumVertices)
 		{
 			FDynamicMeshNormalOverlay* Normals = EditMesh.Attributes()->PrimaryNormals();
-			VertexIDMap.Reset();
+			VertexIDMap.SetNum(NumVertices);
 			for (int32 k = 0; k < NumVertices; ++k)
 			{
 				int32 NewElementID = Normals->AppendElement((FVector3f)Buffers.Normals[k]);
@@ -505,7 +505,7 @@ UDynamicMesh* UGeometryScriptLibrary_MeshBasicEditFunctions::AppendBuffersToMesh
 		for (int32 li = 0; li < NumUVLayers; ++li)
 		{
 			FDynamicMeshUVOverlay* UVs = EditMesh.Attributes()->GetUVLayer(li);
-			VertexIDMap.Reset();
+			VertexIDMap.SetNum(NumVertices);
 			for (int32 k = 0; k < NumVertices; ++k)
 			{
 				int32 NewElementID = UVs->AppendElement( (FVector2f) (*AllUVSets[li])[k] );
@@ -524,7 +524,7 @@ UDynamicMesh* UGeometryScriptLibrary_MeshBasicEditFunctions::AppendBuffersToMesh
 		{
 			EditMesh.Attributes()->EnablePrimaryColors();
 			FDynamicMeshColorOverlay* Colors = EditMesh.Attributes()->PrimaryColors();
-			VertexIDMap.Reset();
+			VertexIDMap.SetNum(NumVertices);
 			for (int32 k = 0; k < NumVertices; ++k)
 			{
 				int32 NewElementID = Colors->AppendElement( ToVector4<float>(Buffers.VertexColors[k]) );
