GeometryScript option struct access specifiers

diff --git a/Engine/Plugins/Experimental/GeometryScripting/Source/GeometryScriptingCore/Public/GeometryScript/CollisionFunctions.h b/Engine/Plugins/Experimental/GeometryScripting/Source/GeometryScriptingCore/Public/GeometryScript/CollisionFunctions.h
--- a/Engine/Plugins/Experimental/GeometryScripting/Source/GeometryScriptingCore/Public/GeometryScript/CollisionFunctions.h
+++ b/Engine/Plugins/Experimental/GeometryScripting/Source/GeometryScriptingCore/Public/GeometryScript/CollisionFunctions.h
@@ -42,53 +42,53 @@ struct GEOMETRYSCRIPTINGCORE_API FGeometryScriptCollisionFromMeshOptions
 {
 	GENERATED_BODY()
 public:
-	UPROPERTY(BlueprintReadWrite, Category = Options)
+	UPROPERTY(EditAnywhere, BlueprintReadWrite, Category = Options)
 	bool bEmitTransaction = true;
 
-	UPROPERTY(BlueprintReadWrite, Category = Options)
+	UPROPERTY(EditAnywhere, BlueprintReadWrite, Category = Options)
 	EGeometryScriptCollisionGenerationMethod Method = EGeometryScriptCollisionGenerationMethod::MinVolumeShapes;
 
-	UPROPERTY(BlueprintReadWrite, Category = Options)
+	UPROPERTY(EditAnywhere, BlueprintReadWrite, Category = Options)
 	bool bAutoDetectSpheres = true;
 
-	UPROPERTY(BlueprintReadWrite, Category = Options)
+	UPROPERTY(EditAnywhere, BlueprintReadWrite, Category = Options)
 	bool bAutoDetectBoxes = true;
 
-	UPROPERTY(BlueprintReadWrite, Category = Options)
+	UPROPERTY(EditAnywhere, BlueprintReadWrite, Category = Options)
 	bool bAutoDetectCapsules = true;
 
-	UPROPERTY(BlueprintReadWrite, Category = Options)
+	UPROPERTY(EditAnywhere, BlueprintReadWrite, Category = Options)
 	float MinThickness = 1.0;
 
-	UPROPERTY(BlueprintReadWrite, Category = Options)
+	UPROPERTY(EditAnywhere, BlueprintReadWrite, Category = Options)
 	bool bSimplifyHulls = true;
 
-	UPROPERTY(BlueprintReadWrite, Category = Options)
+	UPROPERTY(EditAnywhere, BlueprintReadWrite, Category = Options)
 	int ConvexHullTargetFaceCount = 25;
 
-	UPROPERTY(BlueprintReadWrite, Category = Options)
+	UPROPERTY(EditAnywhere, BlueprintReadWrite, Category = Options)
 	int MaxConvexHullsPerMesh = 1;
 
-	UPROPERTY(BlueprintReadWrite, Category = Options)
+	UPROPERTY(EditAnywhere, BlueprintReadWrite, Category = Options)
 	float ConvexDecompositionSearchFactor = .5f;
 
-	UPROPERTY(BlueprintReadWrite, Category = Options)
+	UPROPERTY(EditAnywhere, BlueprintReadWrite, Category = Options)
 	float ConvexDecompositionErrorTolerance = 0;
 
-	UPROPERTY(BlueprintReadWrite, Category = Options)
+	UPROPERTY(EditAnywhere, BlueprintReadWrite, Category = Options)
 	float ConvexDecompositionMinPartThickness = 0.1;
 
-	UPROPERTY(BlueprintReadWrite, Category = Options)
+	UPROPERTY(EditAnywhere, BlueprintReadWrite, Category = Options)
 	float SweptHullSimplifyTolerance = 0.1;
 
-	UPROPERTY(BlueprintReadWrite, Category = Options)
+	UPROPERTY(EditAnywhere, BlueprintReadWrite, Category = Options)
 	EGeometryScriptSweptHullAxis SweptHullAxis = EGeometryScriptSweptHullAxis::Z;
 
 
-	UPROPERTY(BlueprintReadWrite, Category = Options)
+	UPROPERTY(EditAnywhere, BlueprintReadWrite, Category = Options)
 	bool bRemoveFullyContainedShapes = true;
 
-	UPROPERTY(BlueprintReadWrite, Category = Options)
+	UPROPERTY(EditAnywhere, BlueprintReadWrite, Category = Options)
 	int MaxShapeCount = 0;
 };
 
@@ -98,7 +98,7 @@ struct GEOMETRYSCRIPTINGCORE_API FGeometryScriptSetSimpleCollisionOptions
 {
 	GENERATED_BODY()
 public:
-	UPROPERTY(BlueprintReadWrite, Category = Options)
+	UPROPERTY(EditAnywhere, BlueprintReadWrite, Category = Options)
 	bool bEmitTransaction = true;
 };
 
diff --git a/Engine/Plugins/Experimental/GeometryScripting/Source/GeometryScriptingCore/Public/GeometryScript/ContainmentFunctions.h b/Engine/Plugins/Experimental/GeometryScripting/Source/GeometryScriptingCore/Public/GeometryScript/ContainmentFunctions.h
--- a/Engine/Plugins/Experimental/GeometryScripting/Source/GeometryScriptingCore/Public/GeometryScript/ContainmentFunctions.h
+++ b/Engine/Plugins/Experimental/GeometryScripting/Source/GeometryScriptingCore/Public/GeometryScript/ContainmentFunctions.h
@@ -14,14 +14,14 @@ struct GEOMETRYSCRIPTINGCORE_API FGeometryScriptConvexHullOptions
 {
 	GENERATED_BODY()
 public:
-	UPROPERTY(BlueprintReadWrite, Category = Options)
+	UPROPERTY(EditAnywhere, BlueprintReadWrite, Category = Options)
 	bool bPrefilterVertices = true;
 
-	UPROPERTY(BlueprintReadWrite, Category = Options)
+	UPROPERTY(EditAnywhere, BlueprintReadWrite, Category = Options)
 	int PrefilterGridResolution = 128;
 
 	/** */
-	UPROPERTY(BlueprintReadWrite, Category = Options)
+	UPROPERTY(EditAnywhere, BlueprintReadWrite, Category = Options)
 	int SimplifyToFaceCount = 0;
 };
 
@@ -32,23 +32,23 @@ struct GEOMETRYSCRIPTINGCORE_API FGeometryScriptSweptHullOptions
 {
 	GENERATED_BODY()
 public:
-	UPROPERTY(BlueprintReadWrite, Category = Options)
+	UPROPERTY(EditAnywhere, BlueprintReadWrite, Category = Options)
 	bool bPrefilterVertices = true;
 
-	UPROPERTY(BlueprintReadWrite, Category = Options)
+	UPROPERTY(EditAnywhere, BlueprintReadWrite, Category = Options)
 	int PrefilterGridResolution = 128;
 
-	UPROPERTY(BlueprintReadWrite, Category = Options)
+	UPROPERTY(EditAnywhere, BlueprintReadWrite, Category = Options)
 	float MinThickness = 0.01;
 
 	/** */
-	UPROPERTY(BlueprintReadWrite, Category = Options)
+	UPROPERTY(EditAnywhere, BlueprintReadWrite, Category = Options)
 	bool bSimplify = true;
 
-	UPROPERTY(BlueprintReadWrite, Category = Options)
+	UPROPERTY(EditAnywhere, BlueprintReadWrite, Category = Options)
 	float MinEdgeLength = 0.1;
 
-	UPROPERTY(BlueprintReadWrite, Category = Options)
+	UPROPERTY(EditAnywhere, BlueprintReadWrite, Category = Options)
 	float SimplifyTolerance = 0.1;
 };
 
diff --git a/Engine/Plugins/Experimental/GeometryScripting/Source/GeometryScriptingCore/Public/GeometryScript/MeshAssetFunctions.h b/Engine/Plugins/Experimental/GeometryScripting/Source/GeometryScriptingCore/Public/GeometryScript/MeshAssetFunctions.h
--- a/Engine/Plugins/Experimental/GeometryScripting/Source/GeometryScriptingCore/Public/GeometryScript/MeshAssetFunctions.h
+++ b/Engine/Plugins/Experimental/GeometryScripting/Source/GeometryScriptingCore/Public/GeometryScript/MeshAssetFunctions.h
@@ -18,13 +18,13 @@ struct GEOMETRYSCRIPTINGCORE_API FGeometryScriptCopyMeshFromAssetOptions
 {
 	GENERATED_BODY()
 public:
-	UPROPERTY(BlueprintReadWrite, Category = Options)
+	UPROPERTY(EditAnywhere, BlueprintReadWrite, Category = Options)
 	bool bApplyBuildSettings = true;
 
-	UPROPERTY(BlueprintReadWrite, Category = Options)
+	UPROPERTY(EditAnywhere, BlueprintReadWrite, Category = Options)
 	bool bRequestTangents = true;
 
-	UPROPERTY(BlueprintReadWrite, Category = Options)
+	UPROPERTY(EditAnywhere, BlueprintReadWrite, Category = Options)
 	bool bIgnoreRemoveDegenerates = true;
 };
 
@@ -37,15 +37,15 @@ struct GEOMETRYSCRIPTINGCORE_API FGeometryScriptNaniteOptions
 	GENERATED_BODY()
 
 	/** Set Nanite to Enabled/Disabled */
-	UPROPERTY(BlueprintReadWrite, Category = Options)
+	UPROPERTY(EditAnywhere, BlueprintReadWrite, Category = Options)
 	bool bEnabled = true;
 
 	/** Percentage of triangles to maintain in Fallback Mesh used when Nanite is unavailable */
-	UPROPERTY(BlueprintReadWrite, Category = Options)
+	UPROPERTY(EditAnywhere, BlueprintReadWrite, Category = Options)
 	float FallbackPercentTriangles = 100.0f;
 
 	/** Relative Error to maintain in Fallback Mesh used when Nanite is unavailable. Overrides FallbackPercentTriangles. Set to 0 to only use FallbackPercentTriangles (default). */
-	UPROPERTY(BlueprintReadWrite, Category = Options)
+	UPROPERTY(EditAnywhere, BlueprintReadWrite, Category = Options)
 	float FallbackRelativeError = 0.0f;
 };
 
@@ -55,37 +55,37 @@ struct GEOMETRYSCRIPTINGCORE_API FGeometryScriptCopyMeshToAssetOptions
 {
 	GENERATED_BODY()
 public:
-	UPROPERTY(BlueprintReadWrite, Category = Options)
+	UPROPERTY(EditAnywhere, BlueprintReadWrite, Category = Options)
 	bool bEnableRecomputeNormals = false;
 
-	UPROPERTY(BlueprintReadWrite, Category = Options)
+	UPROPERTY(EditAnywhere, BlueprintReadWrite, Category = Options)
 	bool bEnableRecomputeTangents = false;
 
-	UPROPERTY(BlueprintReadWrite, Category = Options)
+	UPROPERTY(EditAnywhere, BlueprintReadWrite, Category = Options)
 	bool bEnableRemoveDegenerates = false;
 
 	
-	UPROPERTY(BlueprintReadWrite, Category = Options)
+	UPROPERTY(EditAnywhere, BlueprintReadWrite, Category = Options)
 	bool bReplaceMaterials = false;
 
-	UPROPERTY(BlueprintReadWrite, Category = Options)
+	UPROPERTY(EditAnywhere, BlueprintReadWrite, Category = Options)
 	TArray<UMaterialInterface*> NewMaterials;
 
-	UPROPERTY(BlueprintReadWrite, Category = Options)
+	UPROPERTY(EditAnywhere, BlueprintReadWrite, Category = Options)
 	TArray<FName> NewMaterialSlotNames;
 
 	/** If enabled, NaniteSettings will be applied to the target Asset if possible */
-	UPROPERTY(BlueprintReadWrite, Category = Options)
+	UPROPERTY(EditAnywhere, BlueprintReadWrite, Category = Options)
 	bool bApplyNaniteSettings = false;
 
 	/** Nanite Settings applied to the target Asset, if bApplyNaniteSettings = true */
-	UPROPERTY(BlueprintReadWrite, Category = Options)
+	UPROPERTY(EditAnywhere, BlueprintReadWrite, Category = Options)
 	FGeometryScriptNaniteOptions NaniteSettings = FGeometryScriptNaniteOptions();
 
-	UPROPERTY(BlueprintReadWrite, Category = Options)
+	UPROPERTY(EditAnywhere, BlueprintReadWrite, Category = Options)
 	bool bEmitTransaction = true;
 
-	UPROPERTY(BlueprintReadWrite, Category = Options)
+	UPROPERTY(EditAnywhere, BlueprintReadWrite, Category = Options)
 	bool bDeferMeshPostEditChange = false;
 };
 
diff --git a/Engine/Plugins/Experimental/GeometryScripting/Source/GeometryScriptingCore/Public/GeometryScript/MeshBakeFunctions.h b/Engine/Plugins/Experimental/GeometryScripting/Source/GeometryScriptingCore/Public/GeometryScript/MeshBakeFunctions.h
--- a/Engine/Plugins/Experimental/GeometryScripting/Source/GeometryScriptingCore/Public/GeometryScript/MeshBakeFunctions.h
+++ b/Engine/Plugins/Experimental/GeometryScripting/Source/GeometryScriptingCore/Public/GeometryScript/MeshBakeFunctions.h
@@ -199,29 +199,29 @@ struct GEOMETRYSCRIPTINGCORE_API FGeometryScriptBakeTextureOptions
 	GENERATED_BODY()
 
 	/** The pixel resolution of the generated textures */
-	UPROPERTY(BlueprintReadWrite, Category = Options)
+	UPROPERTY(EditAnywhere, BlueprintReadWrite, Category = Options)
 	EGeometryScriptBakeResolution Resolution = EGeometryScriptBakeResolution::Resolution256;
 
 	/** The bit depth for each channel of the generated textures */
-	UPROPERTY(BlueprintReadWrite, Category = Options)
+	UPROPERTY(EditAnywhere, BlueprintReadWrite, Category = Options)
 	EGeometryScriptBakeBitDepth BitDepth = EGeometryScriptBakeBitDepth::ChannelBits8;
 
 	/** Number of samples per pixel */
-	UPROPERTY(BlueprintReadWrite, Category = Options)
+	UPROPERTY(EditAnywhere, BlueprintReadWrite, Category = Options)
 	EGeometryScriptBakeSamplesPerPixel SamplesPerPixel = EGeometryScriptBakeSamplesPerPixel::Sample1;
 
 	/** Mask texture for filtering out samples/pixels from the output texture */
-	UPROPERTY(BlueprintReadWrite, Category = Options)
+	UPROPERTY(EditAnywhere, BlueprintReadWrite, Category = Options)
 	TObjectPtr<UTexture2D> SampleFilterMask = nullptr;
 
 	/** Maximum allowed distance for the projection from target mesh to source mesh for the sample to be considered valid.
 	 * This is only relevant if a separate source mesh is provided. */
-	UPROPERTY(BlueprintReadWrite, Category = Options)
+	UPROPERTY(EditAnywhere, BlueprintReadWrite, Category = Options)
 	float ProjectionDistance = 3.0f;
 
 	/** If true, uses the world space positions for the projection from target mesh to source mesh, otherwise it uses their object space positions.
 	 * This is only relevant if a separate source mesh is provided. */
-	UPROPERTY(BlueprintReadWrite, Category = Options)
+	UPROPERTY(EditAnywhere, BlueprintReadWrite, Category = Options)
 	bool bProjectionInWorldSpace = false;
 };
 
@@ -231,7 +231,7 @@ struct GEOMETRYSCRIPTINGCORE_API FGeometryScriptBakeVertexOptions
 	GENERATED_BODY()
 
 	/** If true, compute a separate vertex color for each unique normal on a vertex */
-	UPROPERTY(BlueprintReadWrite, Category = Options)
+	UPROPERTY(EditAnywhere, BlueprintReadWrite, Category = Options)
 	bool bSplitAtNormalSeams = false;
 
 	/** If true, compute a separate vertex color for each unique UV on a vertex. */
@@ -240,12 +240,12 @@ struct GEOMETRYSCRIPTINGCORE_API FGeometryScriptBakeVertexOptions
 
 	/** Maximum allowed distance for the projection from target mesh to source mesh for the sample to be considered valid.
 	 * This is only relevant if a separate source mesh is provided. */
-	UPROPERTY(BlueprintReadWrite, Category = Options)
+	UPROPERTY(EditAnywhere, BlueprintReadWrite, Category = Options)
 	float ProjectionDistance = 3.0f;
 
 	/** If true, uses the world space positions for the projection from target mesh to source mesh, otherwise it uses their object space positions.
 	 * This is only relevant if a separate source mesh is provided. */
-	UPROPERTY(BlueprintReadWrite, Category = Options)
+	UPROPERTY(EditAnywhere, BlueprintReadWrite, Category = Options)
 	bool bProjectionInWorldSpace = false;
 };
 
@@ -255,7 +255,7 @@ struct GEOMETRYSCRIPTINGCORE_API FGeometryScriptBakeOutputType
 	GENERATED_BODY()
 	
 	/** The bake output mode */
-	UPROPERTY(BlueprintReadWrite, Category = Options)
+	UPROPERTY(EditAnywhere, BlueprintReadWrite, Category = Options)
 	EGeometryScriptBakeOutputMode OutputMode = EGeometryScriptBakeOutputMode::RGBA;
 
 	UPROPERTY(BlueprintReadWrite, Category = Output)
@@ -288,13 +288,13 @@ struct GEOMETRYSCRIPTINGCORE_API FGeometryScriptBakeSourceMeshOptions
 {
 	GENERATED_BODY()
 
-	UPROPERTY(BlueprintReadWrite, Category = Options)
+	UPROPERTY(EditAnywhere, BlueprintReadWrite, Category = Options)
 	TObjectPtr<UTexture2D> SourceNormalMap = nullptr;
 
 	UPROPERTY(BlueprintReadWrite, Category = Options, meta = (DisplayName="Source Normal UV Channel"))
 	int SourceNormalUVLayer = 0;
 
-	UPROPERTY(BlueprintReadWrite, Category = Options)
+	UPROPERTY(EditAnywhere, BlueprintReadWrite, Category = Options)
 	EGeometryScriptBakeNormalSpace SourceNormalSpace = EGeometryScriptBakeNormalSpace::Tangent;
 };
 
diff --git a/Engine/Plugins/Experimental/GeometryScripting/Source/GeometryScriptingCore/Public/GeometryScript/MeshBasicEditFunctions.h b/Engine/Plugins/Experimental/GeometryScripting/Source/GeometryScriptingCore/Public/GeometryScript/MeshBasicEditFunctions.h
--- a/Engine/Plugins/Experimental/GeometryScripting/Source/GeometryScriptingCore/Public/GeometryScript/MeshBasicEditFunctions.h
+++ b/Engine/Plugins/Experimental/GeometryScripting/Source/GeometryScriptingCore/Public/GeometryScript/MeshBasicEditFunctions.h
@@ -15,37 +15,37 @@ struct GEOMETRYSCRIPTINGCORE_API FGeometryScriptSimpleMeshBuffers
 {
 	GENERATED_BODY()
 public:
-	UPROPERTY(BlueprintReadWrite, Category = Options)
+	UPROPERTY(EditAnywhere, BlueprintReadWrite, Category = Options)
 	TArray<FVector> Vertices;
 
-	UPROPERTY(BlueprintReadWrite, Category = Options)
+	UPROPERTY(EditAnywhere, BlueprintReadWrite, Category = Options)
 	TArray<FVector> Normals;
 
-	UPROPERTY(BlueprintReadWrite, Category = Options)
+	UPROPERTY(EditAnywhere, BlueprintReadWrite, Category = Options)
 	TArray<FVector2D> UV0;
-	UPROPERTY(BlueprintReadWrite, Category = Options)
+	UPROPERTY(EditAnywhere, BlueprintReadWrite, Category = Options)
 	TArray<FVector2D> UV1;
-	UPROPERTY(BlueprintReadWrite, Category = Options)
+	UPROPERTY(EditAnywhere, BlueprintReadWrite, Category = Options)
 	TArray<FVector2D> UV2;
-	UPROPERTY(BlueprintReadWrite, Category = Options)
+	UPROPERTY(EditAnywhere, BlueprintReadWrite, Category = Options)
 	TArray<FVector2D> UV3;
-	UPROPERTY(BlueprintReadWrite, Category = Options)
+	UPROPERTY(EditAnywhere, BlueprintReadWrite, Category = Options)
 	TArray<FVector2D> UV4;
-	UPROPERTY(BlueprintReadWrite, Category = Options)
+	UPROPERTY(EditAnywhere, BlueprintReadWrite, Category = Options)
 	TArray<FVector2D> UV5;
-	UPROPERTY(BlueprintReadWrite, Category = Options)
+	UPROPERTY(EditAnywhere, BlueprintReadWrite, Category = Options)
 	TArray<FVector2D> UV6;
-	UPROPERTY(BlueprintReadWrite, Category = Options)
+	UPROPERTY(EditAnywhere, BlueprintReadWrite, Category = Options)
 	TArray<FVector2D> UV7;
 
-	UPROPERTY(BlueprintReadWrite, Category = Options)
+	UPROPERTY(EditAnywhere, BlueprintReadWrite, Category = Options)
 	TArray<FLinearColor> VertexColors;
 
 
-	UPROPERTY(BlueprintReadWrite, Category = Options)
+	UPROPERTY(EditAnywhere, BlueprintReadWrite, Category = Options)
 	TArray<FIntVector> Triangles;
 
-	UPROPERTY(BlueprintReadWrite, Category = Options)
+	UPROPERTY(EditAnywhere, BlueprintReadWrite, Category = Options)
 	TArray<int> TriGroupIDs;
 };
 
diff --git a/Engine/Plugins/Experimental/GeometryScripting/Source/GeometryScriptingCore/Public/GeometryScript/MeshBooleanFunctions.h b/Engine/Plugins/Experimental/GeometryScripting/Source/GeometryScriptingCore/Public/GeometryScript/MeshBooleanFunctions.h
--- a/Engine/Plugins/Experimental/GeometryScripting/Source/GeometryScriptingCore/Public/GeometryScript/MeshBooleanFunctions.h
+++ b/Engine/Plugins/Experimental/GeometryScripting/Source/GeometryScriptingCore/Public/GeometryScript/MeshBooleanFunctions.h
@@ -24,13 +24,13 @@ struct GEOMETRYSCRIPTINGCORE_API FGeometryScriptMeshBooleanOptions
 {
 	GENERATED_BODY()
 public:
-	UPROPERTY(BlueprintReadWrite, Category = Options)
+	UPROPERTY(EditAnywhere, BlueprintReadWrite, Category = Options)
 	bool bFillHoles = true;
 
-	UPROPERTY(BlueprintReadWrite, Category = Options)
+	UPROPERTY(EditAnywhere, BlueprintReadWrite, Category = Options)
 	bool bSimplifyOutput = true;
 
-	UPROPERTY(BlueprintReadWrite, Category = Options)
+	UPROPERTY(EditAnywhere, BlueprintReadWrite, Category = Options)
 	float SimplifyPlanarTolerance = 0.01f;
 };
 
@@ -40,19 +40,19 @@ struct GEOMETRYSCRIPTINGCORE_API FGeometryScriptMeshSelfUnionOptions
 {
 	GENERATED_BODY()
 public:
-	UPROPERTY(BlueprintReadWrite, Category = Options)
+	UPROPERTY(EditAnywhere, BlueprintReadWrite, Category = Options)
 	bool bFillHoles = true;
 
-	UPROPERTY(BlueprintReadWrite, Category = Options)
+	UPROPERTY(EditAnywhere, BlueprintReadWrite, Category = Options)
 	bool bTrimFlaps = true;
 
-	UPROPERTY(BlueprintReadWrite, Category = Options)
+	UPROPERTY(EditAnywhere, BlueprintReadWrite, Category = Options)
 	bool bSimplifyOutput = true;
 
-	UPROPERTY(BlueprintReadWrite, Category = Options)
+	UPROPERTY(EditAnywhere, BlueprintReadWrite, Category = Options)
 	float SimplifyPlanarTolerance = 0.01f;
 
-	UPROPERTY(BlueprintReadWrite, Category = Options)
+	UPROPERTY(EditAnywhere, BlueprintReadWrite, Category = Options)
 	float WindingThreshold = 0.5f;
 };
 
@@ -63,16 +63,16 @@ struct GEOMETRYSCRIPTINGCORE_API FGeometryScriptMeshPlaneCutOptions
 {
 	GENERATED_BODY()
 public:
-	UPROPERTY(BlueprintReadWrite, Category = Options)
+	UPROPERTY(EditAnywhere, BlueprintReadWrite, Category = Options)
 	bool bFillHoles = true;
 
-	UPROPERTY(BlueprintReadWrite, Category = Options)
+	UPROPERTY(EditAnywhere, BlueprintReadWrite, Category = Options)
 	bool bFillSpans = true;
 
-	UPROPERTY(BlueprintReadWrite, Category = Options)
+	UPROPERTY(EditAnywhere, BlueprintReadWrite, Category = Options)
 	bool bFlipCutSide = false;
 
-	UPROPERTY(BlueprintReadWrite, Category = Options)
+	UPROPERTY(EditAnywhere, BlueprintReadWrite, Category = Options)
 	float UVWorldDimension = 1.0;
 };
 
@@ -82,16 +82,16 @@ struct GEOMETRYSCRIPTINGCORE_API FGeometryScriptMeshPlaneSliceOptions
 {
 	GENERATED_BODY()
 public:
-	UPROPERTY(BlueprintReadWrite, Category = Options)
+	UPROPERTY(EditAnywhere, BlueprintReadWrite, Category = Options)
 	bool bFillHoles = true;
 
-	UPROPERTY(BlueprintReadWrite, Category = Options)
+	UPROPERTY(EditAnywhere, BlueprintReadWrite, Category = Options)
 	bool bFillSpans = true;
 
-	UPROPERTY(BlueprintReadWrite, Category = Options)
+	UPROPERTY(EditAnywhere, BlueprintReadWrite, Category = Options)
 	float GapWidth = 0.01;
 
-	UPROPERTY(BlueprintReadWrite, Category = Options)
+	UPROPERTY(EditAnywhere, BlueprintReadWrite, Category = Options)
 	float UVWorldDimension = 1.0;
 };
 
@@ -101,13 +101,13 @@ struct GEOMETRYSCRIPTINGCORE_API FGeometryScriptMeshMirrorOptions
 {
 	GENERATED_BODY()
 public:
-	UPROPERTY(BlueprintReadWrite, Category = Options)
+	UPROPERTY(EditAnywhere, BlueprintReadWrite, Category = Options)
 	bool bApplyPlaneCut = true;
 
-	UPROPERTY(BlueprintReadWrite, Category = Options)
+	UPROPERTY(EditAnywhere, BlueprintReadWrite, Category = Options)
 	bool bFlipCutSide = false;
 
-	UPROPERTY(BlueprintReadWrite, Category = Options)
+	UPROPERTY(EditAnywhere, BlueprintReadWrite, Category = Options)
 	bool bWeldAlongPlane = true;
 };
 
diff --git a/Engine/Plugins/Experimental/GeometryScripting/Source/GeometryScriptingCore/Public/GeometryScript/MeshComparisonFunctions.h b/Engine/Plugins/Experimental/GeometryScripting/Source/GeometryScriptingCore/Public/GeometryScript/MeshComparisonFunctions.h
--- a/Engine/Plugins/Experimental/GeometryScripting/Source/GeometryScriptingCore/Public/GeometryScript/MeshComparisonFunctions.h
+++ b/Engine/Plugins/Experimental/GeometryScripting/Source/GeometryScriptingCore/Public/GeometryScript/MeshComparisonFunctions.h
@@ -15,28 +15,28 @@ struct GEOMETRYSCRIPTINGCORE_API FGeometryScriptIsSameMeshOptions
 {
 	GENERATED_BODY()
 public:
-	UPROPERTY(BlueprintReadWrite, Category = Options)
+	UPROPERTY(EditAnywhere, BlueprintReadWrite, Category = Options)
 	bool bCheckConnectivity = true;
 
-	UPROPERTY(BlueprintReadWrite, Category = Options)
+	UPROPERTY(EditAnywhere, BlueprintReadWrite, Category = Options)
 	bool bCheckEdgeIDs = false;
 
-	UPROPERTY(BlueprintReadWrite, Category = Options)
+	UPROPERTY(EditAnywhere, BlueprintReadWrite, Category = Options)
 	bool bCheckNormals = false;
 
-	UPROPERTY(BlueprintReadWrite, Category = Options)
+	UPROPERTY(EditAnywhere, BlueprintReadWrite, Category = Options)
 	bool bCheckColors = false;
 
-	UPROPERTY(BlueprintReadWrite, Category = Options)
+	UPROPERTY(EditAnywhere, BlueprintReadWrite, Category = Options)
 	bool bCheckUVs = false;
 
-	UPROPERTY(BlueprintReadWrite, Category = Options)
+	UPROPERTY(EditAnywhere, BlueprintReadWrite, Category = Options)
 	bool bCheckGroups = false;
 
-	UPROPERTY(BlueprintReadWrite, Category = Options)
+	UPROPERTY(EditAnywhere, BlueprintReadWrite, Category = Options)
 	bool bCheckAttributes = false;
 
-	UPROPERTY(BlueprintReadWrite, Category = Options)
+	UPROPERTY(EditAnywhere, BlueprintReadWrite, Category = Options)
 	float Epsilon = 1e-06;
 };
 
@@ -47,7 +47,7 @@ struct GEOMETRYSCRIPTINGCORE_API FGeometryScriptMeasureMeshDistanceOptions
 {
 	GENERATED_BODY()
 public:
-	UPROPERTY(BlueprintReadWrite, Category = Options)
+	UPROPERTY(EditAnywhere, BlueprintReadWrite, Category = Options)
 	bool bSymmetric = true;
 };
 
diff --git a/Engine/Plugins/Experimental/GeometryScripting/Source/GeometryScriptingCore/Public/GeometryScript/MeshDeformFunctions.h b/Engine/Plugins/Experimental/GeometryScripting/Source/GeometryScriptingCore/Public/GeometryScript/MeshDeformFunctions.h
--- a/Engine/Plugins/Experimental/GeometryScripting/Source/GeometryScriptingCore/Public/GeometryScript/MeshDeformFunctions.h
+++ b/Engine/Plugins/Experimental/GeometryScripting/Source/GeometryScriptingCore/Public/GeometryScript/MeshDeformFunctions.h
@@ -16,7 +16,7 @@ struct GEOMETRYSCRIPTINGCORE_API FGeometryScriptBendWarpOptions
 	GENERATED_BODY()
 public:
 	/** Symmetric extents are [-BendExtent,BendExtent], if disabled, then [-LowerExtent,BendExtent] is used  */
-	UPROPERTY(BlueprintReadWrite, Category = Options)
+	UPROPERTY(EditAnywhere, BlueprintReadWrite, Category = Options)
 	bool bSymmetricExtents = true;
 
 	/** Lower extent used when bSymmetricExtents = false */
@@ -24,7 +24,7 @@ public:
 	float LowerExtent = 10;
 
 	/** If true, the Bend is "centered" at the Origin, ie the regions on either side of the extents are rigidly transformed. If false, the Bend begins at the start of the Lower Extents, and the "lower" region is not affected. */
-	UPROPERTY(BlueprintReadWrite, Category = Options)
+	UPROPERTY(EditAnywhere, BlueprintReadWrite, Category = Options)
 	bool bBidirectional = true;
 };
 
@@ -36,7 +36,7 @@ struct GEOMETRYSCRIPTINGCORE_API FGeometryScriptTwistWarpOptions
 	GENERATED_BODY()
 public:
 	/** Symmetric extents are [-BendExtent,BendExtent], if disabled, then [-LowerExtent,BendExtent] is used  */
-	UPROPERTY(BlueprintReadWrite, Category = Options)
+	UPROPERTY(EditAnywhere, BlueprintReadWrite, Category = Options)
 	bool bSymmetricExtents = true;
 
 	/** Lower extent used when bSymmetricExtents = false */
@@ -44,7 +44,7 @@ public:
 	float LowerExtent = 10;
 
 	/** If true, the Twist is "centered" at the Origin, ie the regions on either side of the extents are rigidly transformed. If false, the Twist begins at the start of the Lower Extents, and the "lower" region is not affected. */
-	UPROPERTY(BlueprintReadWrite, Category = Options)
+	UPROPERTY(EditAnywhere, BlueprintReadWrite, Category = Options)
 	bool bBidirectional = true;
 };
 
@@ -67,7 +67,7 @@ struct GEOMETRYSCRIPTINGCORE_API FGeometryScriptFlareWarpOptions
 	GENERATED_BODY()
 public:
 	/** Symmetric extents are [-BendExtent,BendExtent], if disabled, then [-LowerExtent,BendExtent] is used  */
-	UPROPERTY(BlueprintReadWrite, Category = Options)
+	UPROPERTY(EditAnywhere, BlueprintReadWrite, Category = Options)
 	bool bSymmetricExtents = true;
 
 	/** Lower extent used when bSymmetricExtents = false */
@@ -75,7 +75,7 @@ public:
 	float LowerExtent = 10;
 
 	/** Determines the profile used as a displacement */
-	UPROPERTY(BlueprintReadWrite, Category = Options)
+	UPROPERTY(EditAnywhere, BlueprintReadWrite, Category = Options)
 	EGeometryScriptFlareType FlareType = EGeometryScriptFlareType::SinMode;
 };
 
@@ -86,16 +86,16 @@ struct GEOMETRYSCRIPTINGCORE_API FGeometryScriptPerlinNoiseLayerOptions
 {
 	GENERATED_BODY()
 public:
-	UPROPERTY(BlueprintReadWrite, Category = Options)
+	UPROPERTY(EditAnywhere, BlueprintReadWrite, Category = Options)
 	float Magnitude = 5.0;
 
-	UPROPERTY(BlueprintReadWrite, Category = Options)
+	UPROPERTY(EditAnywhere, BlueprintReadWrite, Category = Options)
 	float Frequency = 0.25;
 
-	UPROPERTY(BlueprintReadWrite, Category = Options)
+	UPROPERTY(EditAnywhere, BlueprintReadWrite, Category = Options)
 	FVector FrequencyShift = FVector::Zero();
 
-	UPROPERTY(BlueprintReadWrite, Category = Options)
+	UPROPERTY(EditAnywhere, BlueprintReadWrite, Category = Options)
 	int RandomSeed = 0;
 };
 
@@ -115,13 +115,13 @@ struct GEOMETRYSCRIPTINGCORE_API FGeometryScriptMathWarpOptions
 {
 	GENERATED_BODY()
 public:
-	UPROPERTY(BlueprintReadWrite, Category = Options)
+	UPROPERTY(EditAnywhere, BlueprintReadWrite, Category = Options)
 	float Magnitude = 5.0f;
 
-	UPROPERTY(BlueprintReadWrite, Category = Options)
+	UPROPERTY(EditAnywhere, BlueprintReadWrite, Category = Options)
 	float Frequency = 0.25f;
 
-	UPROPERTY(BlueprintReadWrite, Category = Options)
+	UPROPERTY(EditAnywhere, BlueprintReadWrite, Category = Options)
 	float FrequencyShift = 0.0;
 };
 
@@ -133,10 +133,10 @@ struct GEOMETRYSCRIPTINGCORE_API FGeometryScriptPerlinNoiseOptions
 {
 	GENERATED_BODY()
 public:
-	UPROPERTY(BlueprintReadWrite, Category = Options)
+	UPROPERTY(EditAnywhere, BlueprintReadWrite, Category = Options)
 	FGeometryScriptPerlinNoiseLayerOptions BaseLayer;
 
-	UPROPERTY(BlueprintReadWrite, Category = Options)
+	UPROPERTY(EditAnywhere, BlueprintReadWrite, Category = Options)
 	bool bApplyAlongNormal = true;
 };
 
@@ -146,10 +146,10 @@ struct GEOMETRYSCRIPTINGCORE_API FGeometryScriptIterativeMeshSmoothingOptions
 {
 	GENERATED_BODY()
 public:
-	UPROPERTY(BlueprintReadWrite, Category = Options)
+	UPROPERTY(EditAnywhere, BlueprintReadWrite, Category = Options)
 	int NumIterations = 10;
 
-	UPROPERTY(BlueprintReadWrite, Category = Options)
+	UPROPERTY(EditAnywhere, BlueprintReadWrite, Category = Options)
 	float Alpha = 0.2;
 };
 
@@ -161,19 +161,19 @@ struct GEOMETRYSCRIPTINGCORE_API FGeometryScriptDisplaceFromTextureOptions
 {
 	GENERATED_BODY()
 public:
-	UPROPERTY(BlueprintReadWrite, Category = Options)
+	UPROPERTY(EditAnywhere, BlueprintReadWrite, Category = Options)
 	float Magnitude = 1.0f;
 
-	UPROPERTY(BlueprintReadWrite, Category = Options)
+	UPROPERTY(EditAnywhere, BlueprintReadWrite, Category = Options)
 	FVector2D UVScale = FVector2D(1,1);
 
-	UPROPERTY(BlueprintReadWrite, Category = Options)
+	UPROPERTY(EditAnywhere, BlueprintReadWrite, Category = Options)
 	FVector2D UVOffset = FVector2D(0,0);
 
-	UPROPERTY(BlueprintReadWrite, Category = Options)
+	UPROPERTY(EditAnywhere, BlueprintReadWrite, Category = Options)
 	float Center = 0.5;
 
-	UPROPERTY(BlueprintReadWrite, Category = Options)
+	UPROPERTY(EditAnywhere, BlueprintReadWrite, Category = Options)
 	int ImageChannel = 0;
 };
 
diff --git a/Engine/Plugins/Experimental/GeometryScripting/Source/GeometryScriptingCore/Public/GeometryScript/MeshModelingFunctions.h b/Engine/Plugins/Experimental/GeometryScripting/Source/GeometryScriptingCore/Public/GeometryScript/MeshModelingFunctions.h
--- a/Engine/Plugins/Experimental/GeometryScripting/Source/GeometryScriptingCore/Public/GeometryScript/MeshModelingFunctions.h
+++ b/Engine/Plugins/Experimental/GeometryScripting/Source/GeometryScriptingCore/Public/GeometryScript/MeshModelingFunctions.h
@@ -15,23 +15,23 @@ struct GEOMETRYSCRIPTINGCORE_API FGeometryScriptMeshOffsetOptions
 {
 	GENERATED_BODY()
 public:
-	UPROPERTY(BlueprintReadWrite, Category = Options)
+	UPROPERTY(EditAnywhere, BlueprintReadWrite, Category = Options)
 	float OffsetDistance = 1.0;
 
-	UPROPERTY(BlueprintReadWrite, Category = Options)
+	UPROPERTY(EditAnywhere, BlueprintReadWrite, Category = Options)
 	bool bFixedBoundary = false;
 
-	UPROPERTY(BlueprintReadWrite, Category = Options)
+	UPROPERTY(EditAnywhere, BlueprintReadWrite, Category = Options)
 	int SolveSteps = 5;
 
-	UPROPERTY(BlueprintReadWrite, Category = Options)
+	UPROPERTY(EditAnywhere, BlueprintReadWrite, Category = Options)
 	float SmoothAlpha = 0.1f;
 
-	UPROPERTY(BlueprintReadWrite, Category = Options)
+	UPROPERTY(EditAnywhere, BlueprintReadWrite, Category = Options)
 	bool bReprojectDuringSmoothing = false;
 
 	// should not be > 0.9
-	UPROPERTY(BlueprintReadWrite, Category = Options)
+	UPROPERTY(EditAnywhere, BlueprintReadWrite, Category = Options)
 	float BoundaryAlpha = 0.2f;
 };
 
@@ -42,16 +42,16 @@ struct GEOMETRYSCRIPTINGCORE_API FGeometryScriptMeshExtrudeOptions
 {
 	GENERATED_BODY()
 public:
-	UPROPERTY(BlueprintReadWrite, Category = Options)
+	UPROPERTY(EditAnywhere, BlueprintReadWrite, Category = Options)
 	float ExtrudeDistance = 1.0;
 
-	UPROPERTY(BlueprintReadWrite, Category = Options)
+	UPROPERTY(EditAnywhere, BlueprintReadWrite, Category = Options)
 	FVector ExtrudeDirection = FVector(0,0,1);
 
-	UPROPERTY(BlueprintReadWrite, Category = Options)
+	UPROPERTY(EditAnywhere, BlueprintReadWrite, Category = Options)
 	float UVScale = 1.0f;
 
-	UPROPERTY(BlueprintReadWrite, Category = Options)
+	UPROPERTY(EditAnywhere, BlueprintReadWrite, Category = Options)
 	bool bSolidsToShells = true;
 };
 
@@ -62,13 +62,13 @@ struct GEOMETRYSCRIPTINGCORE_API FGeometryScriptMeshBevelOptions
 {
 	GENERATED_BODY()
 public:
-	UPROPERTY(BlueprintReadWrite, Category = Options)
+	UPROPERTY(EditAnywhere, BlueprintReadWrite, Category = Options)
 	float BevelDistance = 1.0;
 
-	UPROPERTY(BlueprintReadWrite, Category = Options)
+	UPROPERTY(EditAnywhere, BlueprintReadWrite, Category = Options)
 	bool bInferMaterialID = false;
 
-	UPROPERTY(BlueprintReadWrite, Category = Options)
+	UPROPERTY(EditAnywhere, BlueprintReadWrite, Category = Options)
 	int SetMaterialID = 0;
 
 
diff --git a/Engine/Plugins/Experimental/GeometryScripting/Source/GeometryScriptingCore/Public/GeometryScript/MeshNormalsFunctions.h b/Engine/Plugins/Experimental/GeometryScripting/Source/GeometryScriptingCore/Public/GeometryScript/MeshNormalsFunctions.h
--- a/Engine/Plugins/Experimental/GeometryScripting/Source/GeometryScriptingCore/Public/GeometryScript/MeshNormalsFunctions.h
+++ b/Engine/Plugins/Experimental/GeometryScripting/Source/GeometryScriptingCore/Public/GeometryScript/MeshNormalsFunctions.h
@@ -14,10 +14,10 @@ struct GEOMETRYSCRIPTINGCORE_API FGeometryScriptCalculateNormalsOptions
 {
 	GENERATED_BODY()
 public:
-	UPROPERTY(BlueprintReadWrite, Category = Options)
+	UPROPERTY(EditAnywhere, BlueprintReadWrite, Category = Options)
 	bool bAngleWeighted = true;
 
-	UPROPERTY(BlueprintReadWrite, Category = Options)
+	UPROPERTY(EditAnywhere, BlueprintReadWrite, Category = Options)
 	bool bAreaWeighted = true;
 };
 
@@ -28,16 +28,16 @@ struct GEOMETRYSCRIPTINGCORE_API FGeometryScriptSplitNormalsOptions
 {
 	GENERATED_BODY()
 public:
-	UPROPERTY(BlueprintReadWrite, Category = Options)
+	UPROPERTY(EditAnywhere, BlueprintReadWrite, Category = Options)
 	bool bSplitByOpeningAngle = true;
 
-	UPROPERTY(BlueprintReadWrite, Category = Options)
+	UPROPERTY(EditAnywhere, BlueprintReadWrite, Category = Options)
 	float OpeningAngleDeg = 15.0f;
 
-	UPROPERTY(BlueprintReadWrite, Category = Options)
+	UPROPERTY(EditAnywhere, BlueprintReadWrite, Category = Options)
 	bool bSplitByFaceGroup = false;
 
-	UPROPERTY(BlueprintReadWrite, Category = Options)
+	UPROPERTY(EditAnywhere, BlueprintReadWrite, Category = Options)
 	FGeometryScriptGroupLayer GroupLayer;
 };
 
@@ -55,10 +55,10 @@ struct GEOMETRYSCRIPTINGCORE_API FGeometryScriptTangentsOptions
 {
 	GENERATED_BODY()
 public:
-	UPROPERTY(BlueprintReadWrite, Category = Options)
+	UPROPERTY(EditAnywhere, BlueprintReadWrite, Category = Options)
 	EGeometryScriptTangentTypes Type = EGeometryScriptTangentTypes::FastMikkT;
 
-	UPROPERTY(BlueprintReadWrite, Category = Options)
+	UPROPERTY(EditAnywhere, BlueprintReadWrite, Category = Options)
 	int UVLayer = 0;
 };
 
diff --git a/Engine/Plugins/Experimental/GeometryScripting/Source/GeometryScriptingCore/Public/GeometryScript/MeshPrimitiveFunctions.h b/Engine/Plugins/Experimental/GeometryScripting/Source/GeometryScriptingCore/Public/GeometryScript/MeshPrimitiveFunctions.h
--- a/Engine/Plugins/Experimental/GeometryScripting/Source/GeometryScriptingCore/Public/GeometryScript/MeshPrimitiveFunctions.h
+++ b/Engine/Plugins/Experimental/GeometryScripting/Source/GeometryScriptingCore/Public/GeometryScript/MeshPrimitiveFunctions.h
@@ -36,13 +36,13 @@ struct GEOMETRYSCRIPTINGCORE_API FGeometryScriptPrimitiveOptions
 {
 	GENERATED_BODY()
 public:
-	UPROPERTY(BlueprintReadWrite, Category = Options)
+	UPROPERTY(EditAnywhere, BlueprintReadWrite, Category = Options)
 	EGeometryScriptPrimitivePolygroupMode PolygroupMode = EGeometryScriptPrimitivePolygroupMode::PerFace;
 
-	UPROPERTY(BlueprintReadWrite, Category = Options)
+	UPROPERTY(EditAnywhere, BlueprintReadWrite, Category = Options)
 	bool bFlipOrientation = false;
 
-	UPROPERTY(BlueprintReadWrite, Category = Options)
+	UPROPERTY(EditAnywhere, BlueprintReadWrite, Category = Options)
 	EGeometryScriptPrimitiveUVMode UVMode = EGeometryScriptPrimitiveUVMode::Uniform;
 };
 
@@ -52,25 +52,25 @@ struct GEOMETRYSCRIPTINGCORE_API FGeometryScriptRevolveOptions
 {
 	GENERATED_BODY()
 public:
-	UPROPERTY(BlueprintReadWrite, Category = Options)
+	UPROPERTY(EditAnywhere, BlueprintReadWrite, Category = Options)
 	float RevolveDegrees = 360.0f;
 
-	UPROPERTY(BlueprintReadWrite, Category = Options)
+	UPROPERTY(EditAnywhere, BlueprintReadWrite, Category = Options)
 	float DegreeOffset = 0.0f;
 
-	UPROPERTY(BlueprintReadWrite, Category = Options)
+	UPROPERTY(EditAnywhere, BlueprintReadWrite, Category = Options)
 	bool bReverseDirection = false;
 
-	UPROPERTY(BlueprintReadWrite, Category = Options)
+	UPROPERTY(EditAnywhere, BlueprintReadWrite, Category = Options)
 	bool bHardNormals = false;
 
-	UPROPERTY(BlueprintReadWrite, Category = Options)
+	UPROPERTY(EditAnywhere, BlueprintReadWrite, Category = Options)
 	float HardNormalAngle = 30.0f;
 
-	UPROPERTY(BlueprintReadWrite, Category = Options)
+	UPROPERTY(EditAnywhere, BlueprintReadWrite, Category = Options)
 	bool bProfileAtMidpoint = false;
 
-	UPROPERTY(BlueprintReadWrite, Category = Options)
+	UPROPERTY(EditAnywhere, BlueprintReadWrite, Category = Options)
 	bool bFillPartialRevolveEndcaps = true;
 };
 
diff --git a/Engine/Plugins/Experimental/GeometryScripting/Source/GeometryScriptingCore/Public/GeometryScript/MeshRepairFunctions.h b/Engine/Plugins/Experimental/GeometryScripting/Source/GeometryScriptingCore/Public/GeometryScript/MeshRepairFunctions.h
--- a/Engine/Plugins/Experimental/GeometryScripting/Source/GeometryScriptingCore/Public/GeometryScript/MeshRepairFunctions.h
+++ b/Engine/Plugins/Experimental/GeometryScripting/Source/GeometryScriptingCore/Public/GeometryScript/MeshRepairFunctions.h
@@ -15,11 +15,11 @@ struct GEOMETRYSCRIPTINGCORE_API FGeometryScriptWeldEdgesOptions
 	GENERATED_BODY()
 public:
 	/** Edges are coincident if both pairs of endpoint vertices, and their midpoint, are closer than this distance */
-	UPROPERTY(BlueprintReadWrite, Category = Options)
+	UPROPERTY(EditAnywhere, BlueprintReadWrite, Category = Options)
 	float Tolerance = 1e-06f;
 
 	/** Only merge unambiguous pairs that have unique duplicate-edge matches */
-	UPROPERTY(BlueprintReadWrite, Category = Options)
+	UPROPERTY(EditAnywhere, BlueprintReadWrite, Category = Options)
 	bool bOnlyUniquePairs = true;
 };
 
@@ -39,7 +39,7 @@ struct GEOMETRYSCRIPTINGCORE_API FGeometryScriptFillHolesOptions
 {
 	GENERATED_BODY()
 public:
-	UPROPERTY(BlueprintReadWrite, Category = Options)
+	UPROPERTY(EditAnywhere, BlueprintReadWrite, Category = Options)
 	EGeometryScriptFillHolesMethod FillMethod = EGeometryScriptFillHolesMethod::Automatic;
 
 	/** Delete floating, disconnected triangles, as they produce a "hole" that cannot be filled */
@@ -54,14 +54,14 @@ struct GEOMETRYSCRIPTINGCORE_API FGeometryScriptRemoveSmallComponentOptions
 	GENERATED_BODY()
 public:
 	/**  */
-	UPROPERTY(BlueprintReadWrite, Category = Options)
+	UPROPERTY(EditAnywhere, BlueprintReadWrite, Category = Options)
 	float MinVolume = 0.0001;
 
 	/**  */
-	UPROPERTY(BlueprintReadWrite, Category = Options)
+	UPROPERTY(EditAnywhere, BlueprintReadWrite, Category = Options)
 	float MinArea = 0.0001;
 
-	UPROPERTY(BlueprintReadWrite, Category = Options)
+	UPROPERTY(EditAnywhere, BlueprintReadWrite, Category = Options)
 	int MinTriangleCount = 1;
 };
 
@@ -82,32 +82,32 @@ struct GEOMETRYSCRIPTINGCORE_API FGeometryScriptRemoveHiddenTrianglesOptions
 	GENERATED_BODY()
 public:
 
-	UPROPERTY(BlueprintReadWrite, Category = Options)
+	UPROPERTY(EditAnywhere, BlueprintReadWrite, Category = Options)
 	EGeometryScriptRemoveHiddenTrianglesMethod Method = EGeometryScriptRemoveHiddenTrianglesMethod::FastWindingNumber;
 
 	// add triangle samples per triangle (in addition to TriangleSamplingMethod)
-	UPROPERTY(BlueprintReadWrite, Category = Options)
+	UPROPERTY(EditAnywhere, BlueprintReadWrite, Category = Options)
 	int SamplesPerTriangle = 0;
 
 	// once triangles to remove are identified, do iterations of boundary erosion, ie contract selection by boundary vertex one-rings
-	UPROPERTY(BlueprintReadWrite, Category = Options)
+	UPROPERTY(EditAnywhere, BlueprintReadWrite, Category = Options)
 	int ShrinkSelection = 0;
 
 	// use this as winding isovalue for WindingNumber mode
-	UPROPERTY(BlueprintReadWrite, Category = Options)
+	UPROPERTY(EditAnywhere, BlueprintReadWrite, Category = Options)
 	float WindingIsoValue = 0.5;
 
 	// random rays to add beyond +/- major axes, for raycast sampling
-	UPROPERTY(BlueprintReadWrite, Category = Options)
+	UPROPERTY(EditAnywhere, BlueprintReadWrite, Category = Options)
 	int RaysPerSample = 0;
 
 
 	/** Nudge sample points out by this amount to try to counteract numerical issues */
-	UPROPERTY(BlueprintReadWrite, Category = Options)
+	UPROPERTY(EditAnywhere, BlueprintReadWrite, Category = Options)
 	float NormalOffset = 1e-6f;
 
 	/**  */
-	UPROPERTY(BlueprintReadWrite, Category = Options)
+	UPROPERTY(EditAnywhere, BlueprintReadWrite, Category = Options)
 	bool bCompactResult = true;
 };
 
diff --git a/Engine/Plugins/Experimental/GeometryScripting/Source/GeometryScriptingCore/Public/GeometryScript/MeshSimplifyFunctions.h b/Engine/Plugins/Experimental/GeometryScripting/Source/GeometryScriptingCore/Public/GeometryScript/MeshSimplifyFunctions.h
--- a/Engine/Plugins/Experimental/GeometryScripting/Source/GeometryScriptingCore/Public/GeometryScript/MeshSimplifyFunctions.h
+++ b/Engine/Plugins/Experimental/GeometryScripting/Source/GeometryScriptingCore/Public/GeometryScript/MeshSimplifyFunctions.h
@@ -15,7 +15,7 @@ struct GEOMETRYSCRIPTINGCORE_API FGeometryScriptPlanarSimplifyOptions
 	GENERATED_BODY()
 public:
 	/** */
-	UPROPERTY(BlueprintReadWrite, Category = Options)
+	UPROPERTY(EditAnywhere, BlueprintReadWrite, Category = Options)
 	float AngleThreshold = 0.001;
 };
 
@@ -25,7 +25,7 @@ struct GEOMETRYSCRIPTINGCORE_API FGeometryScriptPolygroupSimplifyOptions
 	GENERATED_BODY()
 public:
 	/** */
-	UPROPERTY(BlueprintReadWrite, Category = Options)
+	UPROPERTY(EditAnywhere, BlueprintReadWrite, Category = Options)
 	float AngleThreshold = 0.001;
 };
 
@@ -45,22 +45,22 @@ struct GEOMETRYSCRIPTINGCORE_API FGeometryScriptSimplifyMeshOptions
 	GENERATED_BODY()
 public:
 
-	UPROPERTY(BlueprintReadWrite, Category = Options)
+	UPROPERTY(EditAnywhere, BlueprintReadWrite, Category = Options)
 	EGeometryScriptRemoveMeshSimplificationType Method = EGeometryScriptRemoveMeshSimplificationType::AttributeAware;
 
-	UPROPERTY(BlueprintReadWrite, Category = Options)
+	UPROPERTY(EditAnywhere, BlueprintReadWrite, Category = Options)
 	bool bAllowSeamCollapse = true;
 
-	UPROPERTY(BlueprintReadWrite, Category = Options)
+	UPROPERTY(EditAnywhere, BlueprintReadWrite, Category = Options)
 	bool bAllowSeamSmoothing = true;
 
-	UPROPERTY(BlueprintReadWrite, Category = Options)
+	UPROPERTY(EditAnywhere, BlueprintReadWrite, Category = Options)
 	bool bAllowSeamSplits = true;
 
-	UPROPERTY(BlueprintReadWrite, Category = Options)
+	UPROPERTY(EditAnywhere, BlueprintReadWrite, Category = Options)
 	bool bPreserveVertexPositions = false;
 
-	UPROPERTY(BlueprintReadWrite, Category = Options)
+	UPROPERTY(EditAnywhere, BlueprintReadWrite, Category = Options)
 	bool bRetainQuadricMemory = false;
 };
 
diff --git a/Engine/Plugins/Experimental/GeometryScripting/Source/GeometryScriptingCore/Public/GeometryScript/MeshSpatialFunctions.h b/Engine/Plugins/Experimental/GeometryScripting/Source/GeometryScriptingCore/Public/GeometryScript/MeshSpatialFunctions.h
--- a/Engine/Plugins/Experimental/GeometryScripting/Source/GeometryScriptingCore/Public/GeometryScript/MeshSpatialFunctions.h
+++ b/Engine/Plugins/Experimental/GeometryScripting/Source/GeometryScriptingCore/Public/GeometryScript/MeshSpatialFunctions.h
@@ -15,13 +15,13 @@ struct GEOMETRYSCRIPTINGCORE_API FGeometryScriptSpatialQueryOptions
 {
 	GENERATED_BODY()
 public:
-	UPROPERTY(BlueprintReadWrite, Category = Options)
+	UPROPERTY(EditAnywhere, BlueprintReadWrite, Category = Options)
 	float MaxDistance = 0;
 
-	UPROPERTY(BlueprintReadWrite, Category = Options)
+	UPROPERTY(EditAnywhere, BlueprintReadWrite, Category = Options)
 	bool bAllowUnsafeModifiedQueries = false;
 
-	UPROPERTY(BlueprintReadWrite, Category = Options)
+	UPROPERTY(EditAnywhere, BlueprintReadWrite, Category = Options)
 	float WindingIsoThreshold = 0.5;
 };
 
@@ -34,19 +34,19 @@ struct GEOMETRYSCRIPTINGCORE_API FGeometryScriptRayHitResult
 {
 	GENERATED_BODY()
 public:
-	UPROPERTY(BlueprintReadWrite, Category = Options)
+	UPROPERTY(EditAnywhere, BlueprintReadWrite, Category = Options)
 	bool bHit = false;
 
-	UPROPERTY(BlueprintReadWrite, Category = Options)
+	UPROPERTY(EditAnywhere, BlueprintReadWrite, Category = Options)
 	float RayParameter = 0;
 
-	UPROPERTY(BlueprintReadWrite, Category = Options)
+	UPROPERTY(EditAnywhere, BlueprintReadWrite, Category = Options)
 	int HitTriangleID = -1;
 
-	UPROPERTY(BlueprintReadWrite, Category = Options)
+	UPROPERTY(EditAnywhere, BlueprintReadWrite, Category = Options)
 	FVector HitPosition = FVector::ZeroVector;
 
-	UPROPERTY(BlueprintReadWrite, Category = Options)
+	UPROPERTY(EditAnywhere, BlueprintReadWrite, Category = Options)
 	FVector HitBaryCoords = FVector::ZeroVector;
 };
 
diff --git a/Engine/Plugins/Experimental/GeometryScripting/Source/GeometryScriptingCore/Public/GeometryScript/MeshSubdivideFunctions.h b/Engine/Plugins/Experimental/GeometryScripting/Source/GeometryScriptingCore/Public/GeometryScript/MeshSubdivideFunctions.h
--- a/Engine/Plugins/Experimental/GeometryScripting/Source/GeometryScriptingCore/Public/GeometryScript/MeshSubdivideFunctions.h
+++ b/Engine/Plugins/Experimental/GeometryScripting/Source/GeometryScriptingCore/Public/GeometryScript/MeshSubdivideFunctions.h
@@ -17,7 +17,7 @@ struct GEOMETRYSCRIPTINGCORE_API FGeometryScriptPNTessellateOptions
 {
 	GENERATED_BODY()
 public:
-	UPROPERTY(BlueprintReadWrite, Category = Options)
+	UPROPERTY(EditAnywhere, BlueprintReadWrite, Category = Options)
 	bool bRecomputeNormals = true;
 };
 
diff --git a/Engine/Plugins/Experimental/GeometryScripting/Source/GeometryScriptingCore/Public/GeometryScript/MeshUVFunctions.h b/Engine/Plugins/Experimental/GeometryScripting/Source/GeometryScriptingCore/Public/GeometryScript/MeshUVFunctions.h
--- a/Engine/Plugins/Experimental/GeometryScripting/Source/GeometryScriptingCore/Public/GeometryScript/MeshUVFunctions.h
+++ b/Engine/Plugins/Experimental/GeometryScripting/Source/GeometryScriptingCore/Public/GeometryScript/MeshUVFunctions.h
@@ -15,10 +15,10 @@ struct GEOMETRYSCRIPTINGCORE_API FGeometryScriptRepackUVsOptions
 {
 	GENERATED_BODY()
 public:
-	UPROPERTY(BlueprintReadWrite, Category = Options)
+	UPROPERTY(EditAnywhere, BlueprintReadWrite, Category = Options)
 	int TargetImageWidth = 512;
 
-	UPROPERTY(BlueprintReadWrite, Category = Options)
+	UPROPERTY(EditAnywhere, BlueprintReadWrite, Category = Options)
 	bool bOptimizeIslandRotation = true;
 };
 
@@ -43,10 +43,10 @@ struct GEOMETRYSCRIPTINGCORE_API FGeometryScriptExpMapUVOptions
 {
 	GENERATED_BODY()
 
-	UPROPERTY(BlueprintReadWrite, Category = Options)
+	UPROPERTY(EditAnywhere, BlueprintReadWrite, Category = Options)
 	int NormalSmoothingRounds = 0;
 
-	UPROPERTY(BlueprintReadWrite, Category = Options)
+	UPROPERTY(EditAnywhere, BlueprintReadWrite, Category = Options)
 	float NormalSmoothingAlpha = 0.25f;
 };
 
@@ -55,7 +55,7 @@ struct GEOMETRYSCRIPTINGCORE_API FGeometryScriptSpectralConformalUVOptions
 {
 	GENERATED_BODY()
 
-	UPROPERTY(BlueprintReadWrite, Category = Options)
+	UPROPERTY(EditAnywhere, BlueprintReadWrite, Category = Options)
 	bool bPreserveIrregularity = true;
 };
 
@@ -65,22 +65,22 @@ struct GEOMETRYSCRIPTINGCORE_API FGeometryScriptRecomputeUVsOptions
 {
 	GENERATED_BODY()
 public:
-	UPROPERTY(BlueprintReadWrite, Category = Options)
+	UPROPERTY(EditAnywhere, BlueprintReadWrite, Category = Options)
 	EGeometryScriptUVFlattenMethod Method = EGeometryScriptUVFlattenMethod::SpectralConformal;
 
-	UPROPERTY(BlueprintReadWrite, Category = Options)
+	UPROPERTY(EditAnywhere, BlueprintReadWrite, Category = Options)
 	EGeometryScriptUVIslandSource IslandSource = EGeometryScriptUVIslandSource::UVIslands;
 
-	UPROPERTY(BlueprintReadWrite, Category = Options)
+	UPROPERTY(EditAnywhere, BlueprintReadWrite, Category = Options)
 	FGeometryScriptExpMapUVOptions ExpMapOptions;
 
-	UPROPERTY(BlueprintReadWrite, Category = Options)
+	UPROPERTY(EditAnywhere, BlueprintReadWrite, Category = Options)
 	FGeometryScriptSpectralConformalUVOptions SpectralConformalOptions;
 
-	UPROPERTY(BlueprintReadWrite, Category = Options)
+	UPROPERTY(EditAnywhere, BlueprintReadWrite, Category = Options)
 	FGeometryScriptGroupLayer GroupLayer;
 
-	UPROPERTY(BlueprintReadWrite, Category = Options)
+	UPROPERTY(EditAnywhere, BlueprintReadWrite, Category = Options)
 	bool bAutoAlignIslandsWithAxes = true;
 };
 
@@ -93,35 +93,35 @@ struct GEOMETRYSCRIPTINGCORE_API FGeometryScriptPatchBuilderOptions
 {
 	GENERATED_BODY()
 public:
-	UPROPERTY(BlueprintReadWrite, Category = Options)
+	UPROPERTY(EditAnywhere, BlueprintReadWrite, Category = Options)
 	int InitialPatchCount = 100;
 
-	UPROPERTY(BlueprintReadWrite, Category = Options)
+	UPROPERTY(EditAnywhere, BlueprintReadWrite, Category = Options)
 	int MinPatchSize = 2;
 
-	UPROPERTY(BlueprintReadWrite, Category = Options)
+	UPROPERTY(EditAnywhere, BlueprintReadWrite, Category = Options)
 	float PatchCurvatureAlignmentWeight = 1.0f;
 
-	UPROPERTY(BlueprintReadWrite, Category = Options)
+	UPROPERTY(EditAnywhere, BlueprintReadWrite, Category = Options)
 	float PatchMergingMetricThresh = 1.5f;
 
-	UPROPERTY(BlueprintReadWrite, Category = Options)
+	UPROPERTY(EditAnywhere, BlueprintReadWrite, Category = Options)
 	float PatchMergingAngleThresh = 45.0f;
 
-	UPROPERTY(BlueprintReadWrite, Category = Options)
+	UPROPERTY(EditAnywhere, BlueprintReadWrite, Category = Options)
 	FGeometryScriptExpMapUVOptions ExpMapOptions;
 
-	UPROPERTY(BlueprintReadWrite, Category = Options)
+	UPROPERTY(EditAnywhere, BlueprintReadWrite, Category = Options)
 	bool bRespectInputGroups = false;
 
-	UPROPERTY(BlueprintReadWrite, Category = Options)
+	UPROPERTY(EditAnywhere, BlueprintReadWrite, Category = Options)
 	FGeometryScriptGroupLayer GroupLayer;
 
 
-	UPROPERTY(BlueprintReadWrite, Category = Options)
+	UPROPERTY(EditAnywhere, BlueprintReadWrite, Category = Options)
 	bool bAutoPack = true;
 
-	UPROPERTY(BlueprintReadWrite, Category = Options)
+	UPROPERTY(EditAnywhere, BlueprintReadWrite, Category = Options)
 	FGeometryScriptRepackUVsOptions PackingOptions;
 };
 
@@ -133,7 +133,7 @@ struct GEOMETRYSCRIPTINGCORE_API FGeometryScriptXAtlasOptions
 {
 	GENERATED_BODY()
 public:
-	UPROPERTY(BlueprintReadWrite, Category = Options)
+	UPROPERTY(EditAnywhere, BlueprintReadWrite, Category = Options)
 	int MaxIterations = 2;
 };
 
diff --git a/Engine/Plugins/Experimental/GeometryScripting/Source/GeometryScriptingCore/Public/GeometryScript/MeshVoxelFunctions.h b/Engine/Plugins/Experimental/GeometryScripting/Source/GeometryScriptingCore/Public/GeometryScript/MeshVoxelFunctions.h
--- a/Engine/Plugins/Experimental/GeometryScripting/Source/GeometryScriptingCore/Public/GeometryScript/MeshVoxelFunctions.h
+++ b/Engine/Plugins/Experimental/GeometryScripting/Source/GeometryScriptingCore/Public/GeometryScript/MeshVoxelFunctions.h
@@ -26,7 +26,7 @@ struct GEOMETRYSCRIPTINGCORE_API FGeometryScript3DGridParameters
 	GENERATED_BODY()
 public:
 	/** SizeMethod determines how the parameters below will be interpreted to define the size of a 3D sampling/voxel grid */
-	UPROPERTY(BlueprintReadWrite, Category = Options)
+	UPROPERTY(EditAnywhere, BlueprintReadWrite, Category = Options)
 	EGeometryScriptGridSizingMethod SizeMethod = EGeometryScriptGridSizingMethod::GridResolution;
 
 	/** Use a specific grid cell size, and construct a grid with dimensions large enough to contain the target object */
@@ -44,19 +44,19 @@ struct GEOMETRYSCRIPTINGCORE_API FGeometryScriptSolidifyOptions
 {
 	GENERATED_BODY()
 public:
-	UPROPERTY(BlueprintReadWrite, Category = Options)
+	UPROPERTY(EditAnywhere, BlueprintReadWrite, Category = Options)
 	FGeometryScript3DGridParameters GridParameters;
 
-	UPROPERTY(BlueprintReadWrite, Category = Options)
+	UPROPERTY(EditAnywhere, BlueprintReadWrite, Category = Options)
 	float WindingThreshold = 0.5;
 
-	UPROPERTY(BlueprintReadWrite, Category = Options)
+	UPROPERTY(EditAnywhere, BlueprintReadWrite, Category = Options)
 	bool bSolidAtBoundaries = true;
 
-	UPROPERTY(BlueprintReadWrite, Category = Options)
+	UPROPERTY(EditAnywhere, BlueprintReadWrite, Category = Options)
 	float ExtendBounds = 1.0;
 
-	UPROPERTY(BlueprintReadWrite, Category = Options)
+	UPROPERTY(EditAnywhere, BlueprintReadWrite, Category = Options)
 	int SurfaceSearchSteps = 3;
 };
 
@@ -84,19 +84,19 @@ struct GEOMETRYSCRIPTINGCORE_API FGeometryScriptMorphologyOptions
 {
 	GENERATED_BODY()
 public:
-	UPROPERTY(BlueprintReadWrite, Category = Options)
+	UPROPERTY(EditAnywhere, BlueprintReadWrite, Category = Options)
 	FGeometryScript3DGridParameters SDFGridParameters;
 
-	UPROPERTY(BlueprintReadWrite, Category = Options)
+	UPROPERTY(EditAnywhere, BlueprintReadWrite, Category = Options)
 	bool bUseSeparateMeshGrid = false;
 
-	UPROPERTY(BlueprintReadWrite, Category = Options)
+	UPROPERTY(EditAnywhere, BlueprintReadWrite, Category = Options)
 	FGeometryScript3DGridParameters MeshGridParameters;
 
-	UPROPERTY(BlueprintReadWrite, Category = Options)
+	UPROPERTY(EditAnywhere, BlueprintReadWrite, Category = Options)
 	EGeometryScriptMorphologicalOpType Operation = EGeometryScriptMorphologicalOpType::Dilate;
 
-	UPROPERTY(BlueprintReadWrite, Category = Options)
+	UPROPERTY(EditAnywhere, BlueprintReadWrite, Category = Options)
 	float Distance = 1.0;
 
 };
diff --git a/Engine/Plugins/Experimental/GeometryScripting/Source/GeometryScriptingCore/Public/GeometryScript/PolyPathFunctions.h b/Engine/Plugins/Experimental/GeometryScripting/Source/GeometryScriptingCore/Public/GeometryScript/PolyPathFunctions.h
--- a/Engine/Plugins/Experimental/GeometryScripting/Source/GeometryScriptingCore/Public/GeometryScript/PolyPathFunctions.h
+++ b/Engine/Plugins/Experimental/GeometryScripting/Source/GeometryScriptingCore/Public/GeometryScript/PolyPathFunctions.h
@@ -29,10 +29,10 @@ public:
 	UPROPERTY(BlueprintReadWrite, Category = Options, meta = (EditCondition = "SampleSpacing == EGeometryScriptSampleSpacing::ErrorTolerance"))
 	float ErrorTolerance = 1;
 
-	UPROPERTY(BlueprintReadWrite, Category = Options)
+	UPROPERTY(EditAnywhere, BlueprintReadWrite, Category = Options)
 	EGeometryScriptSampleSpacing SampleSpacing = EGeometryScriptSampleSpacing::UniformDistance;
 
-	UPROPERTY(BlueprintReadWrite, Category = Options)
+	UPROPERTY(EditAnywhere, BlueprintReadWrite, Category = Options)
 	TEnumAsByte<ESplineCoordinateSpace::Type> CoordinateSpace = ESplineCoordinateSpace::Type::Local;
 };
 
diff --git a/Engine/Plugins/Experimental/GeometryScripting/Source/GeometryScriptingCore/Public/GeometryScript/SceneUtilityFunctions.h b/Engine/Plugins/Experimental/GeometryScripting/Source/GeometryScriptingCore/Public/GeometryScript/SceneUtilityFunctions.h
--- a/Engine/Plugins/Experimental/GeometryScripting/Source/GeometryScriptingCore/Public/GeometryScript/SceneUtilityFunctions.h
+++ b/Engine/Plugins/Experimental/GeometryScripting/Source/GeometryScriptingCore/Public/GeometryScript/SceneUtilityFunctions.h
@@ -17,13 +17,13 @@ struct GEOMETRYSCRIPTINGCORE_API FGeometryScriptCopyMeshFromComponentOptions
 {
 	GENERATED_BODY()
 public:
-	UPROPERTY(BlueprintReadWrite, Category = Options)
+	UPROPERTY(EditAnywhere, BlueprintReadWrite, Category = Options)
 	bool bWantNormals = true;
 
-	UPROPERTY(BlueprintReadWrite, Category = Options)
+	UPROPERTY(EditAnywhere, BlueprintReadWrite, Category = Options)
 	bool bWantTangents = true;
 
-	UPROPERTY(BlueprintReadWrite, Category = Options)
+	UPROPERTY(EditAnywhere, BlueprintReadWrite, Category = Options)
 	FGeometryScriptMeshReadLOD RequestedLOD = FGeometryScriptMeshReadLOD();
 };
 
diff --git a/Engine/Plugins/Experimental/GeometryScripting/Source/GeometryScriptingCore/Public/GeometryScript/TextureMapFunctions.h b/Engine/Plugins/Experimental/GeometryScripting/Source/GeometryScriptingCore/Public/GeometryScript/TextureMapFunctions.h
--- a/Engine/Plugins/Experimental/GeometryScripting/Source/GeometryScriptingCore/Public/GeometryScript/TextureMapFunctions.h
+++ b/Engine/Plugins/Experimental/GeometryScripting/Source/GeometryScriptingCore/Public/GeometryScript/TextureMapFunctions.h
@@ -23,16 +23,16 @@ struct GEOMETRYSCRIPTINGCORE_API FGeometryScriptSampleTextureOptions
 {
 	GENERATED_BODY()
 public:
-	UPROPERTY(BlueprintReadWrite, Category = Options)
+	UPROPERTY(EditAnywhere, BlueprintReadWrite, Category = Options)
 	EGeometryScriptPixelSamplingMethod SamplingMethod = EGeometryScriptPixelSamplingMethod::Bilinear;
 
-	UPROPERTY(BlueprintReadWrite, Category = Options)
+	UPROPERTY(EditAnywhere, BlueprintReadWrite, Category = Options)
 	bool bWrap = true;
 
-	UPROPERTY(BlueprintReadWrite, Category = Options)
+	UPROPERTY(EditAnywhere, BlueprintReadWrite, Category = Options)
 	FVector2D UVScale = FVector2D(1,1);
 
-	UPROPERTY(BlueprintReadWrite, Category = Options)
+	UPROPERTY(EditAnywhere, BlueprintReadWrite, Category = Options)
 	FVector2D UVOffset = FVector2D(0,0);
 };
 
