Fixed the text conflict from LOCTEXT macro for namespace and key.
Fixed the wording of the PNTriangles error when the tessellator fails to validate the inputs.

#rnx

diff --git a/Engine/Plugins/Experimental/GeometryScripting/Source/GeometryScriptingCore/Private/MeshSubdivideFunctions.cpp b/Engine/Plugins/Experimental/GeometryScripting/Source/GeometryScriptingCore/Private/MeshSubdivideFunctions.cpp
--- a/Engine/Plugins/Experimental/GeometryScripting/Source/GeometryScriptingCore/Private/MeshSubdivideFunctions.cpp
+++ b/Engine/Plugins/Experimental/GeometryScripting/Source/GeometryScriptingCore/Private/MeshSubdivideFunctions.cpp
@@ -78,13 +78,13 @@ UDynamicMesh* UGeometryScriptLibrary_MeshSubdivideFunctions::ApplyPNTessellation
 		
 		if (Tessellator.Validate() != EOperationValidationResult::Ok)
 		{
-			UE::Geometry::AppendError(Debug, EGeometryScriptErrorType::InvalidInputs, LOCTEXT("ApplyPNTessellation_InvalidInput", "ApplyPNTessellation: TargetMesh is Null"));
+			UE::Geometry::AppendError(Debug, EGeometryScriptErrorType::InvalidInputs, LOCTEXT("ApplyPNTessellation_Error", "ApplyPNTessellation: The inputs are invalid"));
 			return;
 		}
 
 		if (Tessellator.Compute() == false) 
 		{
-			UE::Geometry::AppendError(Debug, EGeometryScriptErrorType::OperationFailed, LOCTEXT("ApplyPNTessellation_Failed", "ApplyPNTessellation: Tessellation failed."));
+			UE::Geometry::AppendError(Debug, EGeometryScriptErrorType::OperationFailed, LOCTEXT("ApplyPNTessellation_Failed", "ApplyPNTessellation: Tessellation failed"));
 			return;
 		}
 		
@@ -115,13 +115,13 @@ UDynamicMesh* UGeometryScriptLibrary_MeshSubdivideFunctions::ApplyUniformTessell
 		
 		if (Tessellator.Validate() != EOperationValidationResult::Ok)
 		{
-			UE::Geometry::AppendError(Debug, EGeometryScriptErrorType::InvalidInputs, LOCTEXT("ApplyUniformTessellation_Error", "ApplyUniformTessellation: The inputs are invalid."));
+			UE::Geometry::AppendError(Debug, EGeometryScriptErrorType::InvalidInputs, LOCTEXT("ApplyUniformTessellation_Error", "ApplyUniformTessellation: The inputs are invalid"));
 			return;
 		} 
 
 		if (Tessellator.Compute() == false) 
 		{
-			UE::Geometry::AppendError(Debug, EGeometryScriptErrorType::OperationFailed, LOCTEXT("ApplyUniformTessellation_Failed", "ApplyUniformTessellation: Tessellation failed.")); 
+			UE::Geometry::AppendError(Debug, EGeometryScriptErrorType::OperationFailed, LOCTEXT("ApplyUniformTessellation_Failed", "ApplyUniformTessellation: Tessellation failed")); 
 		}
 	}, EDynamicMeshChangeType::GeneralEdit, EDynamicMeshAttributeChangeFlags::Unknown, false);
 
@@ -148,12 +148,12 @@ UDynamicMesh* UGeometryScriptLibrary_MeshSubdivideFunctions::ApplyAdaptiveTessel
 	}
 	if (PatternType != EAdaptiveTessellatePatternType::ConcentricRings)
 	{
-		UE::Geometry::AppendError(Debug, EGeometryScriptErrorType::InvalidInputs, LOCTEXT("ApplyAdapativeTessellation_Error", "Only ConcentricRings pattern is currently supported."));
+		UE::Geometry::AppendError(Debug, EGeometryScriptErrorType::InvalidInputs, LOCTEXT("ApplyAdapativeTessellation_InvalidPatternType", "Only ConcentricRings pattern is currently supported"));
 		return TargetMesh;
 	} 
 	if (IndexList.IndexType != EGeometryScriptIndexType::Triangle)
 	{
-		UE::Geometry::AppendError(Debug, EGeometryScriptErrorType::InvalidInputs, LOCTEXT("ApplyAdapativeTessellation_Error", "Only Triangle selection is currently supported."));
+		UE::Geometry::AppendError(Debug, EGeometryScriptErrorType::InvalidInputs, LOCTEXT("ApplyAdapativeTessellation_InvalidSelectionType", "Only Triangle selection is currently supported"));
 		return TargetMesh;
 	} 
 
@@ -219,13 +219,13 @@ UDynamicMesh* UGeometryScriptLibrary_MeshSubdivideFunctions::ApplyAdaptiveTessel
 		
 		if (Tessellator.Validate() != EOperationValidationResult::Ok)
 		{
-			UE::Geometry::AppendError(Debug, EGeometryScriptErrorType::InvalidInputs, LOCTEXT("ApplyAdapativeTessellation_Error", "ApplyAdapativeTessellation: The inputs are invalid."));
+			UE::Geometry::AppendError(Debug, EGeometryScriptErrorType::InvalidInputs, LOCTEXT("ApplyAdapativeTessellation_Error", "ApplyAdapativeTessellation: The inputs are invalid"));
 			return;
 		} 
 
 		if (Tessellator.Compute() == false)
 		{
-			UE::Geometry::AppendError(Debug, EGeometryScriptErrorType::OperationFailed, LOCTEXT("ApplyAdapativeTessellate_Failed", "ApplyAdapativeTessellate: Tessellation failed.")); 
+			UE::Geometry::AppendError(Debug, EGeometryScriptErrorType::OperationFailed, LOCTEXT("ApplyAdapativeTessellate_Failed", "ApplyAdapativeTessellate: Tessellation failed")); 
 		}
 
 	}, EDynamicMeshChangeType::GeneralEdit, EDynamicMeshAttributeChangeFlags::Unknown, false);
