GeometryScript: Remove async bake functions from MeshBakeFunctions function library. Intent is to replace these with a more holistic solution (subsystem) at a later date.

diff --git a/Engine/Plugins/Experimental/GeometryScripting/Source/GeometryScriptingCore/Private/MeshBakeFunctions.cpp b/Engine/Plugins/Experimental/GeometryScripting/Source/GeometryScriptingCore/Private/MeshBakeFunctions.cpp
--- a/Engine/Plugins/Experimental/GeometryScripting/Source/GeometryScriptingCore/Private/MeshBakeFunctions.cpp
+++ b/Engine/Plugins/Experimental/GeometryScripting/Source/GeometryScriptingCore/Private/MeshBakeFunctions.cpp
@@ -4,7 +4,8 @@
 #include "UDynamicMesh.h"
 
 #include "Sampling/MeshBakerCommon.h"
-#include "Sampling/MeshConstantMapEvaluator.h"
+#include "Sampling/MeshMapBaker.h"
+#include "Sampling/MeshVertexBaker.h"
 #include "Sampling/MeshCurvatureMapEvaluator.h"
 #include "Sampling/MeshOcclusionMapEvaluator.h"
 #include "Sampling/MeshNormalMapEvaluator.h"
@@ -15,7 +16,6 @@
 
 #include "AssetUtils/Texture2DBuilder.h"
 #include "AssetUtils/Texture2DUtil.h"
-#include "Async/Async.h"
 
 using namespace UE::Geometry;
 
@@ -911,73 +911,6 @@ TArray<UTexture2D*> UGeometryScriptLibrary_MeshBakeFunctions::BakeTexture(
 	return TextureOutput;
 }
 
-void UGeometryScriptLibrary_MeshBakeFunctions::BakeTextureAsyncBegin(
-	const FBakeTextureDelegate& Completed,
-	const int BakeId,
-	UDynamicMesh* TargetMesh,
-	FTransform TargetTransform,
-	FGeometryScriptBakeTargetMeshOptions TargetOptions,
-	UDynamicMesh* SourceMesh,
-	FTransform SourceTransform,
-	FGeometryScriptBakeSourceMeshOptions SourceOptions,
-	const TArray<FGeometryScriptBakeTypeOptions>& BakeTypes,
-	FGeometryScriptBakeTextureOptions BakeOptions,
-	UGeometryScriptDebug* Debug)
-{
-	AsyncTask(ENamedThreads::AnyBackgroundThreadNormalTask, [Completed, BakeId, TargetMesh, TargetTransform, TargetOptions, SourceMesh, SourceTransform, SourceOptions, BakeTypes, BakeOptions, Debug]()
-	{
-		// FGeometryScriptDebug is not thread-safe. Instantiate a local
-		// debug message array to collect any debug errors and append those
-		// messages back on the game thread.
-		TSharedPtr<TArray<FGeometryScriptDebugMessage>, ESPMode::ThreadSafe> DebugMessages;
-		if (Debug)
-		{
-			DebugMessages = MakeShared<TArray<FGeometryScriptDebugMessage>, ESPMode::ThreadSafe>();
-		}
-		
-		TUniquePtr<FMeshMapBaker> Baker = GeometryScriptBakeLocals::BakeTextureImpl(
-			TargetMesh,
-			TargetTransform,
-			TargetOptions,
-			SourceMesh,
-			SourceTransform,
-			SourceOptions,
-			BakeTypes,
-			BakeOptions,
-			DebugMessages.Get());
-
-		TSharedPtr<FMeshMapBaker> SharedBaker = MakeShareable<FMeshMapBaker>(Baker.Release());
-
-		AsyncTask(ENamedThreads::GameThread, [Completed, BakeId, SharedBaker, BakeOptions, TargetMesh, SourceMesh, Debug, DebugMessages]()
-		{
-			if (Debug && DebugMessages)
-			{
-				for (const FGeometryScriptDebugMessage& Message : *DebugMessages)
-				{
-					Debug->Messages.Add(Message);
-				}
-			}
-			
-			FGeometryScriptBakeTextureAsyncResult Result;
-			Result.BakeOptions = BakeOptions;
-			Result.BakeResult = SharedBaker;
-			Completed.ExecuteIfBound(BakeId, Result);
-		});
-	});
-}
-
-TArray<UTexture2D*> UGeometryScriptLibrary_MeshBakeFunctions::BakeTextureAsyncEnd(
-	const FGeometryScriptBakeTextureAsyncResult& Result)
-{
-	FMeshMapBaker* Baker = Result.BakeResult.Get();
-	const FGeometryScriptBakeTextureOptions BakeOptions = Result.BakeOptions;
-
-	TArray<UTexture2D*> TextureOutput;
-	GeometryScriptBakeLocals::GetTexturesFromBaker(Baker, BakeOptions.BitDepth, TextureOutput);
-
-	return TextureOutput;
-}
-
 UDynamicMesh* UGeometryScriptLibrary_MeshBakeFunctions::BakeVertex(
 		UDynamicMesh* TargetMesh,
 		FTransform TargetTransform,
@@ -1006,71 +939,5 @@ UDynamicMesh* UGeometryScriptLibrary_MeshBakeFunctions::BakeVertex(
 	return TargetMesh;
 }
 
-void UGeometryScriptLibrary_MeshBakeFunctions::BakeVertexAsyncBegin(
-	const FBakeVertexDelegate& Completed,
-	int BakeId,
-	UDynamicMesh* TargetMesh,
-	FTransform TargetTransform,
-	FGeometryScriptBakeTargetMeshOptions TargetOptions,
-	UDynamicMesh* SourceMesh,
-	FTransform SourceTransform,
-	FGeometryScriptBakeSourceMeshOptions SourceOptions,
-	FGeometryScriptBakeOutputType BakeTypes,
-	FGeometryScriptBakeVertexOptions BakeOptions,
-	UGeometryScriptDebug* Debug)
-{
-	AsyncTask(ENamedThreads::AnyBackgroundThreadNormalTask, [Completed, BakeId, TargetMesh, TargetTransform, TargetOptions, SourceMesh, SourceTransform, SourceOptions, BakeTypes, BakeOptions, Debug]()
-	{
-		// FGeometryScriptDebug is not thread-safe. Instantiate a local
-		// debug message array to collect any debug errors and append those
-		// messages back on the game thread.
-		TSharedPtr<TArray<FGeometryScriptDebugMessage>, ESPMode::ThreadSafe> DebugMessages;
-		if (Debug)
-		{
-			DebugMessages = MakeShared<TArray<FGeometryScriptDebugMessage>, ESPMode::ThreadSafe>();
-		}
-		
-		TUniquePtr<FMeshVertexBaker> Baker = GeometryScriptBakeLocals::BakeVertexImpl(
-			TargetMesh,
-			TargetTransform,
-			TargetOptions,
-			SourceMesh,
-			SourceTransform,
-			SourceOptions,
-			BakeTypes,
-			BakeOptions,
-			DebugMessages.Get());
-
-		TSharedPtr<FMeshVertexBaker> SharedBaker = MakeShareable<FMeshVertexBaker>(Baker.Release());
-
-		AsyncTask(ENamedThreads::GameThread, [Completed, BakeId, SharedBaker, TargetMesh, Debug, DebugMessages]()
-		{
-			if (Debug && DebugMessages)
-			{
-				for (const FGeometryScriptDebugMessage& Message : *DebugMessages)
-				{
-					Debug->Messages.Add(Message);
-				}
-			}
-			
-			FGeometryScriptBakeVertexAsyncResult Result;
-			Result.BakeResult = SharedBaker;
-			Result.TargetMesh = TargetMesh;
-			Completed.ExecuteIfBound(BakeId, Result);
-		});
-	});
-}
-
-UDynamicMesh* UGeometryScriptLibrary_MeshBakeFunctions::BakeVertexAsyncEnd(const FGeometryScriptBakeVertexAsyncResult& Result)
-{
-	FMeshVertexBaker* Baker = Result.BakeResult.Get();
-
-	// Extract the vertex bake data and apply to target mesh.
-	GeometryScriptBakeLocals::ApplyVertexBakeToMesh(Baker, Result.TargetMesh);
-
-	return Result.TargetMesh;
-}
-
-
 
 #undef LOCTEXT_NAMESPACE
diff --git a/Engine/Plugins/Experimental/GeometryScripting/Source/GeometryScriptingCore/Public/GeometryScript/MeshBakeFunctions.h b/Engine/Plugins/Experimental/GeometryScripting/Source/GeometryScriptingCore/Public/GeometryScript/MeshBakeFunctions.h
--- a/Engine/Plugins/Experimental/GeometryScripting/Source/GeometryScriptingCore/Public/GeometryScript/MeshBakeFunctions.h
+++ b/Engine/Plugins/Experimental/GeometryScripting/Source/GeometryScriptingCore/Public/GeometryScript/MeshBakeFunctions.h
@@ -5,8 +5,6 @@
 #include "CoreMinimal.h"
 #include "Kismet/BlueprintFunctionLibrary.h"
 #include "GeometryScript/GeometryScriptTypes.h"
-#include "Sampling/MeshMapBaker.h"
-#include "Sampling/MeshVertexBaker.h"
 #include "MeshBakeFunctions.generated.h"
 
 class UDynamicMesh;
@@ -298,27 +296,6 @@ struct GEOMETRYSCRIPTINGCORE_API FGeometryScriptBakeSourceMeshOptions
 	EGeometryScriptBakeNormalSpace SourceNormalSpace = EGeometryScriptBakeNormalSpace::Tangent;
 };
 
-USTRUCT(BlueprintType)
-struct GEOMETRYSCRIPTINGCORE_API FGeometryScriptBakeTextureAsyncResult
-{
-	GENERATED_BODY();
-
-	FGeometryScriptBakeTextureOptions BakeOptions;
-	TSharedPtr<UE::Geometry::FMeshMapBaker, ESPMode::ThreadSafe> BakeResult;
-};
-
-USTRUCT(BlueprintType)
-struct GEOMETRYSCRIPTINGCORE_API FGeometryScriptBakeVertexAsyncResult
-{
-	GENERATED_BODY();
-
-	TSharedPtr<UE::Geometry::FMeshVertexBaker, ESPMode::ThreadSafe> BakeResult;
-	TObjectPtr<UDynamicMesh> TargetMesh;
-};
-
-DECLARE_DYNAMIC_DELEGATE_TwoParams(FBakeTextureDelegate, int, BakeId, FGeometryScriptBakeTextureAsyncResult, ResultOut);
-DECLARE_DYNAMIC_DELEGATE_TwoParams(FBakeVertexDelegate, int, BakeId, FGeometryScriptBakeVertexAsyncResult, ResultOut);
-
 UCLASS(meta = (ScriptName = "GeometryScript_Bake"))
 class GEOMETRYSCRIPTINGCORE_API UGeometryScriptLibrary_MeshBakeFunctions : public UBlueprintFunctionLibrary
 {
@@ -385,54 +362,6 @@ public:
 		FGeometryScriptBakeTextureOptions BakeOptions,
 		UGeometryScriptDebug* Debug = nullptr);
 
-	/** 
-	 * BakeTextureAsyncBegin() is the entry point for an asynchronous variant of BakeTexture.
-	 * Usage of this method is as follows:
-	 *
-	 * 1. BakeTextureAsyncBegin() kicks off an async compute.
-	 * 2. When the async compute is complete, it invokes the provided delegate back on the game thread.
-	 * 3. The delegate output is consumed by BakeTextureAsyncEnd() which converts the bake results into UTexture2D.
-	 * 4. An optional BakeId can be provided to BakeTextureAsyncBegin() that will be associated with the async compute.
-	 *    The BakeId can be used to distinguish multiple async computes on the same delegate.
-	 *
-	 * @param Completed the delegate to invoke on completion. Delegate signature: (int, FGeometryScriptBakeTextureAsyncResult)
-	 * @param BakeId optional integer identifier to aid in distinguishing multiple async computes on the same delegate.
-	 * @param TargetMesh the target mesh to bake to.
-	 * @param TargetTransform the transform of the target mesh. Only applicable when projecting in world space.
-	 * @param TargetOptions target mesh options such as which UV layer to bake with.
-	 * @param SourceMesh the source mesh to bake from.
-	 * @param SourceTransform the transform of the source mesh. Only applicable when projecting in world space.
-	 * @param SourceOptions source mesh options such source normal map.
-	 * @param BakeTypes the array of bake types & type specific parameters to bake
-	 * @param BakeOptions bake options (ex. resolution, number of samples, projection distance).
-	 * @param Debug debug structure to capture debug messages.
-	 */
-	UFUNCTION(/*BlueprintCallable,*/ Category = "GeometryScript|Bake")
-	static void BakeTextureAsyncBegin(
-		const FBakeTextureDelegate& Completed,
-		int BakeId,
-		UDynamicMesh* TargetMesh,
-		FTransform TargetTransform,
-		FGeometryScriptBakeTargetMeshOptions TargetOptions,
-		UDynamicMesh* SourceMesh,
-		FTransform SourceTransform,
-		FGeometryScriptBakeSourceMeshOptions SourceOptions,
-		const TArray<FGeometryScriptBakeTypeOptions>& BakeTypes,
-		FGeometryScriptBakeTextureOptions BakeOptions,
-		UGeometryScriptDebug* Debug = nullptr);
-
-	/**
-	 * Converts the bake results of BakeTextureAsyncBegin() into UTexture2D.
-	 * This function is intended to be invoked by the delegate passed to
-	 * BakeTextureAsyncBegin()
-	 *
-	 * @param Result the result of a BakeTextureAsyncBegin() via delegate
-	 * @return array of UTexture2D corresponding to the requested bake types in BakeTextureAsyncBegin().
-	 */
-	UFUNCTION(/*BlueprintCallable,*/ Category = "GeometryScript|Bake")
-	static UPARAM(DisplayName="Textures Out") TArray<UTexture2D*> BakeTextureAsyncEnd(
-		const FGeometryScriptBakeTextureAsyncResult& Result);
-
 	UFUNCTION(BlueprintCallable, Category = "GeometryScript|Bake")
 	static UPARAM(DisplayName="Target Mesh") UDynamicMesh* BakeVertex(
 		UDynamicMesh* TargetMesh,
@@ -444,22 +373,4 @@ public:
 		FGeometryScriptBakeOutputType BakeTypes,
 		FGeometryScriptBakeVertexOptions BakeOptions,
 		UGeometryScriptDebug* Debug = nullptr);
-
-	UFUNCTION(/*BlueprintCallable,*/ Category = "GeometryScript|Bake")
-	static void BakeVertexAsyncBegin(
-		const FBakeVertexDelegate& Completed,
-		int BakeId,
-		UDynamicMesh* TargetMesh,
-		FTransform TargetTransform,
-		FGeometryScriptBakeTargetMeshOptions TargetOptions,
-		UDynamicMesh* SourceMesh,
-		FTransform SourceTransform,
-		FGeometryScriptBakeSourceMeshOptions SourceOptions,
-		FGeometryScriptBakeOutputType BakeTypes,
-		FGeometryScriptBakeVertexOptions BakeOptions,
-		UGeometryScriptDebug* Debug = nullptr);
-
-	UFUNCTION(/*BlueprintCallable,*/ Category = "GeometryScript|Bake")
-	static UPARAM(DisplayName="Target Mesh") UDynamicMesh* BakeVertexAsyncEnd(
-		const FGeometryScriptBakeVertexAsyncResult& Result);
 };
