TObjectPtr fix

diff --git a/Engine/Plugins/Experimental/GeometryScripting/Source/GeometryScriptingCore/Public/GeometryScript/MeshAssetFunctions.h b/Engine/Plugins/Experimental/GeometryScripting/Source/GeometryScriptingCore/Public/GeometryScript/MeshAssetFunctions.h
--- a/Engine/Plugins/Experimental/GeometryScripting/Source/GeometryScriptingCore/Public/GeometryScript/MeshAssetFunctions.h
+++ b/Engine/Plugins/Experimental/GeometryScripting/Source/GeometryScriptingCore/Public/GeometryScript/MeshAssetFunctions.h
@@ -69,7 +69,7 @@ public:
 	bool bReplaceMaterials = false;
 
 	UPROPERTY(EditAnywhere, BlueprintReadWrite, Category = Options)
-	TArray<UMaterialInterface*> NewMaterials;
+	TArray<TObjectPtr<UMaterialInterface>> NewMaterials;
 
 	UPROPERTY(EditAnywhere, BlueprintReadWrite, Category = Options)
 	TArray<FName> NewMaterialSlotNames;
