GeometryScript: add missing ListUtil functions for UV, Color, and Index lists

diff --git a/Engine/Plugins/Experimental/GeometryScripting/Source/GeometryScriptingCore/Private/ListUtilityFunctions.cpp b/Engine/Plugins/Experimental/GeometryScripting/Source/GeometryScriptingCore/Private/ListUtilityFunctions.cpp
--- a/Engine/Plugins/Experimental/GeometryScripting/Source/GeometryScriptingCore/Private/ListUtilityFunctions.cpp
+++ b/Engine/Plugins/Experimental/GeometryScripting/Source/GeometryScriptingCore/Private/ListUtilityFunctions.cpp
@@ -29,6 +29,17 @@ int UGeometryScriptLibrary_ListUtilityFunctions::GetIndexListItem(FGeometryScrip
 	return -1;
 }
 
+void UGeometryScriptLibrary_ListUtilityFunctions::SetIndexListItem(FGeometryScriptIndexList& IndexList, int Index, int NewValue, bool& bIsValidIndex)
+{
+	bIsValidIndex = false;
+	if (IndexList.List.IsValid() && Index >= 0 && Index < IndexList.List->Num())
+	{
+		bIsValidIndex = true;
+		(*IndexList.List)[Index] = NewValue;
+	}
+}
+
+
 void UGeometryScriptLibrary_ListUtilityFunctions::ConvertIndexListToArray(FGeometryScriptIndexList IndexList, TArray<int>& IndexArray)
 {
 	IndexArray.Reset();
@@ -44,6 +55,18 @@ void UGeometryScriptLibrary_ListUtilityFunctions::ConvertArrayToIndexList(const
 	IndexList.List->Append(IndexArray);
 }
 
+void UGeometryScriptLibrary_ListUtilityFunctions::DuplicateIndexList(FGeometryScriptIndexList IndexList, FGeometryScriptIndexList& DuplicateList)
+{
+	DuplicateList.Reset(IndexList.IndexType);
+	*DuplicateList.List = *IndexList.List;
+}
+
+void UGeometryScriptLibrary_ListUtilityFunctions::ClearIndexList(FGeometryScriptIndexList& IndexList, int ClearValue)
+{
+	int32 Num = IndexList.List->Num();
+	IndexList.Reset(IndexList.IndexType);
+	IndexList.List->Init(ClearValue, Num);
+}
 
 
 
@@ -234,6 +257,16 @@ FVector2D UGeometryScriptLibrary_ListUtilityFunctions::GetUVListItem(FGeometrySc
 	return FVector2D::ZeroVector;
 }
 
+void UGeometryScriptLibrary_ListUtilityFunctions::SetUVListItem(FGeometryScriptUVList& UVList, int Index, FVector2D NewUV, bool& bIsValidIndex)
+{
+	bIsValidIndex = false;
+	if (UVList.List.IsValid() && Index >= 0 && Index < UVList.List->Num())
+	{
+		bIsValidIndex = true;
+		(*UVList.List)[Index] = NewUV;
+	}
+}
+
 void UGeometryScriptLibrary_ListUtilityFunctions::ConvertUVListToArray(FGeometryScriptUVList UVList, TArray<FVector2D>& UVArray)
 {
 	UVArray.Reset();
@@ -249,6 +282,18 @@ void UGeometryScriptLibrary_ListUtilityFunctions::ConvertArrayToUVList(const TAr
 	UVList.List->Append(UVArray);
 }
 
+void UGeometryScriptLibrary_ListUtilityFunctions::DuplicateUVList(FGeometryScriptUVList UVList, FGeometryScriptUVList& DuplicateList)
+{
+	DuplicateList.Reset();
+	*DuplicateList.List = *UVList.List;
+}
+
+void UGeometryScriptLibrary_ListUtilityFunctions::ClearUVList(FGeometryScriptUVList& UVList, FVector2D ClearUV)
+{
+	int32 Num = UVList.List->Num();
+	UVList.Reset();
+	UVList.List->Init(ClearUV, Num);
+}
 
 
 
@@ -274,6 +319,17 @@ FLinearColor UGeometryScriptLibrary_ListUtilityFunctions::GetColorListItem(FGeom
 	return FLinearColor::White;
 }
 
+void UGeometryScriptLibrary_ListUtilityFunctions::SetColorListItem(FGeometryScriptColorList& ColorList, int Index, FLinearColor NewColor, bool& bIsValidIndex)
+{
+	bIsValidIndex = false;
+	if (ColorList.List.IsValid() && Index >= 0 && Index < ColorList.List->Num())
+	{
+		bIsValidIndex = true;
+		(*ColorList.List)[Index] = NewColor;
+	}
+}
+
+
 void UGeometryScriptLibrary_ListUtilityFunctions::ConvertColorListToArray(FGeometryScriptColorList ColorList, TArray<FLinearColor>& ColorArray)
 {
 	ColorArray.Reset();
@@ -289,6 +345,21 @@ void UGeometryScriptLibrary_ListUtilityFunctions::ConvertArrayToColorList(const
 	ColorList.List->Append(ColorArray);
 }
 
+
+void UGeometryScriptLibrary_ListUtilityFunctions::DuplicateColorList(FGeometryScriptColorList ColorList, FGeometryScriptColorList& DuplicateList)
+{
+	DuplicateList.Reset();
+	*DuplicateList.List = *ColorList.List;
+}
+
+void UGeometryScriptLibrary_ListUtilityFunctions::ClearColorList(FGeometryScriptColorList& ColorList, FLinearColor ClearColor)
+{
+	int32 Num = ColorList.List->Num();
+	ColorList.Reset();
+	ColorList.List->Init(ClearColor, Num);
+}
+
+
 void UGeometryScriptLibrary_ListUtilityFunctions::ExtractColorListChannel(FGeometryScriptColorList ColorList, FGeometryScriptScalarList& ScalarList, int32 ChannelIndex)
 {
 	if (ChannelIndex < 0 || ChannelIndex > 3)
diff --git a/Engine/Plugins/Experimental/GeometryScripting/Source/GeometryScriptingCore/Public/GeometryScript/ListUtilityFunctions.h b/Engine/Plugins/Experimental/GeometryScripting/Source/GeometryScriptingCore/Public/GeometryScript/ListUtilityFunctions.h
--- a/Engine/Plugins/Experimental/GeometryScripting/Source/GeometryScriptingCore/Public/GeometryScript/ListUtilityFunctions.h
+++ b/Engine/Plugins/Experimental/GeometryScripting/Source/GeometryScriptingCore/Public/GeometryScript/ListUtilityFunctions.h
@@ -25,12 +25,20 @@ public:
 	UFUNCTION(BlueprintPure, Category = "GeometryScript|ListUtils", meta=(ScriptMethod))
 	static int GetIndexListItem(FGeometryScriptIndexList IndexList, int Index, bool& bIsValidIndex);
 
+	UFUNCTION(BlueprintCallable, Category = "GeometryScript|ListUtils", meta=(ScriptMethod))
+	static void SetIndexListItem(UPARAM(ref) FGeometryScriptIndexList& IndexList, int Index, int NewValue, bool& bIsValidIndex);
+
 	UFUNCTION(BlueprintCallable, Category = "GeometryScript|ListUtils", meta=(ScriptMethod))
 	static void ConvertIndexListToArray(FGeometryScriptIndexList IndexList, TArray<int>& IndexArray);
 
 	UFUNCTION(BlueprintCallable, Category = "GeometryScript|ListUtils")
 	static void ConvertArrayToIndexList(const TArray<int>& IndexArray, FGeometryScriptIndexList& IndexList, EGeometryScriptIndexType IndexType = EGeometryScriptIndexType::Any);
 
+	UFUNCTION(BlueprintCallable, Category = "GeometryScript|ListUtils", meta=(ScriptMethod))
+	static void DuplicateIndexList(FGeometryScriptIndexList IndexList, FGeometryScriptIndexList& DuplicateList);
+
+	UFUNCTION(BlueprintCallable, Category = "GeometryScript|ListUtils", meta=(ScriptMethod))
+	static void ClearIndexList(UPARAM(ref) FGeometryScriptIndexList& IndexList, int ClearValue = 0);
 
 
 	UFUNCTION(BlueprintPure, Category = "GeometryScript|ListUtils", meta=(ScriptMethod))
@@ -113,12 +121,20 @@ public:
 	UFUNCTION(BlueprintPure, Category = "GeometryScript|ListUtils", meta=(ScriptMethod))
 	static FVector2D GetUVListItem(FGeometryScriptUVList UVList, int Index, bool& bIsValidIndex);
 
+	UFUNCTION(BlueprintCallable, Category = "GeometryScript|ListUtils", meta=(ScriptMethod))
+	static void SetUVListItem(UPARAM(ref) FGeometryScriptUVList& UVList, int Index, FVector2D NewUV, bool& bIsValidIndex);
+
 	UFUNCTION(BlueprintCallable, Category = "GeometryScript|ListUtils", meta=(ScriptMethod))
 	static void ConvertUVListToArray(FGeometryScriptUVList UVList, TArray<FVector2D>& UVArray);
 
 	UFUNCTION(BlueprintCallable, Category = "GeometryScript|ListUtils")
 	static void ConvertArrayToUVList(const TArray<FVector2D>& UVArray, FGeometryScriptUVList& UVList);
 
+	UFUNCTION(BlueprintCallable, Category = "GeometryScript|ListUtils", meta=(ScriptMethod))
+	static void DuplicateUVList(FGeometryScriptUVList UVList, FGeometryScriptUVList& DuplicateList);
+
+	UFUNCTION(BlueprintCallable, Category = "GeometryScript|ListUtils", meta=(ScriptMethod))
+	static void ClearUVList(UPARAM(ref) FGeometryScriptUVList& UVList, FVector2D ClearUV);
 
 
 
@@ -131,12 +147,21 @@ public:
 	UFUNCTION(BlueprintPure, Category = "GeometryScript|ListUtils", meta=(ScriptMethod))
 	static FLinearColor GetColorListItem(FGeometryScriptColorList ColorList, int Index, bool& bIsValidIndex);
 
+	UFUNCTION(BlueprintCallable, Category = "GeometryScript|ListUtils", meta=(ScriptMethod))
+	static void SetColorListItem(UPARAM(ref) FGeometryScriptColorList& ColorList, int Index, FLinearColor NewColor, bool& bIsValidIndex);
+
 	UFUNCTION(BlueprintCallable, Category = "GeometryScript|ListUtils", meta=(ScriptMethod))
 	static void ConvertColorListToArray(FGeometryScriptColorList ColorList, TArray<FLinearColor>& ColorArray);
 
 	UFUNCTION(BlueprintCallable, Category = "GeometryScript|ListUtils")
 	static void ConvertArrayToColorList(const TArray<FLinearColor>& ColorArray, FGeometryScriptColorList& ColorList);
 
+	UFUNCTION(BlueprintCallable, Category = "GeometryScript|ListUtils", meta=(ScriptMethod))
+	static void DuplicateColorList(FGeometryScriptColorList ColorList, FGeometryScriptColorList& DuplicateList);
+
+	UFUNCTION(BlueprintCallable, Category = "GeometryScript|ListUtils", meta=(ScriptMethod))
+	static void ClearColorList(UPARAM(ref) FGeometryScriptColorList& ColorList, FLinearColor ClearColor);
+
 	UFUNCTION(BlueprintCallable, Category = "GeometryScript|ListUtils", meta=(ScriptMethod))
 	static void ExtractColorListChannel(FGeometryScriptColorList ColorList, FGeometryScriptScalarList& ScalarList, int32 ChannelIndex = 0);
 
