GeometryScript: add types and functions for creating mesh "Selections", ie subsets defined by a list of element IDs

GeometryScriptSelectionTypes.h:
New type FGeometryScriptMeshSelection stores Vertex, Triangle, and Polygroup selections. Currently internally represented as a TSharedPtr<FGeometrySelection>, however this is not exposed to BP at all. FGeometryScriptMeshSelection is an opaque object w/ no exposed properties, and it's contents cannot be serialized, ie it is only usable as temporary data during active GeometryScript computations.

MeshSelectionFunctions.h: library of functions for creating/manipulating mesh selections
CreateSelectAllMeshSelection - create a Mesh Selection for all elements of Mesh
ConvertMeshSelection - convert a Mesh Selection to a different type
CombineMeshSelections - combine two Mesh Selections using set operations (union, difference, intersection)
ConvertIndexArrayToMeshSelection / ConvertIndexSetToMeshSelection - convert TArray<int> / TSet<int> to a Mesh Selection
ConvertIndexSetToMeshSelection - convert a Mesh Selection to TArray<int>
ConvertIndexListToMeshSelection / ConvertMeshSelectionToIndexList - convert between an Index List and a Mesh Selection
SelectMeshElementsInBox, SelectMeshElementsInSphere, SelectMeshElementsWithPlane - create a Mesh Selection of all elements "inside" a geometric element
SelectMeshElementsByNormalAngle  - create a Mesh Selection of all elements based on normal vector direction
SelectMeshElementsInsideMesh - create a Mesh Selection of all elements inside another mesh (with optional extra shell thickness)
ExpandMeshSelectionToConnected - expand a Mesh Selection to Connected regions (various connectivity criteria)
GetMeshSelectionInfo - get the type and size of a Mesh Selection
DebugPrintMeshSelection - print a Mesh Selection to the output log

MeshSelectionQueryFunctions.h: library of functions for doing mesh queries based on a selection
GetMeshSelectionBoundingBox - get the bounding box of a Mesh Selection
GetMeshSelectionBoundaryLoops - get the boundary loops (indices and 3D polylines) of a Mesh Selection

diff --git a/Engine/Plugins/Experimental/GeometryScripting/Source/GeometryScriptingCore/Private/GeometryScriptSelectionTypes.cpp b/Engine/Plugins/Experimental/GeometryScripting/Source/GeometryScriptingCore/Private/GeometryScriptSelectionTypes.cpp
new file mode 100644
--- /dev/null
+++ b/Engine/Plugins/Experimental/GeometryScripting/Source/GeometryScriptingCore/Private/GeometryScriptSelectionTypes.cpp
@@ -0,0 +1,398 @@
+// Copyright Epic Games, Inc. All Rights Reserved.
+
+#include "GeometryScript/GeometryScriptSelectionTypes.h"
+#include "Selections/GeometrySelection.h"
+#include "Selections/GeometrySelectionUtil.h"
+#include "UDynamicMesh.h"
+#include "Misc/StringBuilder.h"
+
+using namespace UE::Geometry;
+
+
+FGeometryScriptMeshSelection::FGeometryScriptMeshSelection()
+{
+	GeoSelection = MakeShared<UE::Geometry::FGeometrySelection>();
+	GeoSelection->InitializeTypes(EGeometryElementType::Face, EGeometryTopologyType::Triangle);
+}
+
+void FGeometryScriptMeshSelection::SetSelection(const FGeometryScriptMeshSelection& Selection)
+{
+	*GeoSelection = *Selection.GeoSelection;
+}
+
+void FGeometryScriptMeshSelection::SetSelection(const UE::Geometry::FGeometrySelection& Selection)
+{
+	*GeoSelection = Selection;
+}
+
+void FGeometryScriptMeshSelection::SetSelection(UE::Geometry::FGeometrySelection&& Selection)
+{
+	*GeoSelection = MoveTemp(Selection);
+}
+
+
+void FGeometryScriptMeshSelection::ClearSelection()
+{
+	GeoSelection->Reset();
+}
+
+bool FGeometryScriptMeshSelection::IsEmpty() const
+{
+	return GeoSelection->IsEmpty();
+}
+
+EGeometryScriptMeshSelectionType FGeometryScriptMeshSelection::GetSelectionType() const
+{
+	if (GeoSelection->TopologyType == EGeometryTopologyType::Polygroup)
+	{
+		if (GeoSelection->ElementType == EGeometryElementType::Face)
+		{
+			return EGeometryScriptMeshSelectionType::Polygroups;
+		}
+	}
+	else if (GeoSelection->TopologyType == EGeometryTopologyType::Triangle)
+	{
+		if (GeoSelection->ElementType == EGeometryElementType::Face)
+		{
+			return EGeometryScriptMeshSelectionType::Triangles;
+		}
+		else if (GeoSelection->ElementType == EGeometryElementType::Vertex)
+		{
+			return EGeometryScriptMeshSelectionType::Vertices;
+		}
+	}
+	UE_LOG(LogGeometry, Warning, TEXT("GeometryScriptMeshSelection::GetSelectionType() - GeoSelection has invalid type for FGeometryScriptMeshSelection!"));
+	return EGeometryScriptMeshSelectionType::Triangles;
+}
+
+
+int32 FGeometryScriptMeshSelection::GetNumSelected() const
+{
+	return GeoSelection->Num();
+}
+
+void FGeometryScriptMeshSelection::DebugPrint() const
+{
+	bool bIsPolygroups = (GeoSelection->TopologyType == EGeometryTopologyType::Polygroup);
+	if (bIsPolygroups)
+	{
+		UE_LOG(LogGeometry, Log, TEXT("FGeometryScriptSelectionList: Selection contains %d groups:"), GeoSelection->Selection.Num());
+	}
+	else
+	{
+		UE_LOG(LogGeometry, Log, TEXT("FGeometryScriptSelectionList: Selection contains %d %s"),
+			GeoSelection->Selection.Num(), ( GeoSelection->ElementType == EGeometryElementType::Vertex ? TEXT("Vertices") : TEXT("Triangles") ) );
+	}
+
+	TStringBuilder<4096> StringBuilder;
+	for (uint64 Item : GeoSelection->Selection)
+	{
+		StringBuilder.Appendf(TEXT(" %d"), (bIsPolygroups) ? FGeoSelectionID(Item).TopologyID : FGeoSelectionID(Item).GeometryID);
+	}
+	FString ResultString(StringBuilder.ToString());
+	UE_LOG(LogGeometry, Log, TEXT("Indices: %s"), *ResultString);
+}
+
+
+
+
+void FGeometryScriptMeshSelection::CombineSelectionInPlace(const FGeometryScriptMeshSelection& SelectionB, EGeometryScriptCombineSelectionMode CombineMode)
+{
+	if (ensure(GetSelectionType() == SelectionB.GetSelectionType()) == false)
+	{
+		return;
+	}
+
+	if (CombineMode == EGeometryScriptCombineSelectionMode::Add)
+	{
+		for (uint64 Item : SelectionB.GeoSelection->Selection)
+		{
+			GeoSelection->Selection.Add(Item);
+		}
+	}
+	else if (CombineMode == EGeometryScriptCombineSelectionMode::Subtract)
+	{
+		for (uint64 Item : SelectionB.GeoSelection->Selection)
+		{
+			GeoSelection->Selection.Remove(Item);
+		}
+	}
+	else if (CombineMode == EGeometryScriptCombineSelectionMode::Intersection)
+	{
+		TArray<uint64, TInlineAllocator<32>> ToRemove;
+		for (uint64 Item : GeoSelection->Selection)
+		{
+			if ( ! SelectionB.GeoSelection->Selection.Contains(Item) )
+			{
+				ToRemove.Add(Item);
+			}
+		}
+		for (uint64 Item : ToRemove)
+		{
+			GeoSelection->Selection.Remove(Item);
+		}
+	}
+
+}
+
+
+
+
+EGeometryScriptIndexType FGeometryScriptMeshSelection::ConvertToMeshIndexArray(
+	const FDynamicMesh3& Mesh, TArray<int32>& IndexListOut, EGeometryScriptIndexType ConvertToType) const
+{
+	IndexListOut.Reset();
+
+	// utility fn to use below
+	auto EnumerateGroupTriangles = [&Mesh, this](TFunctionRef<void(int&)> TriFunc)
+	{
+		TSet<int32> GroupIDs;
+		for (uint64 Item : GeoSelection->Selection)
+		{
+			GroupIDs.Add(FGeoSelectionID(Item).TopologyID);
+		}
+		for (int32 tid : Mesh.TriangleIndicesItr())
+		{
+			if ( GroupIDs.Contains( Mesh.GetTriangleGroup(tid) )  )
+			{
+				TriFunc(tid);
+			}
+		}
+	};
+
+	bool bIsSameType =
+		(ConvertToType == EGeometryScriptIndexType::Vertex && GeoSelection->TopologyType == EGeometryTopologyType::Triangle && GeoSelection->ElementType == EGeometryElementType::Vertex)
+		|| (ConvertToType == EGeometryScriptIndexType::Triangle && GeoSelection->TopologyType == EGeometryTopologyType::Triangle && GeoSelection->ElementType == EGeometryElementType::Face)
+		|| (ConvertToType == EGeometryScriptIndexType::PolygroupID && GeoSelection->TopologyType == EGeometryTopologyType::Polygroup && GeoSelection->ElementType == EGeometryElementType::Face);
+
+	if (bIsSameType || ConvertToType == EGeometryScriptIndexType::Any)
+	{
+		// keep existing type
+		if (GeoSelection->TopologyType == EGeometryTopologyType::Triangle)
+		{
+			if (GeoSelection->ElementType == EGeometryElementType::Vertex)
+			{
+				for (uint64 Item : GeoSelection->Selection)
+				{
+					int32 VertexID = FGeoSelectionID(Item).GeometryID;
+					if (Mesh.IsVertex(VertexID))
+					{
+						IndexListOut.Add( VertexID );
+					}
+				}
+				return EGeometryScriptIndexType::Vertex;
+			}
+			else if (GeoSelection->ElementType == EGeometryElementType::Face)
+			{
+				for (uint64 Item : GeoSelection->Selection)
+				{
+					int32 TriangleID = FGeoSelectionID(Item).GeometryID;
+					if (Mesh.IsTriangle(TriangleID))
+					{
+						IndexListOut.Add( TriangleID );
+					}
+				}
+				return EGeometryScriptIndexType::Triangle;
+			}
+		}
+		else if (GeoSelection->TopologyType == EGeometryTopologyType::Polygroup && GeoSelection->ElementType == EGeometryElementType::Face)
+		{
+			for (uint64 Item : GeoSelection->Selection)
+			{
+				int32 GroupID = FGeoSelectionID(Item).TopologyID;
+				IndexListOut.Add( GroupID );		// too expensive to validate GroupID...
+			}
+			return EGeometryScriptIndexType::PolygroupID;
+		}
+
+	}
+	else if (ConvertToType == EGeometryScriptIndexType::Triangle)
+	{
+		if ( (GeoSelection->TopologyType == EGeometryTopologyType::Triangle) &&
+			 (GeoSelection->ElementType == EGeometryElementType::Vertex) )
+		{
+			UE::Geometry::EnumerateSelectionTriangles(*GeoSelection, Mesh,
+				[&](int32 TriangleID) { IndexListOut.Add(TriangleID); });
+			return EGeometryScriptIndexType::Triangle;
+		}
+		else if (GeoSelection->TopologyType == EGeometryTopologyType::Polygroup && GeoSelection->ElementType == EGeometryElementType::Face)
+		{
+			// Currently doing this locally because the util functions don't support a GeoSelectionID w/ a Group ID but not a Triangle ID
+			EnumerateGroupTriangles( [&](int32 TriangleID) { IndexListOut.Add(TriangleID); });
+			return EGeometryScriptIndexType::Triangle;
+		}
+
+	}
+	else if (ConvertToType == EGeometryScriptIndexType::Vertex)
+	{
+		TSet<int32> UniqueVertices;
+
+		if ( (GeoSelection->TopologyType == EGeometryTopologyType::Triangle) &&
+			(GeoSelection->ElementType == EGeometryElementType::Face) )
+		{
+			UE::Geometry::EnumerateSelectionTriangles(*GeoSelection, Mesh,
+			[&](int32 TriangleID) { 
+				FIndex3i TriV = Mesh.GetTriangle(TriangleID);
+				UniqueVertices.Add(TriV.A); 
+				UniqueVertices.Add(TriV.B);
+				UniqueVertices.Add(TriV.C);
+			});
+			for (int32 vid : UniqueVertices)
+			{
+				IndexListOut.Add(vid);
+			}
+			return EGeometryScriptIndexType::Vertex;
+		}
+		else if (GeoSelection->TopologyType == EGeometryTopologyType::Polygroup && GeoSelection->ElementType == EGeometryElementType::Face)
+		{
+			// Currently doing this locally because the util functions don't support a GeoSelectionID w/ a Group ID but not a Triangle ID
+			EnumerateGroupTriangles([&](int32 TriangleID) {
+				FIndex3i TriV = Mesh.GetTriangle(TriangleID);
+				UniqueVertices.Add(TriV.A);
+				UniqueVertices.Add(TriV.B);
+				UniqueVertices.Add(TriV.C);
+			});
+			for (int32 vid : UniqueVertices)
+			{
+				IndexListOut.Add(vid);
+			}
+			return EGeometryScriptIndexType::Vertex;
+		}
+	}
+	else if (ConvertToType == EGeometryScriptIndexType::PolygroupID)
+	{
+		if ( (GeoSelection->TopologyType == EGeometryTopologyType::Triangle) &&
+			(GeoSelection->ElementType == EGeometryElementType::Vertex || GeoSelection->ElementType == EGeometryElementType::Face) )
+		{
+			TSet<int32> UniqueGroups;
+			UE::Geometry::EnumerateSelectionTriangles(*GeoSelection, Mesh,
+			[&](int32 TriangleID) { 
+				UniqueGroups.Add( Mesh.GetTriangleGroup(TriangleID) ); 
+			});
+			for (int32 gid : UniqueGroups)
+			{
+				IndexListOut.Add(gid);
+			}
+			return EGeometryScriptIndexType::PolygroupID;
+		}
+	}
+
+	// indicates failure;
+	UE_LOG(LogGeometry, Warning, TEXT("FGeometryScriptSelectionListInternal: Conversion not currently supported"));
+	return EGeometryScriptIndexType::Any;
+}
+
+
+
+
+void FGeometryScriptMeshSelection::ProcessByTriangleID(const UE::Geometry::FDynamicMesh3& Mesh,
+	TFunctionRef<void(int32)> PerTriangleFunc) const
+{
+	if (GetSelectionType() == EGeometryScriptMeshSelectionType::Triangles)
+	{
+		for (uint64 Item : GeoSelection->Selection)
+		{
+			int32 TriangleID = FGeoSelectionID(Item).GeometryID;
+			if (Mesh.IsTriangle(TriangleID))
+			{
+				PerTriangleFunc(TriangleID);
+			}
+		}
+	}
+	else if (GetSelectionType() == EGeometryScriptMeshSelectionType::Vertices)
+	{
+		TSet<uint64> UniqueItems;
+		for (uint64 Item : GeoSelection->Selection)
+		{
+			int32 VertexID = FGeoSelectionID(Item).GeometryID;
+			if (Mesh.IsVertex(VertexID))
+			{
+				Mesh.EnumerateVertexTriangles(VertexID, [&](int32 TriangleID)
+				{
+					bool bAlreadyInSet = false;
+					UniqueItems.Add(TriangleID, &bAlreadyInSet);
+					if (bAlreadyInSet == false)
+					{
+						PerTriangleFunc(TriangleID);
+					}
+				});
+			}
+		}
+	}
+	else
+	{
+		TSet<int32> GroupIDs;
+		for (uint64 Item : GeoSelection->Selection)
+		{
+			GroupIDs.Add(FGeoSelectionID(Item).TopologyID);
+		}
+		for (int32 TriangleID : Mesh.TriangleIndicesItr())
+		{
+			if ( GroupIDs.Contains( Mesh.GetTriangleGroup(TriangleID) )  )
+			{
+				PerTriangleFunc(TriangleID);
+			}
+		}
+	}
+}
+
+
+void FGeometryScriptMeshSelection::ProcessByVertexID(const UE::Geometry::FDynamicMesh3& Mesh,
+	TFunctionRef<void(int32)> PerVertexFunc) const
+{
+	if (GetSelectionType() == EGeometryScriptMeshSelectionType::Vertices)
+	{
+		for (uint64 Item : GeoSelection->Selection)
+		{
+			int32 VertexID = FGeoSelectionID(Item).GeometryID;
+			if (Mesh.IsVertex(VertexID))
+			{
+				PerVertexFunc(VertexID);
+			}
+		}
+	}
+	else 
+	{
+		TSet<uint64> UniqueVertices;
+		auto ProcessTriangle = [&Mesh, &PerVertexFunc, &UniqueVertices](int32 TriangleID)
+		{
+			FIndex3i Triangle = Mesh.GetTriangle(TriangleID);
+			for ( int32 j = 0; j < 3; ++j )
+			{
+				bool bAlreadyInSet = false;
+				UniqueVertices.Add(Triangle[j], &bAlreadyInSet);
+				if (bAlreadyInSet == false)
+				{
+					PerVertexFunc(Triangle[j]);
+				}
+			}
+		};
+
+		if (GetSelectionType() == EGeometryScriptMeshSelectionType::Triangles)
+		{
+			for (uint64 Item : GeoSelection->Selection)
+			{
+				int32 TriangleID = FGeoSelectionID(Item).GeometryID;
+				if (Mesh.IsTriangle(TriangleID))
+				{
+					ProcessTriangle(TriangleID);
+				}
+			}
+		}
+		else
+		{
+			TSet<int32> GroupIDs;
+			for (uint64 Item : GeoSelection->Selection)
+			{
+				GroupIDs.Add(FGeoSelectionID(Item).TopologyID);
+			}
+			for (int32 TriangleID : Mesh.TriangleIndicesItr())
+			{
+				if ( GroupIDs.Contains( Mesh.GetTriangleGroup(TriangleID) )  )
+				{
+					ProcessTriangle(TriangleID);
+				}
+			}
+		}
+	}
+}
\ No newline at end of file
diff --git a/Engine/Plugins/Experimental/GeometryScripting/Source/GeometryScriptingCore/Private/MeshSelectionFunctions.cpp b/Engine/Plugins/Experimental/GeometryScripting/Source/GeometryScriptingCore/Private/MeshSelectionFunctions.cpp
new file mode 100644
--- /dev/null
+++ b/Engine/Plugins/Experimental/GeometryScripting/Source/GeometryScriptingCore/Private/MeshSelectionFunctions.cpp
@@ -0,0 +1,812 @@
+// Copyright Epic Games, Inc. All Rights Reserved.
+
+#include "GeometryScript/MeshSelectionFunctions.h"
+#include "GeometryScript/GeometryScriptSelectionTypes.h"
+#include "Selections/GeometrySelection.h"
+#include "UDynamicMesh.h"
+#include "DynamicMesh/DynamicMeshAttributeSet.h"
+#include "DynamicMesh/MeshNormals.h"
+#include "SphereTypes.h"
+#include "DynamicMesh/DynamicMeshAABBTree3.h"
+#include "Spatial/FastWinding.h"
+#include "Selections/MeshConnectedComponents.h"
+
+using namespace UE::Geometry;
+
+#define LOCTEXT_NAMESPACE "UGeometryScriptLibrary_MeshSelectionFunctions"
+
+
+
+
+
+void UGeometryScriptLibrary_MeshSelectionFunctions::GetMeshSelectionInfo(
+	FGeometryScriptMeshSelection Selection,
+	EGeometryScriptMeshSelectionType& SelectionType,
+	int& NumSelected)
+{
+	SelectionType = Selection.GetSelectionType();
+	NumSelected = Selection.GetNumSelected();
+}
+
+
+void UGeometryScriptLibrary_MeshSelectionFunctions::DebugPrintMeshSelection(
+	FGeometryScriptMeshSelection Selection,
+	bool bDisable)
+{
+	if (!bDisable)
+	{
+		Selection.DebugPrint();
+	}
+}
+
+
+UDynamicMesh* UGeometryScriptLibrary_MeshSelectionFunctions::CreateSelectAllMeshSelection(
+	UDynamicMesh* TargetMesh,
+	FGeometryScriptMeshSelection& Selection,
+	EGeometryScriptMeshSelectionType SelectionType)
+{
+	if (TargetMesh == nullptr)
+	{
+		UE_LOG(LogGeometry, Warning, TEXT("CreateSelectAllMeshSelection: TargetMesh is Null"));
+		return TargetMesh;
+	}
+	TargetMesh->ProcessMesh([&](const FDynamicMesh3& ReadMesh)
+	{
+		FGeometrySelection NewSelection;
+		if (SelectionType == EGeometryScriptMeshSelectionType::Vertices)
+		{
+			NewSelection.InitializeTypes(EGeometryElementType::Vertex, EGeometryTopologyType::Triangle);
+			for (int32 vid : ReadMesh.VertexIndicesItr())
+			{
+				NewSelection.Selection.Add(FGeoSelectionID::MeshVertex(vid).Encoded());
+			}
+		}
+		else if (SelectionType == EGeometryScriptMeshSelectionType::Triangles)
+		{
+			NewSelection.InitializeTypes(EGeometryElementType::Face, EGeometryTopologyType::Triangle);
+			for (int32 tid : ReadMesh.TriangleIndicesItr())
+			{
+				NewSelection.Selection.Add(FGeoSelectionID::MeshTriangle(tid).Encoded());
+			}
+		}
+		else
+		{
+			NewSelection.InitializeTypes(EGeometryElementType::Face, EGeometryTopologyType::Polygroup);
+			TSet<int32> UniqueGroupIDs;
+			for (int32 tid : ReadMesh.TriangleIndicesItr())
+			{
+				int32 GroupID = ReadMesh.GetTriangleGroup(tid);
+				if (UniqueGroupIDs.Contains(GroupID) == false)
+				{
+					NewSelection.Selection.Add(FGeoSelectionID::GroupFace(tid, GroupID).Encoded());
+				}
+			}
+		}
+		Selection.SetSelection(MoveTemp(NewSelection));
+	});
+	return TargetMesh;
+}
+
+void UGeometryScriptLibrary_MeshSelectionFunctions::CombineMeshSelections(
+	FGeometryScriptMeshSelection SelectionA,
+	FGeometryScriptMeshSelection SelectionB,
+	FGeometryScriptMeshSelection& ResultSelectionOut,
+	EGeometryScriptCombineSelectionMode CombineMode)
+{
+	if (SelectionA.GetSelectionType() != SelectionB.GetSelectionType())
+	{
+		UE_LOG(LogGeometry, Warning, TEXT("CombineMeshSelections: Selections have different types, cannot combine"));
+	}
+	ResultSelectionOut.SetSelection(SelectionA);
+	ResultSelectionOut.CombineSelectionInPlace(SelectionB, CombineMode);
+}
+
+
+
+UDynamicMesh* UGeometryScriptLibrary_MeshSelectionFunctions::ConvertMeshSelection(
+	UDynamicMesh* TargetMesh,
+	FGeometryScriptMeshSelection FromSelection,
+	FGeometryScriptMeshSelection& ToSelection,
+	EGeometryScriptMeshSelectionType NewType,
+	bool bAllowPartialInclusion)
+{
+	if (TargetMesh == nullptr)
+	{
+		UE_LOG(LogGeometry, Warning, TEXT("ConvertMeshSelection: TargetMesh is Null"));
+		return TargetMesh;
+	}
+	if (FromSelection.GetSelectionType() == NewType)
+	{
+		UE_LOG(LogGeometry, Warning, TEXT("ConvertMeshSelection: Trying to convert to same type"));
+		ToSelection.SetSelection(FromSelection);
+		return TargetMesh;
+	}
+
+	if (NewType == EGeometryScriptMeshSelectionType::Vertices)
+	{
+		TSet<int32> CurTriangles, CurVertices;
+		TargetMesh->ProcessMesh([&](const FDynamicMesh3& ReadMesh)
+		{
+			if (bAllowPartialInclusion)
+			{
+				FromSelection.ProcessByTriangleID(ReadMesh, [&](int32 TriangleID) { 
+					FIndex3i Vertices = ReadMesh.GetTriangle(TriangleID);
+					CurVertices.Add(Vertices.A); CurVertices.Add(Vertices.B); CurVertices.Add(Vertices.C);
+				});
+				ConvertIndexSetToMeshSelection(TargetMesh, CurVertices, EGeometryScriptMeshSelectionType::Vertices, ToSelection);
+			}
+			else
+			{
+				FromSelection.ProcessByTriangleID(ReadMesh, [&](int32 TriangleID) { 
+					CurTriangles.Add(TriangleID);
+					FIndex3i Vertices = ReadMesh.GetTriangle(TriangleID);
+					CurVertices.Add(Vertices.A); CurVertices.Add(Vertices.B); CurVertices.Add(Vertices.C);
+				});
+
+				FGeometrySelection NewSelection;
+				NewSelection.InitializeTypes(EGeometryElementType::Vertex, EGeometryTopologyType::Triangle);
+				for (int32 vid : CurVertices)
+				{
+					bool bAllInSet = true;
+					ReadMesh.EnumerateVertexTriangles(vid, [&](int32 tid) { bAllInSet = bAllInSet && CurTriangles.Contains(tid); });
+					if (bAllInSet)
+					{
+						NewSelection.Selection.Add(FGeoSelectionID::MeshVertex(vid).Encoded());
+					}
+				}
+				ToSelection.SetSelection(MoveTemp(NewSelection));
+			}
+		});
+
+	}
+	else if (NewType == EGeometryScriptMeshSelectionType::Triangles)
+	{
+		FGeometrySelection NewSelection;
+		NewSelection.InitializeTypes(EGeometryElementType::Face, EGeometryTopologyType::Triangle);
+
+		if (FromSelection.GetSelectionType() == EGeometryScriptMeshSelectionType::Polygroups || bAllowPartialInclusion)
+		{
+			TargetMesh->ProcessMesh([&](const FDynamicMesh3& ReadMesh)
+			{
+				FromSelection.ProcessByTriangleID(ReadMesh, [&](int32 TriangleID) {
+					NewSelection.Selection.Add(FGeoSelectionID::MeshTriangle(TriangleID).Encoded());
+				});
+			});
+		}
+		else   // vertex selection w/ no partial inclusion, ie only "full" triangles
+		{
+			TargetMesh->ProcessMesh([&](const FDynamicMesh3& ReadMesh)
+			{
+				// in this case FromSelection already has this set! but we do not have access to it...
+				TSet<int32> CurVertices;
+				FromSelection.ProcessByVertexID(ReadMesh, [&](int32 VertexID) {
+					CurVertices.Add(VertexID);
+				});
+				FromSelection.ProcessByTriangleID(ReadMesh, [&](int32 TriangleID) {
+					FIndex3i Triangle = ReadMesh.GetTriangle(TriangleID);
+					if (CurVertices.Contains(Triangle.A) && CurVertices.Contains(Triangle.B) && CurVertices.Contains(Triangle.C))
+					{
+						NewSelection.Selection.Add(FGeoSelectionID::MeshTriangle(TriangleID).Encoded());
+					}
+				});
+			});
+		}
+
+		ToSelection.SetSelection(MoveTemp(NewSelection));
+	}
+	else   // polygroups
+	{
+		if (bAllowPartialInclusion)
+		{
+			TSet<int32> UniqueGroupIDs;
+			TargetMesh->ProcessMesh([&](const FDynamicMesh3& ReadMesh)
+			{
+				FromSelection.ProcessByTriangleID(ReadMesh, [&](int32 TriangleID) {
+					UniqueGroupIDs.Add(ReadMesh.GetTriangleGroup(TriangleID));
+				});
+			});
+			ConvertIndexSetToMeshSelection(TargetMesh, UniqueGroupIDs, EGeometryScriptMeshSelectionType::Polygroups, ToSelection);
+		}
+		else
+		{
+			TSet<int32> UniqueGroupIDs;
+			TargetMesh->ProcessMesh([&](const FDynamicMesh3& ReadMesh)
+			{
+				TSet<int32> AllTriangles;
+				FromSelection.ProcessByTriangleID(ReadMesh, [&](int32 TriangleID) {
+					AllTriangles.Add(TriangleID);
+					UniqueGroupIDs.Add(ReadMesh.GetTriangleGroup(TriangleID));
+				});
+
+				// if we have a non-selected triangle whose group is in our group set, that group is not fully selected
+				TSet<int32> FailGroups;
+				for (int32 TriangleID : ReadMesh.TriangleIndicesItr())
+				{
+					int32 GroupID = ReadMesh.GetTriangleGroup(TriangleID);
+					if (UniqueGroupIDs.Contains(GroupID) && AllTriangles.Contains(TriangleID) == false)
+					{
+						FailGroups.Add(GroupID);
+					}
+				}
+				for (int32 GroupID : FailGroups)
+				{
+					UniqueGroupIDs.Remove(GroupID);
+				}
+			});
+			ConvertIndexSetToMeshSelection(TargetMesh, UniqueGroupIDs, EGeometryScriptMeshSelectionType::Polygroups, ToSelection);
+		}
+	}
+
+
+	return TargetMesh;
+}
+
+
+UDynamicMesh* UGeometryScriptLibrary_MeshSelectionFunctions::ConvertIndexArrayToMeshSelection(
+	UDynamicMesh* TargetMesh,
+	const TArray<int32>& IndexArray,
+	EGeometryScriptMeshSelectionType SelectionType,
+	FGeometryScriptMeshSelection& SelectionOut)
+{
+	if (TargetMesh == nullptr)
+	{
+		UE_LOG(LogGeometry, Warning, TEXT("ConvertIndexArrayToMeshSelection: TargetMesh is Null"));
+		return TargetMesh;
+	}
+
+	FGeometrySelection GeoSelection;
+
+	switch (SelectionType)
+	{
+	case EGeometryScriptMeshSelectionType::Triangles:
+		for (int32 tid : IndexArray)
+		{
+			GeoSelection.Selection.Add( FGeoSelectionID::MeshTriangle(tid).Encoded() );
+		}
+		GeoSelection.InitializeTypes(EGeometryElementType::Face, EGeometryTopologyType::Triangle);
+		break;
+
+	case EGeometryScriptMeshSelectionType::Vertices:
+		for (int32 vid : IndexArray)
+		{
+			GeoSelection.Selection.Add( FGeoSelectionID::MeshVertex(vid).Encoded() );
+		}
+		GeoSelection.InitializeTypes(EGeometryElementType::Vertex, EGeometryTopologyType::Triangle);
+		break;
+
+	case EGeometryScriptMeshSelectionType::Polygroups:
+		for (int32 gid : IndexArray)
+		{
+			FGeoSelectionID HackGroupFace;
+			HackGroupFace.TopologyID = (uint32)gid;
+			HackGroupFace.GeometryID = 0xFFFFFFFFu;
+			GeoSelection.Selection.Add( HackGroupFace.Encoded() );
+		}
+		GeoSelection.InitializeTypes(EGeometryElementType::Face, EGeometryTopologyType::Polygroup);
+		break;
+
+	}
+
+	SelectionOut.SetSelection(MoveTemp(GeoSelection));
+	return TargetMesh;
+}
+
+
+
+UDynamicMesh* UGeometryScriptLibrary_MeshSelectionFunctions::ConvertIndexSetToMeshSelection(
+	UDynamicMesh* TargetMesh,
+	const TSet<int32>& IndexSet,
+	EGeometryScriptMeshSelectionType SelectionType,
+	FGeometryScriptMeshSelection& SelectionOut)
+{
+	if (TargetMesh == nullptr)
+	{
+		UE_LOG(LogGeometry, Warning, TEXT("ConvertIndexSetToMeshSelection: TargetMesh is Null"));
+		return TargetMesh;
+	}
+
+	FGeometrySelection GeoSelection;
+
+	switch (SelectionType)
+	{
+	case EGeometryScriptMeshSelectionType::Triangles:
+		for (int32 tid : IndexSet)
+		{
+			GeoSelection.Selection.Add( FGeoSelectionID::MeshTriangle(tid).Encoded() );
+		}
+		GeoSelection.InitializeTypes(EGeometryElementType::Face, EGeometryTopologyType::Triangle);
+		break;
+
+	case EGeometryScriptMeshSelectionType::Vertices:
+		for (int32 vid : IndexSet)
+		{
+			GeoSelection.Selection.Add( FGeoSelectionID::MeshVertex(vid).Encoded() );
+		}
+		GeoSelection.InitializeTypes(EGeometryElementType::Vertex, EGeometryTopologyType::Triangle);
+		break;
+
+	case EGeometryScriptMeshSelectionType::Polygroups:
+		for (int32 gid : IndexSet)
+		{
+			FGeoSelectionID HackGroupFace;
+			HackGroupFace.TopologyID = (uint32)gid;
+			HackGroupFace.GeometryID = 0xFFFFFFFFu;
+			GeoSelection.Selection.Add( HackGroupFace.Encoded() );
+		}
+		GeoSelection.InitializeTypes(EGeometryElementType::Face, EGeometryTopologyType::Polygroup);
+
+		break;
+	}
+
+	SelectionOut.SetSelection(MoveTemp(GeoSelection));
+	return TargetMesh;
+}
+
+
+UDynamicMesh* UGeometryScriptLibrary_MeshSelectionFunctions::ConvertMeshSelectionToIndexArray(
+	UDynamicMesh* TargetMesh,
+	FGeometryScriptMeshSelection Selection,
+	TArray<int32>& IndexArray,
+	EGeometryScriptMeshSelectionType& SelectionType)
+{
+	if (TargetMesh == nullptr)
+	{
+		UE_LOG(LogGeometry, Warning, TEXT("ConvertMeshSelectionToIndexArray: TargetMesh is Null"));
+		return TargetMesh;
+	}
+
+	TargetMesh->ProcessMesh([&](const FDynamicMesh3& Mesh)
+	{
+		Selection.ConvertToMeshIndexArray(Mesh, IndexArray);
+		SelectionType = Selection.GetSelectionType();
+	});
+
+	return TargetMesh;
+}
+
+
+
+UDynamicMesh* UGeometryScriptLibrary_MeshSelectionFunctions::ConvertIndexListToMeshSelection(
+	UDynamicMesh* TargetMesh,
+	FGeometryScriptIndexList IndexList,
+	EGeometryScriptMeshSelectionType SelectionType,
+	FGeometryScriptMeshSelection& SelectionOut)
+{
+	if (TargetMesh == nullptr)
+	{
+		UE_LOG(LogGeometry, Warning, TEXT("ConvertIndexListToMeshSelection: TargetMesh is Null"));
+		return TargetMesh;
+	}
+	if (IndexList.IndexType == EGeometryScriptIndexType::Any)
+	{
+		UE_LOG(LogGeometry, Warning, TEXT("ConvertIndexListToMeshSelection: IndexList has type Any, cannot convert"));
+		return TargetMesh;
+	}
+
+	if (IndexList.IndexType == EGeometryScriptIndexType::Vertex
+		|| IndexList.IndexType == EGeometryScriptIndexType::Triangle
+		|| IndexList.IndexType == EGeometryScriptIndexType::PolygroupID )
+	{
+		EGeometryScriptMeshSelectionType InitialType =
+			(IndexList.IndexType == EGeometryScriptIndexType::Vertex) ? EGeometryScriptMeshSelectionType::Vertices :
+				((IndexList.IndexType == EGeometryScriptIndexType::Triangle) ? EGeometryScriptMeshSelectionType::Triangles : EGeometryScriptMeshSelectionType::Polygroups);
+		if (SelectionType == InitialType)
+		{
+			ConvertIndexArrayToMeshSelection(TargetMesh, *IndexList.List, InitialType, SelectionOut);
+		}
+		else
+		{
+			FGeometryScriptMeshSelection TempSelection;
+			ConvertIndexArrayToMeshSelection(TargetMesh, *IndexList.List, InitialType, TempSelection);
+			ConvertMeshSelection(TargetMesh, TempSelection, SelectionOut, SelectionType, true);
+		}
+	}
+	else if (IndexList.IndexType == EGeometryScriptIndexType::MaterialID)
+	{
+		TArray<int32> Triangles;
+		const TArray<int32>& MaterialIDSelection = *IndexList.List;
+		TargetMesh->ProcessMesh([&](const FDynamicMesh3& ReadMesh)
+		{
+			if (ReadMesh.HasAttributes() && ReadMesh.Attributes()->HasMaterialID())
+			{
+				const FDynamicMeshMaterialAttribute* MaterialIDs = ReadMesh.Attributes()->GetMaterialID();
+				for (int32 tid : ReadMesh.TriangleIndicesItr())
+				{
+					if (MaterialIDSelection.Contains(MaterialIDs->GetValue(tid)))
+					{
+						Triangles.Add(tid);
+					}
+				}
+			}
+		});
+
+		if (SelectionType == EGeometryScriptMeshSelectionType::Triangles)
+		{
+			ConvertIndexArrayToMeshSelection(TargetMesh, Triangles, EGeometryScriptMeshSelectionType::Triangles, SelectionOut);
+		}
+		else
+		{
+			FGeometryScriptMeshSelection TempSelection;
+			ConvertIndexArrayToMeshSelection(TargetMesh, Triangles, EGeometryScriptMeshSelectionType::Triangles, TempSelection);
+			ConvertMeshSelection(TargetMesh, TempSelection, SelectionOut, SelectionType, true);
+		}
+	}
+
+	return TargetMesh;
+}
+
+
+
+
+UDynamicMesh* UGeometryScriptLibrary_MeshSelectionFunctions::ConvertMeshSelectionToIndexList(
+	UDynamicMesh* TargetMesh,
+	FGeometryScriptMeshSelection Selection,
+	FGeometryScriptIndexList& IndexList,
+	EGeometryScriptIndexType& ResultType,
+	EGeometryScriptIndexType ConvertToType)
+{
+	if (TargetMesh == nullptr)
+	{
+		UE_LOG(LogGeometry, Warning, TEXT("ConvertMeshSelectionToIndexList: TargetMesh is Null"));
+		return TargetMesh;
+	}
+
+	ResultType = EGeometryScriptIndexType::Any;
+	TArray<int32> TempArray;
+	TargetMesh->ProcessMesh([&](const FDynamicMesh3& Mesh)
+	{
+		ResultType = Selection.ConvertToMeshIndexArray(Mesh, TempArray, ConvertToType);
+	});
+
+	IndexList.Reset(ResultType);
+	*IndexList.List = MoveTemp(TempArray);
+
+	if (ResultType == EGeometryScriptIndexType::Any)
+	{
+		UE_LOG(LogGeometry, Warning, TEXT("ConvertMeshSelectionToIndexList: Conversion is not currently supported"));
+	}
+
+	return TargetMesh;
+}
+
+
+
+
+namespace UELocal
+{
+
+static void SelectMeshElementsWithContainmentTest(
+	UDynamicMesh* TargetMesh,
+	TFunctionRef<bool(const FVector3d& Position, const FVector3d& Normal)> ContainmentFunc,
+	FGeometryScriptMeshSelection& SelectionOut,
+	EGeometryScriptMeshSelectionType SelectionType,
+	int NumTrianglePoints,
+	bool bNeedsNormals)
+{
+	NumTrianglePoints = FMath::Clamp(NumTrianglePoints, 1, 3);
+	FVector3d UnitNormal = FVector3d::UnitZ();
+
+	FGeometrySelection GeoSelection;
+	if (SelectionType == EGeometryScriptMeshSelectionType::Vertices)
+	{
+		GeoSelection.InitializeTypes(EGeometryElementType::Vertex, EGeometryTopologyType::Triangle);
+
+		TargetMesh->ProcessMesh([&](const FDynamicMesh3& Mesh) 
+		{
+			for (int32 vid : Mesh.VertexIndicesItr())
+			{
+				FVector3d UseNormal = (bNeedsNormals) ? FMeshNormals::ComputeVertexNormal(Mesh,vid) : UnitNormal;
+				if ( ContainmentFunc(Mesh.GetVertex(vid), UseNormal) )
+				{
+					GeoSelection.Selection.Add( FGeoSelectionID::MeshVertex(vid).Encoded() );
+				}
+			}
+		});
+	}
+	else
+	{
+		GeoSelection.InitializeTypes(EGeometryElementType::Face, 
+			(SelectionType == EGeometryScriptMeshSelectionType::Triangles) ? EGeometryTopologyType::Triangle : EGeometryTopologyType::Polygroup);
+
+		TargetMesh->ProcessMesh([&](const FDynamicMesh3& Mesh) 
+		{
+			for (int32 tid : Mesh.TriangleIndicesItr())
+			{
+				FVector3d UseNormal = (bNeedsNormals) ? Mesh.GetTriNormal(tid) : UnitNormal;
+				FIndex3i Tri = Mesh.GetTriangle(tid);
+				// may be wasteful to test each vertex multiple times...could accumulate a cache at cost of some memory allocation...
+				int NumContained =
+					(ContainmentFunc(Mesh.GetVertex(Tri.A), UseNormal) ? 1 : 0) +
+					(ContainmentFunc(Mesh.GetVertex(Tri.B), UseNormal) ? 1 : 0) +
+					(ContainmentFunc(Mesh.GetVertex(Tri.C), UseNormal) ? 1 : 0);
+
+				if ( NumContained >= NumTrianglePoints )
+				{
+					if (SelectionType == EGeometryScriptMeshSelectionType::Triangles)
+					{
+						GeoSelection.Selection.Add( FGeoSelectionID::MeshTriangle(tid).Encoded() );
+					}
+					else
+					{
+						int32 gid = Mesh.GetTriangleGroup(tid);
+						GeoSelection.Selection.Add( FGeoSelectionID::GroupFace(tid, gid).Encoded() );
+					}
+				}
+			}
+		});
+
+	}
+
+	SelectionOut.SetSelection(MoveTemp(GeoSelection));
+}
+
+} // end namespace UELocal
+
+
+UDynamicMesh* UGeometryScriptLibrary_MeshSelectionFunctions::SelectMeshElementsInBox(
+	UDynamicMesh* TargetMesh,
+	FGeometryScriptMeshSelection& SelectionOut,
+	FBox Box,
+	EGeometryScriptMeshSelectionType SelectionType,
+	bool bInvert,
+	int MinNumTrianglePoints)
+{
+	if (TargetMesh == nullptr)
+	{
+		UE_LOG(LogGeometry, Warning, TEXT("SelectMeshElementsInBox: TargetMesh is Null"));
+		return TargetMesh;
+	}
+
+	FAxisAlignedBox3d Container(Box);
+	auto ContainsFunc = [&Container, bInvert](const FVector3d& Point, const FVector3d& Normal) 
+	{ 
+		return Container.Contains(Point) != bInvert; 
+	};
+	UELocal::SelectMeshElementsWithContainmentTest(TargetMesh, ContainsFunc,
+		SelectionOut, SelectionType, MinNumTrianglePoints, false);
+	return TargetMesh;
+}
+
+
+UDynamicMesh* UGeometryScriptLibrary_MeshSelectionFunctions::SelectMeshElementsInSphere(
+	UDynamicMesh* TargetMesh,
+	FGeometryScriptMeshSelection& SelectionOut,
+	FVector SphereOrigin,
+	double SphereRadius,
+	EGeometryScriptMeshSelectionType SelectionType,
+	bool bInvert,
+	int MinNumTrianglePoints)
+{
+	if (TargetMesh == nullptr)
+	{
+		UE_LOG(LogGeometry, Warning, TEXT("SelectMeshElementsInSphere: TargetMesh is Null"));
+		return TargetMesh;
+	}
+
+	SphereRadius = FMathd::Clamp(SphereRadius, FMathd::Epsilon, FMathd::Sqrt(TNumericLimits<double>::Max()));
+	UE::Geometry::FSphere3d Container(SphereOrigin, SphereRadius);
+
+	auto ContainsFunc = [&Container, bInvert](const FVector3d& Point, const FVector3d& Normal) 
+	{ 
+		return Container.Contains(Point) != bInvert; 
+	};
+	UELocal::SelectMeshElementsWithContainmentTest(TargetMesh, ContainsFunc,
+		SelectionOut, SelectionType, MinNumTrianglePoints, false);
+	return TargetMesh;
+}
+
+
+
+UDynamicMesh* UGeometryScriptLibrary_MeshSelectionFunctions::SelectMeshElementsWithPlane(
+	UDynamicMesh* TargetMesh,
+	FGeometryScriptMeshSelection& SelectionOut,
+	FVector PlaneOrigin,
+	FVector PlaneNormal,
+	EGeometryScriptMeshSelectionType SelectionType,
+	bool bInvert,
+	int MinNumTrianglePoints)
+{
+	if (TargetMesh == nullptr)
+	{
+		UE_LOG(LogGeometry, Warning, TEXT("SelectMeshElementsWithPlane: TargetMesh is Null"));
+		return TargetMesh;
+	}
+
+	PlaneNormal = Normalized(PlaneNormal);
+	auto ContainsFunc = [&PlaneOrigin, &PlaneNormal, bInvert](const FVector3d& Point, const FVector3d& Normal) 
+	{ 
+		bool bContains = (Point - PlaneOrigin).Dot(PlaneNormal) >= 0;
+		return bContains != bInvert;
+	};
+	UELocal::SelectMeshElementsWithContainmentTest(TargetMesh, ContainsFunc,
+		SelectionOut, SelectionType, MinNumTrianglePoints, false);
+	return TargetMesh;
+}
+
+
+
+UDynamicMesh* UGeometryScriptLibrary_MeshSelectionFunctions::SelectMeshElementsByNormalAngle(
+	UDynamicMesh* TargetMesh,
+	FGeometryScriptMeshSelection& SelectionOut,
+	FVector PlaneNormal,
+	double MaxAngleDeg,
+	EGeometryScriptMeshSelectionType SelectionType,
+	bool bInvert,
+	int MinNumTrianglePoints)
+{
+	if (TargetMesh == nullptr)
+	{
+		UE_LOG(LogGeometry, Warning, TEXT("SelectMeshElementsByNormalAngle: TargetMesh is Null"));
+		return TargetMesh;
+	}
+
+	PlaneNormal = Normalized(PlaneNormal);
+	double CosMaxAngle = FMathd::Cos(MaxAngleDeg * FMathd::DegToRad);
+	auto ContainsFunc = [&PlaneNormal, &CosMaxAngle, bInvert](const FVector3d& Point, const FVector3d& Normal) 
+	{
+		bool bContains = PlaneNormal.Dot(Normal) >= CosMaxAngle;
+		return bContains != bInvert;
+	};
+	UELocal::SelectMeshElementsWithContainmentTest(TargetMesh, ContainsFunc,
+		SelectionOut, SelectionType, MinNumTrianglePoints, true);
+	return TargetMesh;
+}
+
+
+UDynamicMesh* UGeometryScriptLibrary_MeshSelectionFunctions::SelectMeshElementsInsideMesh(
+	UDynamicMesh* TargetMesh,
+	UDynamicMesh* SelectionMesh,
+	FGeometryScriptMeshSelection& SelectionOut,
+	FTransform SelectionMeshTransform,
+	EGeometryScriptMeshSelectionType SelectionType,
+	bool bInvert,
+	double ShellDistance,
+	double WindingThreshold,
+	/** Only include triangles with this many vertices inside the Selection Mesh */
+	int MinNumTrianglePoints)
+{
+	if (TargetMesh == nullptr)
+	{
+		UE_LOG(LogGeometry, Warning, TEXT("SelectMeshElementsInsideMesh: TargetMesh is Null"));
+		return TargetMesh;
+	}
+	if (SelectionMesh == nullptr)
+	{
+		UE_LOG(LogGeometry, Warning, TEXT("SelectMeshElementsInsideMesh: SelectionMesh is Null"));
+		return TargetMesh;
+	}
+	if (SelectionMesh == TargetMesh)
+	{
+		UE_LOG(LogGeometry, Warning, TEXT("SelectMeshElementsInsideMesh: SelectionMesh == TargetMesh, this is not supported"));
+		// TODO: could select-all here?
+		return TargetMesh;
+	}
+
+	// todo: for small meshes it is possibly cheaper to make a copy?
+	FTransform InvTransform = SelectionMeshTransform.Inverse();
+
+	SelectionMesh->ProcessMesh([&](const FDynamicMesh3& ReadMesh)
+	{
+		FDynamicMeshAABBTree3 Spatial(&ReadMesh, true);
+		TFastWindingTree<FDynamicMesh3> FastWinding(&Spatial, true);
+
+		auto ContainsFunc = [&Spatial, &FastWinding, InvTransform, WindingThreshold, ShellDistance, bInvert](const FVector3d& Point, const FVector3d& Normal)
+		{
+			FVector3d LocalPoint = InvTransform.TransformPosition(Point);
+			bool bContains = FastWinding.IsInside(LocalPoint, WindingThreshold);
+			if (!bContains && ShellDistance > 0)
+			{
+				double NearestDistSqr;
+				int32 NearestTID = Spatial.FindNearestTriangle(LocalPoint, NearestDistSqr, IMeshSpatial::FQueryOptions(ShellDistance));
+				if (NearestTID >= 0 && NearestDistSqr < ShellDistance * ShellDistance)
+				{
+					bContains = true;
+				}
+			}
+			return bContains != bInvert;
+		};
+		UELocal::SelectMeshElementsWithContainmentTest(TargetMesh, ContainsFunc,
+			SelectionOut, SelectionType, MinNumTrianglePoints, true);
+	});
+
+	return TargetMesh;
+}
+
+
+
+
+UDynamicMesh* UGeometryScriptLibrary_MeshSelectionFunctions::ExpandMeshSelectionToConnected(
+	UDynamicMesh* TargetMesh,
+	FGeometryScriptMeshSelection Selection,
+	FGeometryScriptMeshSelection& NewSelection,
+	EGeometryScriptTopologyConnectionType ConnectionType)
+{
+	if (TargetMesh == nullptr)
+	{
+		UE_LOG(LogGeometry, Warning, TEXT("ExpandMeshSelectionToConnected: TargetMesh is Null"));
+		return TargetMesh;
+	}
+	if (Selection.IsEmpty())
+	{
+		UE_LOG(LogGeometry, Warning, TEXT("ExpandMeshSelectionToConnected: Initial Selection is Empty"));
+		return TargetMesh;
+	}
+	if (Selection.GetSelectionType() == EGeometryScriptMeshSelectionType::Vertices)
+	{
+		// todo: this...
+		UE_LOG(LogGeometry, Warning, TEXT("ExpandMeshSelectionToConnected: Vertex Selection currently not supported"));
+		NewSelection.SetSelection(Selection);
+		return TargetMesh;
+	}
+
+	if (Selection.GetSelectionType() == EGeometryScriptMeshSelectionType::Polygroups &&
+		ConnectionType == EGeometryScriptTopologyConnectionType::Polygroup)
+	{
+		UE_LOG(LogGeometry, Warning, TEXT("ExpandMeshSelectionToConnected: Expanding Polygroup Selection to Connected Polygroups will not change selection"));
+		NewSelection.SetSelection(Selection);	// this setup makes no sense
+		return TargetMesh;
+	}
+
+	// collect up existing triangles and (optionally) polygroups
+	TArray<int32> CurTriangles;
+	TargetMesh->ProcessMesh([&](const FDynamicMesh3& ReadMesh)
+	{
+		Selection.ProcessByTriangleID(ReadMesh, [&](int32 TriangleID) { CurTriangles.Add(TriangleID); });
+	});
+
+	TSet<int32> ResultTriangles;
+	if (ConnectionType == EGeometryScriptTopologyConnectionType::Geometric)
+	{
+		TargetMesh->ProcessMesh([&](const FDynamicMesh3& ReadMesh)
+		{
+			FMeshConnectedComponents::GrowToConnectedTriangles(&ReadMesh, CurTriangles, ResultTriangles, nullptr, [](int32, int32) { return true; });
+		});
+	}
+	else if (ConnectionType == EGeometryScriptTopologyConnectionType::Polygroup)
+	{
+		TargetMesh->ProcessMesh([&](const FDynamicMesh3& ReadMesh)
+		{
+			FMeshConnectedComponents::GrowToConnectedTriangles(&ReadMesh, CurTriangles, ResultTriangles, nullptr,
+				[&](int32 FromTriID, int32 ToTriID) {
+					return ReadMesh.GetTriangleGroup(FromTriID) == ReadMesh.GetTriangleGroup(ToTriID);
+			});
+		});
+	}
+	else if (ConnectionType == EGeometryScriptTopologyConnectionType::MaterialID)
+	{
+		TargetMesh->ProcessMesh([&](const FDynamicMesh3& ReadMesh)
+		{
+			if (ReadMesh.HasAttributes() && ReadMesh.Attributes()->GetMaterialID() != nullptr)
+			{
+				const FDynamicMeshMaterialAttribute* MaterialID = ReadMesh.Attributes()->GetMaterialID();
+				FMeshConnectedComponents::GrowToConnectedTriangles(&ReadMesh, CurTriangles, ResultTriangles, nullptr,
+					[&](int32 FromTriID, int32 ToTriID) {
+					return MaterialID->GetValue(FromTriID) == MaterialID->GetValue(ToTriID);
+				});
+			}
+		});
+	}
+
+	if (Selection.GetSelectionType() == EGeometryScriptMeshSelectionType::Triangles)
+	{
+		ConvertIndexSetToMeshSelection(TargetMesh, ResultTriangles, EGeometryScriptMeshSelectionType::Triangles, NewSelection);
+	}
+	else
+	{
+		TSet<int32> ResultGroupIDs;
+		TargetMesh->ProcessMesh([&](const FDynamicMesh3& ReadMesh)
+		{
+			for (int32 tid : ResultTriangles)
+			{
+				ResultGroupIDs.Add(ReadMesh.GetTriangleGroup(tid));
+			}
+		});
+		ConvertIndexSetToMeshSelection(TargetMesh, ResultGroupIDs, EGeometryScriptMeshSelectionType::Polygroups, NewSelection);
+	}
+
+	return TargetMesh;
+}
+
+
+
+#undef LOCTEXT_NAMESPACE
\ No newline at end of file
diff --git a/Engine/Plugins/Experimental/GeometryScripting/Source/GeometryScriptingCore/Private/MeshSelectionQueryFunctions.cpp b/Engine/Plugins/Experimental/GeometryScripting/Source/GeometryScriptingCore/Private/MeshSelectionQueryFunctions.cpp
new file mode 100644
--- /dev/null
+++ b/Engine/Plugins/Experimental/GeometryScripting/Source/GeometryScriptingCore/Private/MeshSelectionQueryFunctions.cpp
@@ -0,0 +1,125 @@
+// Copyright Epic Games, Inc. All Rights Reserved.
+
+#include "GeometryScript/MeshSelectionQueryFunctions.h"
+#include "GeometryScript/GeometryScriptSelectionTypes.h"
+#include "Selections/GeometrySelection.h"
+#include "UDynamicMesh.h"
+#include "DynamicMesh/DynamicMeshAttributeSet.h"
+#include "DynamicMesh/MeshNormals.h"
+#include "BoxTypes.h"
+#include "Selections/MeshConnectedComponents.h"
+#include "MeshRegionBoundaryLoops.h"
+
+using namespace UE::Geometry;
+
+#define LOCTEXT_NAMESPACE "UGeometryScriptLibrary_MeshSelectionQueryFunctions"
+
+
+UDynamicMesh* UGeometryScriptLibrary_MeshSelectionQueryFunctions::GetMeshSelectionBoundingBox(
+	UDynamicMesh* TargetMesh,
+	FGeometryScriptMeshSelection Selection,
+	FBox& SelectionBounds,
+	bool& bIsEmpty,
+	UGeometryScriptDebug* Debug)
+{
+	if (TargetMesh == nullptr)
+	{
+		UE::Geometry::AppendError(Debug, EGeometryScriptErrorType::InvalidInputs, LOCTEXT("GetMeshSelectionBoundingBox_InvalidInput", "GetMeshSelectionBoundingBox: TargetMesh is Null"));
+		return TargetMesh;
+	}
+
+
+	bIsEmpty = true;
+	FAxisAlignedBox3d Bounds = FAxisAlignedBox3d::Empty();
+	TargetMesh->ProcessMesh([&](const FDynamicMesh3& ReadMesh)
+	{
+		if (Selection.GetSelectionType() == EGeometryScriptMeshSelectionType::Vertices)
+		{
+			Selection.ProcessByVertexID(ReadMesh, [&](int32 VertexID)
+			{
+				Bounds.Contain(ReadMesh.GetVertex(VertexID));
+			});
+		}
+		else
+		{
+			Selection.ProcessByTriangleID(ReadMesh, [&](int32 TriangleID) {
+				Bounds.Contain(ReadMesh.GetTriBounds(TriangleID));
+			});
+		}
+	});
+	
+	if (Bounds != FAxisAlignedBox3d::Empty())
+	{
+		bIsEmpty = false;
+		SelectionBounds = (FBox)Bounds;
+	}
+	else
+	{
+		SelectionBounds = FBox(EForceInit::ForceInitToZero);
+	}
+	return TargetMesh;
+}
+
+
+
+
+
+UDynamicMesh* UGeometryScriptLibrary_MeshSelectionQueryFunctions::GetMeshSelectionBoundaryLoops(
+	UDynamicMesh* TargetMesh,
+	FGeometryScriptMeshSelection Selection,
+	TArray<FGeometryScriptIndexList>& IndexLoops,
+	TArray<FGeometryScriptPolyPath>& PathLoops,
+	int& NumLoops,
+	bool& bFoundErrors,
+	UGeometryScriptDebug* Debug)
+{
+	if (TargetMesh == nullptr)
+	{
+		UE::Geometry::AppendError(Debug, EGeometryScriptErrorType::InvalidInputs, LOCTEXT("GetMeshSelectionBoundaryLoops_InvalidInput", "GetMeshSelectionBoundaryLoops: TargetMesh is Null"));
+		return TargetMesh;
+	}
+
+	bFoundErrors = true;
+	NumLoops = 0;
+	TargetMesh->ProcessMesh([&](const FDynamicMesh3& ReadMesh)
+	{
+		TArray<int32> Triangles;
+		Selection.ProcessByTriangleID(ReadMesh, [&](int32 TriangleID) {
+			Triangles.Add(TriangleID);
+		});
+
+		// TODO: if #Triangles == Mesh.TriangleCount, use MeshBoundaryLoops
+
+		FMeshRegionBoundaryLoops Loops(&ReadMesh, Triangles, false);
+		bFoundErrors = Loops.Compute();
+		NumLoops = Loops.Num();
+		IndexLoops.Reserve(NumLoops);
+		PathLoops.Reserve(NumLoops);
+		for (int32 li = 0; li < NumLoops; ++li)
+		{
+			FEdgeLoop& Loop = Loops.Loops[li];
+
+			FGeometryScriptPolyPath LoopVertices;
+			LoopVertices.Reset();
+			LoopVertices.bClosedLoop = true;
+			LoopVertices.Path->Reserve(Loop.Vertices.Num());
+			for (int32 vid : Loop.Vertices)
+			{
+				LoopVertices.Path->Add(ReadMesh.GetVertex(vid));
+			}
+
+			FGeometryScriptIndexList LoopIndices;
+			LoopIndices.Reset(EGeometryScriptIndexType::Vertex);
+			*LoopIndices.List = MoveTemp(Loop.Vertices);
+
+			PathLoops.Add(LoopVertices);
+			IndexLoops.Add(LoopIndices);
+		}
+
+	});
+	
+	return TargetMesh;
+}
+
+
+#undef LOCTEXT_NAMESPACE
\ No newline at end of file
diff --git a/Engine/Plugins/Experimental/GeometryScripting/Source/GeometryScriptingCore/Public/GeometryScript/GeometryScriptSelectionTypes.h b/Engine/Plugins/Experimental/GeometryScripting/Source/GeometryScriptingCore/Public/GeometryScript/GeometryScriptSelectionTypes.h
new file mode 100644
--- /dev/null
+++ b/Engine/Plugins/Experimental/GeometryScripting/Source/GeometryScriptingCore/Public/GeometryScript/GeometryScriptSelectionTypes.h
@@ -0,0 +1,112 @@
+// Copyright Epic Games, Inc. All Rights Reserved.
+
+#pragma once
+
+#include "GeometryScriptTypes.h"
+#include "GeometryScriptSelectionTypes.generated.h"
+
+
+namespace UE::Geometry { struct FGeometrySelection; }
+
+/**
+ * Type of index stored in a FGeometryScriptMeshSelection
+ */
+UENUM(BlueprintType)
+enum class EGeometryScriptMeshSelectionType : uint8
+{
+	Vertices = 0,
+	Triangles = 1,
+	Polygroups = 2
+};
+
+/**
+ * Type of Conversion to apply to a FGeometryScriptMeshSelection
+ */
+UENUM(BlueprintType)
+enum class EGeometryScriptMeshSelectionConversionType : uint8
+{
+	NoConversion = 0,
+	ToVertices = 1,
+	ToTriangles = 2,
+	ToPolygroups = 3
+};
+
+/**
+ * Type of Combine operation to use when combining multiple FGeometryScriptMeshSelection
+ */
+UENUM(BlueprintType)
+enum class EGeometryScriptCombineSelectionMode : uint8
+{
+	Add,
+	Subtract,
+	Intersection
+};
+
+
+/**
+ * FGeometryScriptMeshSelection is a container for a Mesh Selection used in Geometry Script.
+ * The actual selection representation is not exposed to BP, 
+ * use functions in MeshSelectionFunctions/etc to manipulate the selection.
+ * 
+ * Internally the selection is stored as a SharedPtr to a FGeometrySelection, which
+ * stores a TSet (so unique add and remove are efficient, but the selection cannot
+ * be directly indexed without converting to an Array). 
+ *
+ * Note that the Selection storage is not a UProperty, and is not
+ * serialized. FGeometryScriptMeshSelection instances *cannot* be serialized,
+ * they are only transient data structures, that can be created and used Blueprint instances.
+ */
+USTRUCT(BlueprintType, meta = (DisplayName = "Mesh Selection"))
+struct GEOMETRYSCRIPTINGCORE_API FGeometryScriptMeshSelection
+{
+	GENERATED_BODY()
+
+	FGeometryScriptMeshSelection();
+
+	void SetSelection(const FGeometryScriptMeshSelection& Selection);
+	void SetSelection(const UE::Geometry::FGeometrySelection& Selection);
+	void SetSelection(UE::Geometry::FGeometrySelection&& Selection);
+	void ClearSelection();
+
+	bool IsEmpty() const;
+
+	EGeometryScriptMeshSelectionType GetSelectionType() const;
+	int32 GetNumSelected() const;
+	void DebugPrint() const;
+
+	/** 
+	 * Combine SelectionB with current selection, updating current selection, using CombineMode to control how combining happens
+	 */
+	void CombineSelectionInPlace(const FGeometryScriptMeshSelection& SelectionB, EGeometryScriptCombineSelectionMode CombineMode);
+
+	/**
+	 * Convert the current selection to a TArray, optionally converting to ConvertToType.
+	 * For (Tri|Group)=>Vtx, all triangle vertices (in triangles or polygroups) are included.
+	 * For Vtx=>Tri, all one-ring vertices are included. For Group=>Tri, all Triangles are found via enumerating over mesh.
+	 * (Tri|Vtx)=>Group, all GroupIDs of all triangles/one-ring triangles are included
+	 */
+	EGeometryScriptIndexType ConvertToMeshIndexArray(const UE::Geometry::FDynamicMesh3& Mesh, TArray<int32>& IndexListOut, EGeometryScriptIndexType ConvertToType = EGeometryScriptIndexType::Any) const;
+
+	/**
+	 * Call PerTriangleFunc for each TriangleID in the Selection.
+	 * For Vertex Selections, Vertex one-rings are enumerated and accumulated in a TSet.
+	 * For Polygroup Selections, a full mesh iteration is used to find all Triangles in the Groups.
+	 */
+	void ProcessByTriangleID(const UE::Geometry::FDynamicMesh3& Mesh,
+		TFunctionRef<void(int32)> PerTriangleFunc) const;
+
+	/**
+	 * Call PerVertexFunc for each VertexID in the Selection.
+	 * For Triangle Selections, Triangles Vertex tuples are enumerated and accumulated in a TSet.
+	 * For Polygroup Selections, a full mesh iteration is used to find all Triangle Vertices in the Groups (accumulated in a TSet)
+	 */
+	void ProcessByVertexID(const UE::Geometry::FDynamicMesh3& Mesh,
+		TFunctionRef<void(int32)> PerVertexFunc) const;
+
+private:
+	// keeping this private for now in case it needs to be revised in 5.2+
+	TSharedPtr<UE::Geometry::FGeometrySelection> GeoSelection;
+};
+
+
+
diff --git a/Engine/Plugins/Experimental/GeometryScripting/Source/GeometryScriptingCore/Public/GeometryScript/MeshSelectionFunctions.h b/Engine/Plugins/Experimental/GeometryScripting/Source/GeometryScriptingCore/Public/GeometryScript/MeshSelectionFunctions.h
new file mode 100644
--- /dev/null
+++ b/Engine/Plugins/Experimental/GeometryScripting/Source/GeometryScriptingCore/Public/GeometryScript/MeshSelectionFunctions.h
@@ -0,0 +1,265 @@
+// Copyright Epic Games, Inc. All Rights Reserved.
+
+#pragma once
+
+#include "CoreMinimal.h"
+#include "Kismet/BlueprintFunctionLibrary.h"
+#include "GeometryScript/GeometryScriptTypes.h"
+#include "GeometryScript/GeometryScriptSelectionTypes.h"
+#include "MeshSelectionFunctions.generated.h"
+
+
+/**
+ * Types of connection between adjacent Mesh Elements (vertices/triangles/polygoups)
+ */
+UENUM(BlueprintType)
+enum class EGeometryScriptTopologyConnectionType : uint8
+{
+	Geometric,
+	Polygroup,
+	MaterialID
+};
+
+
+
+
+UCLASS(meta = (ScriptName = "GeometryScript_MeshSelection"))
+class GEOMETRYSCRIPTINGCORE_API UGeometryScriptLibrary_MeshSelectionFunctions : public UBlueprintFunctionLibrary
+{
+	GENERATED_BODY()
+public:
+
+	/**
+	 * Query information about a Mesh Selection
+	 */
+	UFUNCTION(BlueprintCallable, Category = "GeometryScript|MeshSelection", meta=(ScriptMethod))
+	static void
+	GetMeshSelectionInfo(
+		FGeometryScriptMeshSelection Selection,
+		EGeometryScriptMeshSelectionType& SelectionType,
+		int& NumSelected);
+
+	/**
+	 * Print information about the Mesh Selection to the Output Log
+	 * @param bDisable if true, printing will be disabled (useful for debugging)
+	 */
+	UFUNCTION(BlueprintCallable, Category = "GeometryScript|MeshSelection", meta=(ScriptMethod))
+	static void
+	DebugPrintMeshSelection(
+		FGeometryScriptMeshSelection Selection,
+		bool bDisable = false);
+
+	/**
+	 * Create a Selection of the given SelectionType that contains all the mesh elements of TargetMesh
+	 */
+	UFUNCTION(BlueprintCallable, Category = "GeometryScript|MeshSelection", meta=(ScriptMethod))
+	static UPARAM(DisplayName = "Target Mesh") UDynamicMesh* 
+	CreateSelectAllMeshSelection(
+		UDynamicMesh* TargetMesh,
+		FGeometryScriptMeshSelection& Selection,
+		EGeometryScriptMeshSelectionType SelectionType = EGeometryScriptMeshSelectionType::Triangles );
+
+
+	/**
+	 * Convert a Mesh Selection to a different Type (eg Vertices to Triangles, etc)
+	 * By default, Vertices map to Triangle one-rings, and Triangles to all contained vertices.
+	 * If bAllowPartialInclusion is disabled, then more restrictive conversions are performed, as follows:
+	 *   For To-Vertices, only include vertices where all one-ring triangles are included in FromSelection.
+	 *   For To-Triangles, only include triangles where all tri vertices are included in FromSelection.
+	 *   For To-Polygroups, only include groups where all group triangles are included in FromSelection
+	 * @param bAllowPartialInclusion if false, perform more limited selection conversion as described above
+	 */
+	UFUNCTION(BlueprintCallable, Category = "GeometryScript|MeshSelection", meta=(ScriptMethod))
+	static UPARAM(DisplayName = "Target Mesh") UDynamicMesh* 
+	ConvertMeshSelection(
+		UDynamicMesh* TargetMesh,
+		FGeometryScriptMeshSelection FromSelection,
+		FGeometryScriptMeshSelection& ToSelection,
+		EGeometryScriptMeshSelectionType NewType = EGeometryScriptMeshSelectionType::Triangles,
+		bool bAllowPartialInclusion = true );
+
+	/**
+	 * Combine two Mesh Selections into a new Mesh Selection.
+	 * The two inputs SelectionA and SelectionB must have the same Type.
+	 * @param CombineMode specifies how the selection elements in SelectionA and SelectionB are interpreted/combined.
+	 */
+	UFUNCTION(BlueprintCallable, Category = "GeometryScript|MeshSelection", meta=(ScriptMethod))
+	static void
+	CombineMeshSelections(
+		FGeometryScriptMeshSelection SelectionA,
+		FGeometryScriptMeshSelection SelectionB,
+		FGeometryScriptMeshSelection& ResultSelection,
+		EGeometryScriptCombineSelectionMode CombineMode = EGeometryScriptCombineSelectionMode::Add );
+
+
+	/**
+	 * Create a Mesh Selection from the IndexArray. 
+	 * @param SelectionType type of indices specified in the selection
+	 */
+	UFUNCTION(BlueprintCallable, Category = "GeometryScript|MeshSelection", meta=(ScriptMethod))
+	static UPARAM(DisplayName = "Target Mesh") UDynamicMesh* 
+	ConvertIndexArrayToMeshSelection(
+		UDynamicMesh* TargetMesh,
+		const TArray<int32>& IndexArray,
+		EGeometryScriptMeshSelectionType SelectionType,
+		FGeometryScriptMeshSelection& Selection);
+
+	/**
+	 * Create a Mesh Selection from the IndexSet. 
+	 * @param SelectionType type of indices specified in the selection
+	 */
+	UFUNCTION(BlueprintCallable, Category = "GeometryScript|MeshSelection", meta=(ScriptMethod))
+	static UPARAM(DisplayName = "Target Mesh") UDynamicMesh* 
+	ConvertIndexSetToMeshSelection(
+		UDynamicMesh* TargetMesh,
+		const TSet<int32>& IndexSet,
+		EGeometryScriptMeshSelectionType SelectionType,
+		FGeometryScriptMeshSelection& Selection);
+
+	/**
+	 * Convert a Mesh Selection to an Index List
+	 * @param ConvertToType optional parameter specifying the type of Index List to convert to. If Set to Any, no conversion will be performed.
+	 */
+	UFUNCTION(BlueprintCallable, Category = "GeometryScript|MeshSelection", meta=(ScriptMethod))
+	static UPARAM(DisplayName = "Target Mesh") UDynamicMesh* 
+	ConvertMeshSelectionToIndexArray(
+		UDynamicMesh* TargetMesh,
+		FGeometryScriptMeshSelection Selection, 
+		TArray<int32>& IndexArray,
+		EGeometryScriptMeshSelectionType& SelectionType);
+
+	/**
+	 * Create a Mesh Selection from the Index List. For cases where the IndexList Type does not match
+	 * the SelectionType, ConvertMeshSelection with bAllowPartialInclusion=true is used to convert.
+	 * @param SelectionType type of indices desired in the Output selection
+	 */
+	UFUNCTION(BlueprintCallable, Category = "GeometryScript|MeshSelection", meta=(ScriptMethod))
+	static UPARAM(DisplayName = "Target Mesh") UDynamicMesh* 
+	ConvertIndexListToMeshSelection(
+		UDynamicMesh* TargetMesh,
+		FGeometryScriptIndexList IndexList,
+		EGeometryScriptMeshSelectionType SelectionType,
+		FGeometryScriptMeshSelection& Selection);
+
+	/**
+	 * Convert a Mesh Selection to an Index List
+	 * @param ConvertToType optional parameter specifying the type of Index List to convert to. If Set to Any, no conversion will be performed.
+	 */
+	UFUNCTION(BlueprintCallable, Category = "GeometryScript|MeshSelection", meta=(ScriptMethod))
+	static UPARAM(DisplayName = "Target Mesh") UDynamicMesh* 
+	ConvertMeshSelectionToIndexList(
+		UDynamicMesh* TargetMesh,
+		FGeometryScriptMeshSelection Selection, 
+		FGeometryScriptIndexList& IndexList,
+		EGeometryScriptIndexType& ResultListType,
+		EGeometryScriptIndexType ConvertToType = EGeometryScriptIndexType::Any);
+
+
+	/**
+	 * Create a new Mesh Selection of the SelectionType for the TargetMesh by finding all elements contained in the Box.
+	 * @param bInvert return a selection of all elements not in the Box
+	 * @param MinNumTrianglePoints number of vertices of a triangle that must be in the box for it to be selected (1,2, or 3)
+	 */
+	UFUNCTION(BlueprintCallable, Category = "GeometryScript|MeshSelection", meta=(ScriptMethod, AdvancedDisplay="NumPointsInBox"))
+	static UPARAM(DisplayName = "Target Mesh") UDynamicMesh* 
+	SelectMeshElementsInBox(
+		UDynamicMesh* TargetMesh,
+		FGeometryScriptMeshSelection& Selection, 
+		FBox Box,
+		EGeometryScriptMeshSelectionType SelectionType = EGeometryScriptMeshSelectionType::Triangles,
+		bool bInvert = false,
+		int MinNumTrianglePoints = 3);
+
+	/**
+	* Create a new Mesh Selection of the SelectionType for the TargetMesh by finding all elements contained in the Sphere.
+	* @param SphereOrigin center point of the Sphere
+	* @param SphereRadius radius of the Sphere
+	* @param bInvert return a selection of all elements not in the Sphere
+	* @param MinNumTrianglePoints number of vertices of a triangle that must be in the Sphere for it to be selected (1,2, or 3)
+	*/
+	UFUNCTION(BlueprintCallable, Category = "GeometryScript|MeshSelection", meta=(ScriptMethod, AdvancedDisplay="NumPointsInBox"))
+	static UPARAM(DisplayName = "Target Mesh") UDynamicMesh* 
+	SelectMeshElementsInSphere(
+		UDynamicMesh* TargetMesh,
+		FGeometryScriptMeshSelection& Selection, 
+		FVector SphereOrigin = FVector::ZeroVector,
+		double SphereRadius = 100.0,
+		EGeometryScriptMeshSelectionType SelectionType = EGeometryScriptMeshSelectionType::Triangles,
+		bool bInvert = false,
+		int MinNumTrianglePoints = 3);
+
+	/**
+	* Create a new Mesh Selection of the SelectionType for the TargetMesh by finding all elements on the "positive" side of a Plane
+	* @param PlaneOrigin center point of the Plane
+	* @param PlaneNormal normal vector for the Plane
+	* @param bInvert return a selection of all elements on the other (negative) side of the Plane
+	* @param MinNumTrianglePoints number of vertices of a triangle that must be on the positive Plane side to be selected (1,2, or 3)
+	*/
+	UFUNCTION(BlueprintCallable, Category = "GeometryScript|MeshSelection", meta=(ScriptMethod, AdvancedDisplay="NumPointsInBox"))
+	static UPARAM(DisplayName = "Target Mesh") UDynamicMesh* 
+	SelectMeshElementsWithPlane(
+		UDynamicMesh* TargetMesh,
+		FGeometryScriptMeshSelection& Selection, 
+		FVector PlaneOrigin = FVector::ZeroVector,
+		FVector PlaneNormal = FVector::UpVector,
+		EGeometryScriptMeshSelectionType SelectionType = EGeometryScriptMeshSelectionType::Triangles,
+		bool bInvert = false,
+		int MinNumTrianglePoints = 3);
+
+	/**
+	* Create a new Mesh Selection of the SelectionType for the TargetMesh by finding all elements that have a normal
+	* vector that is within an angular deviation threshold from the given Normal. 
+	* For Triangle and Polygroup selections the triangle facet normal is used, for Vertex selections the per-vertex averaged normal is used.
+	* @param Normal normal/direction vector to measure against
+	* @param MaxAngleDeg maximum angular deviation from Normal, in degrees
+	* @param bInvert return a selection of all elements not within the given deviation
+	* @param MinNumTrianglePoints number of vertices of a triangle that must be within the angular deviation for it to be selected (1,2, or 3)
+	*/
+	UFUNCTION(BlueprintCallable, Category = "GeometryScript|MeshSelection", meta=(ScriptMethod, AdvancedDisplay="NumPointsInBox"))
+	static UPARAM(DisplayName = "Target Mesh") UDynamicMesh* 
+	SelectMeshElementsByNormalAngle(
+		UDynamicMesh* TargetMesh,
+		FGeometryScriptMeshSelection& Selection, 
+		FVector Normal = FVector::UpVector,
+		double MaxAngleDeg = 1.0,
+		EGeometryScriptMeshSelectionType SelectionType = EGeometryScriptMeshSelectionType::Triangles,
+		bool bInvert = false,
+		int MinNumTrianglePoints = 3);
+
+
+	/**
+	* Create a new Mesh Selection of the SelectionType for the TargetMesh by finding all elements inside a second SelectionMesh
+	* For Triangle and Polygroup selections the triangle facet normal is used, for Vertex selections the per-vertex averaged normal is used.
+	* @param SelectionMeshTransform Transform applied to SelectionMesh for inside/outside testing
+	* @param bInvert return a selection of all elements not within the given deviation
+	* @param ShellDistance If > 0, points within this distance from SelectionMesh will also be considered "inside"
+	* @param WindingThreshold Threshold used for Fast Mesh Winding Number inside/outside test (range is [0,1], with 1 being "inside")
+	* @param MinNumTrianglePoints number of vertices of a triangle that must be within the angular deviation for it to be selected (1,2, or 3)
+	*/
+	UFUNCTION(BlueprintCallable, Category = "GeometryScript|MeshSelection", meta=(ScriptMethod, AdvancedDisplay="NumPointsInBox"))
+	static UPARAM(DisplayName = "Target Mesh") UDynamicMesh* 
+	SelectMeshElementsInsideMesh(
+		UDynamicMesh* TargetMesh,
+		UDynamicMesh* SelectionMesh,
+		FGeometryScriptMeshSelection& Selection, 
+		FTransform SelectionMeshTransform,
+		EGeometryScriptMeshSelectionType SelectionType = EGeometryScriptMeshSelectionType::Triangles,
+		bool bInvert = false,
+		double ShellDistance = 0.0,
+		double WindingThreshold = 0.5,
+		int MinNumTrianglePoints = 3);
+
+	/**
+	 * Expand the Selection on the TargetMesh to connected regions and return the NewSelection
+	 * @param ConnectionType defines what "connected" means, ie purely geometric connection, or some additional constraint like same MaterialIDs/etc 
+	 */
+	UFUNCTION(BlueprintCallable, Category = "GeometryScript|MeshSelection", meta=(ScriptMethod, AdvancedDisplay="NumPointsInBox"))
+	static UPARAM(DisplayName = "Target Mesh") UDynamicMesh* 
+	ExpandMeshSelectionToConnected(
+		UDynamicMesh* TargetMesh,
+		FGeometryScriptMeshSelection Selection, 
+		FGeometryScriptMeshSelection& NewSelection,
+		EGeometryScriptTopologyConnectionType ConnectionType = EGeometryScriptTopologyConnectionType::Geometric
+	);
+
+};
\ No newline at end of file
diff --git a/Engine/Plugins/Experimental/GeometryScripting/Source/GeometryScriptingCore/Public/GeometryScript/MeshSelectionQueryFunctions.h b/Engine/Plugins/Experimental/GeometryScripting/Source/GeometryScriptingCore/Public/GeometryScript/MeshSelectionQueryFunctions.h
new file mode 100644
--- /dev/null
+++ b/Engine/Plugins/Experimental/GeometryScripting/Source/GeometryScriptingCore/Public/GeometryScript/MeshSelectionQueryFunctions.h
@@ -0,0 +1,49 @@
+// Copyright Epic Games, Inc. All Rights Reserved.
+
+#pragma once
+
+#include "CoreMinimal.h"
+#include "Kismet/BlueprintFunctionLibrary.h"
+#include "GeometryScript/GeometryScriptTypes.h"
+#include "GeometryScript/GeometryScriptSelectionTypes.h"
+#include "MeshSelectionQueryFunctions.generated.h"
+
+UCLASS(meta = (ScriptName = "GeometryScript_MeshSelectionQueries"))
+class GEOMETRYSCRIPTINGCORE_API UGeometryScriptLibrary_MeshSelectionQueryFunctions : public UBlueprintFunctionLibrary
+{
+	GENERATED_BODY()
+public:
+
+	/**
+	 * Get the 3D Bounding Box of a Mesh Selection, ie bounding box of vertices contained in the Selection
+	 * @param bIsEmpty will return as true if the selection was empty (the box will be initialized to 0 in this case)
+	 */
+	UFUNCTION(BlueprintCallable, Category = "GeometryScript|SelectionQueries", meta=(ScriptMethod))
+	static UPARAM(DisplayName = "Target Mesh") UDynamicMesh* 
+	GetMeshSelectionBoundingBox(
+		UDynamicMesh* TargetMesh,
+		FGeometryScriptMeshSelection Selection,
+		FBox& SelectionBounds,
+		bool& bIsEmpty,
+		UGeometryScriptDebug* Debug = nullptr);
+
+	/**
+	 * Compute the set of Vertex Loops bordering a Mesh Selection. Both the 3D polylines and lists of vertex indices are returned for each Loop.
+	 * Note that for a Vertex selection this will function return the border loops around the set of vertex triangle one-rings.
+	 * 
+	 * @param IndexLoops for each discovered Loop, the IndexList of mesh vertex indices around the loop is returned here
+	 * @param PathLoops for each discovered Loop, the PolyPath of mesh vertex positions around the loop is returned here. The ordering for each loop is the same as IndexLoops.
+	 * @param NumLoops number of loops found is returned here
+	 * @param bFoundErrors true is returned here if topological errors were found during loop computation. In this case the Loop set may be incomplete.
+	 */
+	UFUNCTION(BlueprintCallable, Category = "GeometryScript|SelectionQueries", meta=(ScriptMethod))
+	static UPARAM(DisplayName = "Target Mesh") UDynamicMesh* 
+	GetMeshSelectionBoundaryLoops(
+		UDynamicMesh* TargetMesh,
+		FGeometryScriptMeshSelection Selection,
+		TArray<FGeometryScriptIndexList>& IndexLoops,
+		TArray<FGeometryScriptPolyPath>& PathLoops,
+		int& NumLoops,
+		bool& bFoundErrors,
+		UGeometryScriptDebug* Debug = nullptr);
+};
\ No newline at end of file
diff --git a/Engine/Plugins/Runtime/GeometryProcessing/Source/DynamicMesh/Public/Selections/GeometrySelection.h b/Engine/Plugins/Runtime/GeometryProcessing/Source/DynamicMesh/Public/Selections/GeometrySelection.h
--- a/Engine/Plugins/Runtime/GeometryProcessing/Source/DynamicMesh/Public/Selections/GeometrySelection.h
+++ b/Engine/Plugins/Runtime/GeometryProcessing/Source/DynamicMesh/Public/Selections/GeometrySelection.h
@@ -171,6 +171,15 @@ struct DYNAMICMESH_API FGeometrySelection
 		Selection.Reset();
 	}
 
+	/**
+	* Initialize the Element and Topology types for this Selection
+	*/
+	void InitializeTypes(EGeometryElementType ElementTypeIn, EGeometryTopologyType TopologyTypeIn)
+	{
+		ElementType = ElementTypeIn;
+		TopologyType = TopologyTypeIn;
+	}
+
 	/**
 	 * Initialize the Element and Topology types for this Selection based on another Selection
 	 */
