GeometryScript: add support for Mesh Selection to various functions

BasicEditFunctions:
DeleteSelectedTrianglesFromMesh - delete triangles in Mesh Selection

TransformFunctions:
RotateMesh - rotate the mesh (not selection related, added for parity)
TransformMeshSelection, TranslateMeshSelection, RotateMeshSelection, ScaleMeshSelection - variants of Transform functions that apply to vertices of Mesh Selection

PolygroupFunctions:
SetPolygroupForMeshSelection - set new polygroup ID (specified or auto-generated) on triangles of Mesh Selection

MaterialFunctions:
SetMaterialIDForMeshSelection - set new Material ID for triangles of Mesh Selection

DecompositionFunctions:
CopyMeshSelectionToMesh - copy mesh triangles identified by [REDACTED_PERSON] into a new DynamicMesh. Can optionally replace or append, and optionally preserve existing GroupIDs.

diff --git a/Engine/Plugins/Experimental/GeometryScripting/Source/GeometryScriptingCore/Private/MeshBasicEditFunctions.cpp b/Engine/Plugins/Experimental/GeometryScripting/Source/GeometryScriptingCore/Private/MeshBasicEditFunctions.cpp
--- a/Engine/Plugins/Experimental/GeometryScripting/Source/GeometryScriptingCore/Private/MeshBasicEditFunctions.cpp
+++ b/Engine/Plugins/Experimental/GeometryScripting/Source/GeometryScriptingCore/Private/MeshBasicEditFunctions.cpp
@@ -299,6 +299,34 @@ UDynamicMesh* UGeometryScriptLibrary_MeshBasicEditFunctions::DeleteTrianglesFrom
 }
 
 
+UDynamicMesh* UGeometryScriptLibrary_MeshBasicEditFunctions::DeleteSelectedTrianglesFromMesh(
+	UDynamicMesh* TargetMesh,
+	FGeometryScriptMeshSelection Selection,
+	int& NumDeleted,
+	bool bDeferChangeNotifications)
+{
+	if (TargetMesh && Selection.IsEmpty() == false )
+	{
+		NumDeleted = 0;
+		TargetMesh->EditMesh([&](FDynamicMesh3& EditMesh)
+		{
+			TArray<int32> Triangles;
+			Selection.ConvertToMeshIndexArray(EditMesh, Triangles, EGeometryScriptIndexType::Triangle);
+			for (int32 TriangleID : Triangles)
+			{
+				EMeshResult Result = EditMesh.RemoveTriangle(TriangleID);
+				if (Result == EMeshResult::Ok)
+				{
+					NumDeleted++;
+				}
+			}
+		}, EDynamicMeshChangeType::GeneralEdit, EDynamicMeshAttributeChangeFlags::Unknown, bDeferChangeNotifications);
+	}
+	return TargetMesh;
+}
+
+
+
 
 UDynamicMesh* UGeometryScriptLibrary_MeshBasicEditFunctions::AppendMesh(
 	UDynamicMesh* TargetMesh,
diff --git a/Engine/Plugins/Experimental/GeometryScripting/Source/GeometryScriptingCore/Private/MeshDecompositionFunctions.cpp b/Engine/Plugins/Experimental/GeometryScripting/Source/GeometryScriptingCore/Private/MeshDecompositionFunctions.cpp
--- a/Engine/Plugins/Experimental/GeometryScripting/Source/GeometryScriptingCore/Private/MeshDecompositionFunctions.cpp
+++ b/Engine/Plugins/Experimental/GeometryScripting/Source/GeometryScriptingCore/Private/MeshDecompositionFunctions.cpp
@@ -212,12 +212,12 @@ UDynamicMesh* UGeometryScriptLibrary_MeshDecompositionFunctions::GetSubMeshFromM
 	}
 	if (TriangleList.List.IsValid() == false || TriangleList.List->Num() == 0)
 	{
-		UE::Geometry::AppendError(Debug, EGeometryScriptErrorType::InvalidInputs, LOCTEXT("SetMeshPerVertexColors_InvalidList", "GetSubMeshFromMesh: TriangleList is empty"));
+		UE::Geometry::AppendError(Debug, EGeometryScriptErrorType::InvalidInputs, LOCTEXT("GetSubMeshFromMesh_InvalidList", "GetSubMeshFromMesh: TriangleList is empty"));
 		return TargetMesh;
 	}
 	if (TriangleList.IsCompatibleWith(EGeometryScriptIndexType::Triangle) == false)
 	{
-		UE::Geometry::AppendError(Debug, EGeometryScriptErrorType::InvalidInputs, LOCTEXT("SetMeshPerVertexColors_InvalidList2", "GetSubMeshFromMesh: TriangleList has incompatible index type"));
+		UE::Geometry::AppendError(Debug, EGeometryScriptErrorType::InvalidInputs, LOCTEXT("GetSubMeshFromMesh_InvalidList2", "GetSubMeshFromMesh: TriangleList has incompatible index type"));
 		return TargetMesh;
 	}
 
@@ -243,6 +243,79 @@ UDynamicMesh* UGeometryScriptLibrary_MeshDecompositionFunctions::GetSubMeshFromM
 }
 
 
+
+
+UDynamicMesh* UGeometryScriptLibrary_MeshDecompositionFunctions::CopyMeshSelectionToMesh(
+	UDynamicMesh* TargetMesh,
+	UDynamicMesh* StoreToSubmesh, 
+	FGeometryScriptMeshSelection Selection,
+	UDynamicMesh*& StoreToSubmeshOut, 
+	bool bAppendToExisting,
+	bool bPreserveGroupIDs,
+	UGeometryScriptDebug* Debug)
+{
+	if (TargetMesh == nullptr)
+	{
+		UE::Geometry::AppendError(Debug, EGeometryScriptErrorType::InvalidInputs, LOCTEXT("CopyMeshSelectionToMesh_InvalidInput", "CopyMeshSelectionToMesh: TargetMesh is Null"));
+		return TargetMesh;
+	}
+	if (StoreToSubmesh == nullptr)
+	{
+		UE::Geometry::AppendError(Debug, EGeometryScriptErrorType::InvalidInputs, LOCTEXT("CopyMeshSelectionToMesh_InvalidInput2", "CopyMeshSelectionToMesh: StoreToSubmesh is Null"));
+		return TargetMesh;
+	}
+	if (Selection.IsEmpty())
+	{
+		UE::Geometry::AppendError(Debug, EGeometryScriptErrorType::InvalidInputs, LOCTEXT("CopyMeshSelectionToMesh_InvalidList", "CopyMeshSelectionToMesh: Selection is empty"));
+		return TargetMesh;
+	}
+
+	FDynamicMesh3 Submesh;
+	if (bAppendToExisting)
+	{
+		StoreToSubmesh->ProcessMesh([&](const FDynamicMesh3& ReadMesh)
+		{
+			Submesh = ReadMesh;
+		});
+	}
+
+	TargetMesh->ProcessMesh([&](const FDynamicMesh3& SourceMesh)
+	{
+		TArray<int32> Triangles;
+		Selection.ConvertToMeshIndexArray(SourceMesh, Triangles, EGeometryScriptIndexType::Triangle);
+
+		if (bAppendToExisting == false)
+		{
+			Submesh.Clear();
+			Submesh.EnableMatchingAttributes(SourceMesh, false);
+		}
+
+		FMeshIndexMappings Mappings;
+		FDynamicMeshEditResult EditResult;
+		FDynamicMeshEditor Editor(&Submesh);
+		Editor.AppendTriangles(&SourceMesh, Triangles, Mappings, EditResult, bPreserveGroupIDs);
+
+		if (bPreserveGroupIDs)
+		{
+			for (int32 tid : EditResult.NewTriangles)
+			{
+				int32 GroupID = Submesh.GetTriangleGroup(tid);
+				int32 OldGroupID = Mappings.GetGroupMap().GetFrom(GroupID);
+				Submesh.SetTriangleGroup(tid, OldGroupID);
+			}
+		}
+
+	});
+
+	StoreToSubmesh->SetMesh(MoveTemp(Submesh));
+	StoreToSubmeshOut = StoreToSubmesh;
+
+	return TargetMesh;
+}
+
+
+
+
 UDynamicMesh* UGeometryScriptLibrary_MeshDecompositionFunctions::CopyMeshToMesh(
 	UDynamicMesh* CopyFromMesh,
 	UDynamicMesh* CopyToMesh, 
@@ -251,12 +324,12 @@ UDynamicMesh* UGeometryScriptLibrary_MeshDecompositionFunctions::CopyMeshToMesh(
 {
 	if (CopyFromMesh == nullptr)
 	{
-		UE::Geometry::AppendError(Debug, EGeometryScriptErrorType::InvalidInputs, LOCTEXT("CopyMeshToMesh_InvalidInput", "CopyMeshToMesh: TargetMesh is Null"));
+		UE::Geometry::AppendError(Debug, EGeometryScriptErrorType::InvalidInputs, LOCTEXT("CopyMeshToMesh_InvalidInput", "CopyMeshToMesh: CopyFromMesh is Null"));
 		return CopyFromMesh;
 	}
 	if (CopyToMesh == nullptr)
 	{
-		UE::Geometry::AppendError(Debug, EGeometryScriptErrorType::InvalidInputs, LOCTEXT("CopyMeshToMesh_InvalidInput2", "CopyMeshToMesh: Submesh is Null"));
+		UE::Geometry::AppendError(Debug, EGeometryScriptErrorType::InvalidInputs, LOCTEXT("CopyMeshToMesh_InvalidInput2", "CopyMeshToMesh: CopyToMesh is Null"));
 		return CopyFromMesh;
 	}
 
diff --git a/Engine/Plugins/Experimental/GeometryScripting/Source/GeometryScriptingCore/Private/MeshMaterialFunctions.cpp b/Engine/Plugins/Experimental/GeometryScripting/Source/GeometryScriptingCore/Private/MeshMaterialFunctions.cpp
--- a/Engine/Plugins/Experimental/GeometryScripting/Source/GeometryScriptingCore/Private/MeshMaterialFunctions.cpp
+++ b/Engine/Plugins/Experimental/GeometryScripting/Source/GeometryScriptingCore/Private/MeshMaterialFunctions.cpp
@@ -461,6 +461,40 @@ UDynamicMesh* UGeometryScriptLibrary_MeshMaterialFunctions::SetMaterialIDOnTrian
 
 
 
+UDynamicMesh* UGeometryScriptLibrary_MeshMaterialFunctions::SetMaterialIDForMeshSelection( 
+	UDynamicMesh* TargetMesh, 
+	FGeometryScriptMeshSelection Selection,
+	int MaterialID,
+	bool bDeferChangeNotifications,
+	UGeometryScriptDebug* Debug)
+{	
+	if (TargetMesh == nullptr)
+	{
+		AppendError(Debug, EGeometryScriptErrorType::InvalidInputs, LOCTEXT("SetMaterialIDForMeshSelection_InvalidMesh", "SetMaterialIDForMeshSelection: TargetMesh is Null"));
+		return TargetMesh;
+	}
+
+	TargetMesh->EditMesh([&](FDynamicMesh3& EditMesh) 
+	{
+		if (EditMesh.HasAttributes() == false)
+		{
+			EditMesh.EnableAttributes();
+		}
+		if (EditMesh.Attributes()->HasMaterialID() == false)
+		{
+			EditMesh.Attributes()->EnableMaterialID();
+		}
+		FDynamicMeshMaterialAttribute* MaterialIDs = EditMesh.Attributes()->GetMaterialID();
+		Selection.ProcessByTriangleID(EditMesh,
+			[&](int32 TriangleID) { MaterialIDs->SetValue(TriangleID, MaterialID); } );
+	}, EDynamicMeshChangeType::GeneralEdit, EDynamicMeshAttributeChangeFlags::Unknown, bDeferChangeNotifications);
+
+	return TargetMesh;
+}
+
+
+
+
 UDynamicMesh* UGeometryScriptLibrary_MeshMaterialFunctions::SetPolygroupMaterialID( 
 	UDynamicMesh* TargetMesh, 
 	FGeometryScriptGroupLayer GroupLayer,
diff --git a/Engine/Plugins/Experimental/GeometryScripting/Source/GeometryScriptingCore/Private/MeshPolygroupFunctions.cpp b/Engine/Plugins/Experimental/GeometryScripting/Source/GeometryScriptingCore/Private/MeshPolygroupFunctions.cpp
--- a/Engine/Plugins/Experimental/GeometryScripting/Source/GeometryScriptingCore/Private/MeshPolygroupFunctions.cpp
+++ b/Engine/Plugins/Experimental/GeometryScripting/Source/GeometryScriptingCore/Private/MeshPolygroupFunctions.cpp
@@ -428,5 +428,42 @@ UDynamicMesh* UGeometryScriptLibrary_MeshPolygroupFunctions::GetTrianglesInPolyg
 
 
 
+UDynamicMesh* UGeometryScriptLibrary_MeshPolygroupFunctions::SetPolygroupForMeshSelection(
+	UDynamicMesh* TargetMesh,
+	FGeometryScriptGroupLayer GroupLayer,
+	FGeometryScriptMeshSelection Selection,
+	int& SetPolygroupIDOut,
+	int SetPolygroupID,
+	bool bGenerateNewPolygroup,
+	bool bDeferChangeNotifications)
+{
+	SetPolygroupIDOut = SetPolygroupID;
+	if (TargetMesh)
+	{
+		TargetMesh->EditMesh([&](FDynamicMesh3& EditMesh)
+		{
+			FPolygroupLayer InputGroupLayer{ GroupLayer.bDefaultLayer, GroupLayer.ExtendedLayerIndex };
+			if (InputGroupLayer.CheckExists(&EditMesh) == false)
+			{
+				UE_LOG(LogGeometry, Warning, TEXT("SetPolygroupForMeshSelection: Specified Polygroup Layer does not exist"));
+				return;
+			}
+
+			FPolygroupSet Groups(&EditMesh, InputGroupLayer);
+			if (bGenerateNewPolygroup)
+			{
+				SetPolygroupID = Groups.AllocateNewGroupID();
+				SetPolygroupIDOut = SetPolygroupID;
+			}
+			Selection.ProcessByTriangleID(EditMesh,
+				[&](int32 TriangleID) { Groups.SetGroup(TriangleID, SetPolygroupID, EditMesh); } );
+
+		}, EDynamicMeshChangeType::GeneralEdit, EDynamicMeshAttributeChangeFlags::Unknown, bDeferChangeNotifications);
+	}
+	return TargetMesh;
+}
+
+
+
 
 #undef LOCTEXT_NAMESPACE
\ No newline at end of file
diff --git a/Engine/Plugins/Experimental/GeometryScripting/Source/GeometryScriptingCore/Private/MeshTransformFunctions.cpp b/Engine/Plugins/Experimental/GeometryScripting/Source/GeometryScriptingCore/Private/MeshTransformFunctions.cpp
--- a/Engine/Plugins/Experimental/GeometryScripting/Source/GeometryScriptingCore/Private/MeshTransformFunctions.cpp
+++ b/Engine/Plugins/Experimental/GeometryScripting/Source/GeometryScriptingCore/Private/MeshTransformFunctions.cpp
@@ -4,6 +4,7 @@
 
 #include "DynamicMesh/DynamicMesh3.h"
 #include "DynamicMesh/MeshTransforms.h"
+#include "Async/ParallelFor.h"
 #include "UDynamicMesh.h"
 
 using namespace UE::Geometry;
@@ -57,9 +58,40 @@ UDynamicMesh* UGeometryScriptLibrary_MeshTransformFunctions::TranslateMesh(
 }
 
 
+UDynamicMesh* UGeometryScriptLibrary_MeshTransformFunctions::RotateMesh(
+	UDynamicMesh* TargetMesh,
+	FRotator Rotation,
+	FVector RotationOrigin,
+	UGeometryScriptDebug* Debug)
+{
+	if (TargetMesh == nullptr)
+	{
+		UE::Geometry::AppendError(Debug, EGeometryScriptErrorType::InvalidInputs, LOCTEXT("RotateMesh_InvalidInput", "RotateMesh: TargetMesh is Null"));
+		return TargetMesh;
+	}
+
+	// todo: publish correct change types
+	TargetMesh->EditMesh([&](FDynamicMesh3& EditMesh) 
+	{
+		int NumVertices = EditMesh.MaxVertexID();
+		ParallelFor(NumVertices, [&](int32 vid)
+		{
+			if (EditMesh.IsVertex(vid))
+			{
+				EditMesh.SetVertex(vid, Rotation.RotateVector( (EditMesh.GetVertex(vid) - RotationOrigin) ) + RotationOrigin);
+			}
+		});
+
+	}, EDynamicMeshChangeType::GeneralEdit, EDynamicMeshAttributeChangeFlags::Unknown, false);
+
+	return TargetMesh;
+}
+
+
 UDynamicMesh* UGeometryScriptLibrary_MeshTransformFunctions::ScaleMesh(
 	UDynamicMesh* TargetMesh,
 	FVector Scale,
+	FVector ScaleOrigin,
 	bool bFixOrientationForNegativeScale,
 	UGeometryScriptDebug* Debug)
 {
@@ -72,7 +104,7 @@ UDynamicMesh* UGeometryScriptLibrary_MeshTransformFunctions::ScaleMesh(
 	// todo: publish correct change types
 	TargetMesh->EditMesh([&](FDynamicMesh3& EditMesh) 
 	{
-		MeshTransforms::Scale(EditMesh, (FVector3d)Scale, FVector3d::Zero(), bFixOrientationForNegativeScale);
+		MeshTransforms::Scale(EditMesh, (FVector3d)Scale, (FVector3d)ScaleOrigin, bFixOrientationForNegativeScale);
 
 	}, EDynamicMeshChangeType::GeneralEdit, EDynamicMeshAttributeChangeFlags::Unknown, false);
 
@@ -82,4 +114,148 @@ UDynamicMesh* UGeometryScriptLibrary_MeshTransformFunctions::ScaleMesh(
 
 
 
+
+
+
+UDynamicMesh* UGeometryScriptLibrary_MeshTransformFunctions::TransformMeshSelection(
+	UDynamicMesh* TargetMesh,
+	FGeometryScriptMeshSelection Selection,
+	FTransform Transform,
+	UGeometryScriptDebug* Debug)
+{
+	if (TargetMesh == nullptr)
+	{
+		UE::Geometry::AppendError(Debug, EGeometryScriptErrorType::InvalidInputs, LOCTEXT("TransformMeshSelection_InvalidInput", "TransformMeshSelection: TargetMesh is Null"));
+		return TargetMesh;
+	}
+
+	// todo: publish correct change types
+	TargetMesh->EditMesh([&](FDynamicMesh3& EditMesh) 
+	{
+		TArray<int32> Vertices;
+		Selection.ConvertToMeshIndexArray(EditMesh, Vertices, EGeometryScriptIndexType::Vertex);
+		for ( int32 vid : Vertices)
+		{
+			if (EditMesh.IsVertex(vid))
+			{
+				EditMesh.SetVertex(vid, Transform.TransformPosition( EditMesh.GetVertex(vid) ) );
+			}
+		}
+
+		// normals?
+
+	}, EDynamicMeshChangeType::GeneralEdit, EDynamicMeshAttributeChangeFlags::Unknown, false);
+
+	return TargetMesh;
+}
+
+
+
+UDynamicMesh* UGeometryScriptLibrary_MeshTransformFunctions::TranslateMeshSelection(
+	UDynamicMesh* TargetMesh,
+	FGeometryScriptMeshSelection Selection,
+	FVector Translation,
+	UGeometryScriptDebug* Debug)
+{
+	if (TargetMesh == nullptr)
+	{
+		UE::Geometry::AppendError(Debug, EGeometryScriptErrorType::InvalidInputs, LOCTEXT("TranslateMeshSelection_InvalidInput", "TranslateMeshSelection: TargetMesh is Null"));
+		return TargetMesh;
+	}
+
+	// todo: publish correct change types
+	TargetMesh->EditMesh([&](FDynamicMesh3& EditMesh) 
+	{
+		if ( ! Selection.IsEmpty() )
+		{
+			TArray<int32> Vertices;
+			Selection.ConvertToMeshIndexArray(EditMesh, Vertices, EGeometryScriptIndexType::Vertex);
+			for ( int32 vid : Vertices)
+			{
+				if (EditMesh.IsVertex(vid))
+				{
+					EditMesh.SetVertex(vid, EditMesh.GetVertex(vid) + Translation);
+				}
+			}
+		}
+
+	}, EDynamicMeshChangeType::GeneralEdit, EDynamicMeshAttributeChangeFlags::Unknown, false);
+
+	return TargetMesh;
+}
+
+
+UDynamicMesh* UGeometryScriptLibrary_MeshTransformFunctions::RotateMeshSelection(
+	UDynamicMesh* TargetMesh,
+	FGeometryScriptMeshSelection Selection,
+	FRotator Rotation,
+	FVector RotationOrigin,
+	UGeometryScriptDebug* Debug)
+{
+	if (TargetMesh == nullptr)
+	{
+		UE::Geometry::AppendError(Debug, EGeometryScriptErrorType::InvalidInputs, LOCTEXT("RotateMeshSelection_InvalidInput", "RotateMeshSelection: TargetMesh is Null"));
+		return TargetMesh;
+	}
+
+	// todo: publish correct change types
+	TargetMesh->EditMesh([&](FDynamicMesh3& EditMesh) 
+	{
+		if ( ! Selection.IsEmpty() )
+		{
+			TArray<int32> Vertices;
+			Selection.ConvertToMeshIndexArray(EditMesh, Vertices, EGeometryScriptIndexType::Vertex);
+			for ( int32 vid : Vertices)
+			{
+				if (EditMesh.IsVertex(vid))
+				{
+					EditMesh.SetVertex(vid, Rotation.RotateVector( (EditMesh.GetVertex(vid) - RotationOrigin) ) + RotationOrigin);
+				}
+			}
+		}
+
+	}, EDynamicMeshChangeType::GeneralEdit, EDynamicMeshAttributeChangeFlags::Unknown, false);
+
+	return TargetMesh;
+}
+
+
+UDynamicMesh* UGeometryScriptLibrary_MeshTransformFunctions::ScaleMeshSelection(
+	UDynamicMesh* TargetMesh,
+	FGeometryScriptMeshSelection Selection,
+	FVector Scale,
+	FVector ScaleOrigin,
+	UGeometryScriptDebug* Debug)
+{
+	if (TargetMesh == nullptr)
+	{
+		UE::Geometry::AppendError(Debug, EGeometryScriptErrorType::InvalidInputs, LOCTEXT("ScaleMeshSelection_InvalidInput", "ScaleMeshSelection: TargetMesh is Null"));
+		return TargetMesh;
+	}
+
+	// todo: publish correct change types
+	TargetMesh->EditMesh([&](FDynamicMesh3& EditMesh) 
+	{
+		if ( ! Selection.IsEmpty() )
+		{
+			TArray<int32> Vertices;
+			Selection.ConvertToMeshIndexArray(EditMesh, Vertices, EGeometryScriptIndexType::Vertex);
+			for ( int32 vid : Vertices)
+			{
+				if (EditMesh.IsVertex(vid))
+				{
+					EditMesh.SetVertex(vid, (EditMesh.GetVertex(vid) - ScaleOrigin) * Scale + ScaleOrigin);
+				}
+			}
+		}
+
+	}, EDynamicMeshChangeType::GeneralEdit, EDynamicMeshAttributeChangeFlags::Unknown, false);
+
+	return TargetMesh;
+}
+
+
+
+
+
 #undef LOCTEXT_NAMESPACE
diff --git a/Engine/Plugins/Experimental/GeometryScripting/Source/GeometryScriptingCore/Public/GeometryScript/MeshBasicEditFunctions.h b/Engine/Plugins/Experimental/GeometryScripting/Source/GeometryScriptingCore/Public/GeometryScript/MeshBasicEditFunctions.h
--- a/Engine/Plugins/Experimental/GeometryScripting/Source/GeometryScriptingCore/Public/GeometryScript/MeshBasicEditFunctions.h
+++ b/Engine/Plugins/Experimental/GeometryScripting/Source/GeometryScriptingCore/Public/GeometryScript/MeshBasicEditFunctions.h
@@ -5,6 +5,7 @@
 #include "CoreMinimal.h"
 #include "Kismet/BlueprintFunctionLibrary.h"
 #include "GeometryScript/GeometryScriptTypes.h"
+#include "GeometryScript/GeometryScriptSelectionTypes.h"
 #include "MeshBasicEditFunctions.generated.h"
 
 class UDynamicMesh;
@@ -144,6 +145,13 @@ public:
 		int& NumDeleted,
 		bool bDeferChangeNotifications = false );
 
+	UFUNCTION(BlueprintCallable, Category = "GeometryScript|MeshEdits", meta=(ScriptMethod))
+	static UPARAM(DisplayName = "Target Mesh") UDynamicMesh* 
+	DeleteSelectedTrianglesFromMesh( 
+		UDynamicMesh* TargetMesh, 
+		FGeometryScriptMeshSelection Selection,
+		int& NumDeleted,
+		bool bDeferChangeNotifications = false );
 
 	UFUNCTION(BlueprintCallable, Category = "GeometryScript|MeshEdits", meta=(ScriptMethod))
 	static UPARAM(DisplayName = "Target Mesh") UDynamicMesh* 
diff --git a/Engine/Plugins/Experimental/GeometryScripting/Source/GeometryScriptingCore/Public/GeometryScript/MeshDecompositionFunctions.h b/Engine/Plugins/Experimental/GeometryScripting/Source/GeometryScriptingCore/Public/GeometryScript/MeshDecompositionFunctions.h
--- a/Engine/Plugins/Experimental/GeometryScripting/Source/GeometryScriptingCore/Public/GeometryScript/MeshDecompositionFunctions.h
+++ b/Engine/Plugins/Experimental/GeometryScripting/Source/GeometryScriptingCore/Public/GeometryScript/MeshDecompositionFunctions.h
@@ -5,6 +5,7 @@
 #include "CoreMinimal.h"
 #include "Kismet/BlueprintFunctionLibrary.h"
 #include "GeometryScript/GeometryScriptTypes.h"
+#include "GeometryScript/GeometryScriptSelectionTypes.h"
 #include "MeshDecompositionFunctions.generated.h"
 
 class UDynamicMesh;
@@ -16,6 +17,12 @@ class GEOMETRYSCRIPTINGCORE_API UGeometryScriptLibrary_MeshDecompositionFunction
 	GENERATED_BODY()
 public:
 
+	/**
+	 * Create a new Mesh for each Connected Component of TargetMesh.
+	 * New meshes are drawn from MeshPool if it is provided, otherwise new UDynamicMesh instances are allocated
+	 * @param ComponentMeshes New List of meshes is returned here
+	 * @param MeshPool New meshes in ComponentMeshes output list are allocated from this pool if it is provided (highly recommended!!)
+	 */
 	UFUNCTION(BlueprintCallable, Category = "GeometryScript|Decomposition", meta=(ScriptMethod))
 	static UPARAM(DisplayName = "Target Mesh") UDynamicMesh* 
 	SplitMeshByComponents(  
@@ -24,6 +31,13 @@ public:
 		UDynamicMeshPool* MeshPool,
 		UGeometryScriptDebug* Debug = nullptr);
 
+	/**
+	 * Create a new Mesh for each MaterialID of TargetMesh.
+	 * New meshes are drawn from MeshPool if it is provided, otherwise new UDynamicMesh instances are allocated
+	 * @param ComponentMeshes New List of meshes is returned here
+	 * @param ComponentMaterialIDs MaterialID for each Mesh in ComponentMeshes is returned here
+	 * @param MeshPool New meshes in ComponentMeshes output list are allocated from this pool if it is provided (highly recommended!!)
+	 */
 	UFUNCTION(BlueprintCallable, Category = "GeometryScript|Decomposition", meta=(ScriptMethod))
 	static UPARAM(DisplayName = "Target Mesh") UDynamicMesh* 
 	SplitMeshByMaterialIDs(  
@@ -33,6 +47,13 @@ public:
 		UDynamicMeshPool* MeshPool,
 		UGeometryScriptDebug* Debug = nullptr);
 
+	/**
+	 * Create a new Mesh for each Polygroup of TargetMesh. Note that this may be a *large* number of meshes!
+	 * New meshes are drawn from MeshPool if it is provided, otherwise new UDynamicMesh instances are allocated
+	 * @param ComponentMeshes New List of meshes is returned here
+	 * @param ComponentPolygroups Original Polygroup for each Mesh in ComponentMeshes is returned here
+	 * @param MeshPool New meshes in ComponentMeshes output list are allocated from this pool if it is provided (highly recommended!!)
+	 */
 	UFUNCTION(BlueprintCallable, Category = "GeometryScript|Decomposition", meta=(ScriptMethod))
 	static UPARAM(DisplayName = "Target Mesh") UDynamicMesh* 
 	SplitMeshByPolygroups(  
@@ -43,6 +64,9 @@ public:
 		UDynamicMeshPool* MeshPool,
 		UGeometryScriptDebug* Debug = nullptr);
 
+	/**
+	 * CopyMeshSelectionToMesh should be used instead of this function
+	 */
 	UFUNCTION(BlueprintCallable, Category = "GeometryScript|Decomposition", meta=(ScriptMethod))
 	static UPARAM(DisplayName = "Target Mesh") UDynamicMesh* 
 	GetSubMeshFromMesh(  
@@ -52,6 +76,25 @@ public:
 		UPARAM(DisplayName = "Copy To Submesh") UDynamicMesh*& StoreToSubmeshOut, 
 		UGeometryScriptDebug* Debug = nullptr);
 
+	/**
+	 * Extract the triangles identified by Selection from TargetMesh and copy/add them to StoreToSubmesh
+	 * @param bAppendToExisting if false (default), StoreToSubmesh is cleared, otherwise selected triangles are appended
+	 * @param bPreserveGroupIDs if true, GroupIDs of triangles on TargetMesh are preserved in StoreToSubmesh. Otherwise new GroupIDs are allocated.
+	 */
+	UFUNCTION(BlueprintCallable, Category = "GeometryScript|Decomposition", meta=(ScriptMethod))
+	static UPARAM(DisplayName = "Target Mesh") UDynamicMesh* 
+	CopyMeshSelectionToMesh(  
+		UDynamicMesh* TargetMesh, 
+		UPARAM(DisplayName = "Copy To Submesh", ref) UDynamicMesh* StoreToSubmesh, 
+		FGeometryScriptMeshSelection Selection,
+		UPARAM(DisplayName = "Copy To Submesh") UDynamicMesh*& StoreToSubmeshOut, 
+		bool bAppendToExisting = false,
+		bool bPreserveGroupIDs = false,
+		UGeometryScriptDebug* Debug = nullptr);
+
+	/**
+	 * Set CopyToMesh to be the same mesh as CopyFromMesh
+	 */
 	UFUNCTION(BlueprintCallable, Category = "GeometryScript|Decomposition", meta=(ScriptMethod))
 	static UPARAM(DisplayName = "Copy From Mesh") UDynamicMesh* 
 	CopyMeshToMesh(  
diff --git a/Engine/Plugins/Experimental/GeometryScripting/Source/GeometryScriptingCore/Public/GeometryScript/MeshMaterialFunctions.h b/Engine/Plugins/Experimental/GeometryScripting/Source/GeometryScriptingCore/Public/GeometryScript/MeshMaterialFunctions.h
--- a/Engine/Plugins/Experimental/GeometryScripting/Source/GeometryScriptingCore/Public/GeometryScript/MeshMaterialFunctions.h
+++ b/Engine/Plugins/Experimental/GeometryScripting/Source/GeometryScriptingCore/Public/GeometryScript/MeshMaterialFunctions.h
@@ -5,6 +5,7 @@
 #include "CoreMinimal.h"
 #include "Kismet/BlueprintFunctionLibrary.h"
 #include "GeometryScript/GeometryScriptTypes.h"
+#include "GeometryScript/GeometryScriptSelectionTypes.h"
 #include "MeshMaterialFunctions.generated.h"
 
 class UDynamicMesh;
@@ -85,6 +86,27 @@ public:
         bool bDeferChangeNotifications = false,
 		UGeometryScriptDebug* Debug = nullptr);
 
+	/**
+	* Set a new MaterialID on all the triangles of the given Selection.
+	* @param MaterialID new Material ID to set
+	* @param bDeferChangeNotifications if true, the UDynamicMesh does not emit a change event/signal for this modification
+	*/
+	UFUNCTION(BlueprintCallable, Category = "GeometryScript|Polygroups", meta=(ScriptMethod))
+	static UPARAM(DisplayName = "Target Mesh") UDynamicMesh* 
+	SetMaterialIDForMeshSelection( 
+		UDynamicMesh* TargetMesh, 
+		FGeometryScriptMeshSelection Selection,
+		int MaterialID,
+		bool bDeferChangeNotifications = false,
+		UGeometryScriptDebug* Debug = nullptr);
+
+	/**
+	* Set a new MaterialID on all the triangles of TargetMesh with the given Polygroup.
+	* @param GroupLayer Polygroup Layer to use as basis for polygroups
+	* @param PolygroupID Polygroup ID that specifies Triangles to set to new MaterialID
+	* @param MaterialID explicit new MaterialID to set
+	* @param bDeferChangeNotifications if true, the UDynamicMesh does not emit a change event/signal for this modification
+	*/
 	UFUNCTION(BlueprintCallable, Category = "GeometryScript|Materials", meta=(ScriptMethod))
 	static UPARAM(DisplayName = "Target Mesh") UDynamicMesh* 
 	SetPolygroupMaterialID( 
@@ -96,6 +118,10 @@ public:
 		bool bDeferChangeNotifications = false,
 		UGeometryScriptDebug* Debug = nullptr);
 
+	/**
+	 * Delete all triangles in TargetMesh with the given MaterialID
+	 * @param NumDeleted number of deleted triangles is returned here
+	 */
 	UFUNCTION(BlueprintCallable, Category = "GeometryScript|Materials", meta=(ScriptMethod))
 	static UPARAM(DisplayName = "Target Mesh") UDynamicMesh* 
 	DeleteTrianglesByMaterialID( 
diff --git a/Engine/Plugins/Experimental/GeometryScripting/Source/GeometryScriptingCore/Public/GeometryScript/MeshPolygroupFunctions.h b/Engine/Plugins/Experimental/GeometryScripting/Source/GeometryScriptingCore/Public/GeometryScript/MeshPolygroupFunctions.h
--- a/Engine/Plugins/Experimental/GeometryScripting/Source/GeometryScriptingCore/Public/GeometryScript/MeshPolygroupFunctions.h
+++ b/Engine/Plugins/Experimental/GeometryScripting/Source/GeometryScriptingCore/Public/GeometryScript/MeshPolygroupFunctions.h
@@ -5,6 +5,7 @@
 #include "CoreMinimal.h"
 #include "Kismet/BlueprintFunctionLibrary.h"
 #include "GeometryScript/GeometryScriptTypes.h"
+#include "GeometryScript/GeometryScriptSelectionTypes.h"
 #include "MeshPolygroupFunctions.generated.h"
 
 class UDynamicMesh;
@@ -119,4 +120,22 @@ public:
 		int PolygroupID, 
 		UPARAM(ref, DisplayName="Triangle IDs Out") FGeometryScriptIndexList& TriangleIDsOut );
 
+
+	/**
+	 * Set a new Polygroup on all the triangles of the given Selection, for the given GroupLayer.
+	 * @param SetPolygroupID explicit new PolygroupID to set
+	 * @param bGenerateNewPolygroup if true, SetPolygroupID is ignored and a new unique PolygroupID is generated
+	 * @param SetPolygroupIDOut the PolygroupID that was set on the triangles is returned here (whether explicit or auto-generated)
+	 * @param bDeferChangeNotifications if true, the UDynamicMesh does not emit a change event/signal for this modification
+	 */
+	UFUNCTION(BlueprintCallable, Category = "GeometryScript|Polygroups", meta=(ScriptMethod))
+	static UPARAM(DisplayName = "Target Mesh") UDynamicMesh* 
+	SetPolygroupForMeshSelection( 
+		UDynamicMesh* TargetMesh, 
+		FGeometryScriptGroupLayer GroupLayer, 
+		FGeometryScriptMeshSelection Selection,
+		int& SetPolygroupIDOut,
+		int SetPolygroupID = 0,
+		bool bGenerateNewPolygroup = false,
+		bool bDeferChangeNotifications = false);
 };
\ No newline at end of file
diff --git a/Engine/Plugins/Experimental/GeometryScripting/Source/GeometryScriptingCore/Public/GeometryScript/MeshTransformFunctions.h b/Engine/Plugins/Experimental/GeometryScripting/Source/GeometryScriptingCore/Public/GeometryScript/MeshTransformFunctions.h
--- a/Engine/Plugins/Experimental/GeometryScripting/Source/GeometryScriptingCore/Public/GeometryScript/MeshTransformFunctions.h
+++ b/Engine/Plugins/Experimental/GeometryScripting/Source/GeometryScriptingCore/Public/GeometryScript/MeshTransformFunctions.h
@@ -5,6 +5,7 @@
 #include "CoreMinimal.h"
 #include "Kismet/BlueprintFunctionLibrary.h"
 #include "GeometryScript/GeometryScriptTypes.h"
+#include "GeometryScript/GeometryScriptSelectionTypes.h"
 #include "MeshTransformFunctions.generated.h"
 
 class UDynamicMesh;
@@ -31,12 +32,59 @@ public:
 		FVector Translation,
 		UGeometryScriptDebug* Debug = nullptr );
 
+	UFUNCTION(BlueprintCallable, Category = "GeometryScript|Transforms", meta=(ScriptMethod))
+	static UPARAM(DisplayName = "Target Mesh") UDynamicMesh* 
+	RotateMesh(
+		UDynamicMesh* TargetMesh,
+		FRotator Rotation,
+		FVector RotationOrigin = FVector(0,0,0),
+		UGeometryScriptDebug* Debug = nullptr );
+
 	UFUNCTION(BlueprintCallable, Category = "GeometryScript|Transforms", meta=(ScriptMethod, AdvancedDisplay = "bFixOrientationForNegativeScale"))
 	static UPARAM(DisplayName = "Target Mesh") UDynamicMesh* 
 	ScaleMesh(
 		UDynamicMesh* TargetMesh,
 		FVector Scale = FVector(1,1,1),
+		FVector ScaleOrigin = FVector(0,0,0),
 		bool bFixOrientationForNegativeScale = true,
 		UGeometryScriptDebug* Debug = nullptr);
 
+
+
+
+	UFUNCTION(BlueprintCallable, Category = "GeometryScript|Transforms", meta=(ScriptMethod, AdvancedDisplay = "bFixOrientationForNegativeScale"))
+	static UPARAM(DisplayName = "Target Mesh") UDynamicMesh*
+	TransformMeshSelection(
+		UDynamicMesh* TargetMesh,
+		FGeometryScriptMeshSelection Selection,
+		FTransform Transform,
+		UGeometryScriptDebug* Debug = nullptr);
+
+	UFUNCTION(BlueprintCallable, Category = "GeometryScript|Transforms", meta=(ScriptMethod))
+	static UPARAM(DisplayName = "Target Mesh") UDynamicMesh* 
+	TranslateMeshSelection(
+		UDynamicMesh* TargetMesh,
+		FGeometryScriptMeshSelection Selection,
+		FVector Translation,
+		UGeometryScriptDebug* Debug = nullptr );
+
+	UFUNCTION(BlueprintCallable, Category = "GeometryScript|Transforms", meta=(ScriptMethod))
+	static UPARAM(DisplayName = "Target Mesh") UDynamicMesh* 
+	RotateMeshSelection(
+		UDynamicMesh* TargetMesh,
+		FGeometryScriptMeshSelection Selection,
+		FRotator Rotation,
+		FVector RotationOrigin = FVector(0,0,0),
+		UGeometryScriptDebug* Debug = nullptr );
+
+	UFUNCTION(BlueprintCallable, Category = "GeometryScript|Transforms", meta=(ScriptMethod, AdvancedDisplay = "bFixOrientationForNegativeScale"))
+	static UPARAM(DisplayName = "Target Mesh") UDynamicMesh* 
+	ScaleMeshSelection(
+		UDynamicMesh* TargetMesh,
+		FGeometryScriptMeshSelection Selection,
+		FVector Scale = FVector(1,1,1),
+		FVector ScaleOrigin = FVector(0,0,0),
+		UGeometryScriptDebug* Debug = nullptr);
+
+
 };
\ No newline at end of file
