GeometryScript: add functions CopyMeshUVLayerToMesh and CopyMeshToMeshUVLayer. These allow mapping a UV layer to a separate mesh, editing it (moving vertices, modifying triangles, etc) and then transferring the edited positions/topology back to the UV layer of the source mesh. This provides a mechanism to do UV manipulation in GeometryScript without having to build out separate UV variants of related functions (transforms, warps, etc)

diff --git a/Engine/Plugins/Experimental/GeometryScripting/Source/GeometryScriptingCore/Private/MeshUVFunctions.cpp b/Engine/Plugins/Experimental/GeometryScripting/Source/GeometryScriptingCore/Private/MeshUVFunctions.cpp
--- a/Engine/Plugins/Experimental/GeometryScripting/Source/GeometryScriptingCore/Private/MeshUVFunctions.cpp
+++ b/Engine/Plugins/Experimental/GeometryScripting/Source/GeometryScriptingCore/Private/MeshUVFunctions.cpp
@@ -8,6 +8,7 @@
 #include "Parameterization/DynamicMeshUVEditor.h"
 #include "Selections/MeshConnectedComponents.h"
 #include "Parameterization/PatchBasedMeshUVGenerator.h"
+#include "DynamicMesh/MeshNormals.h"
 #include "XAtlasWrapper.h"
 
 #include "Async/ParallelFor.h"
@@ -778,4 +779,202 @@ UDynamicMesh* UGeometryScriptLibrary_MeshUVFunctions::GetMeshPerVertexUVs(
 
 
 
+
+
+UDynamicMesh* UGeometryScriptLibrary_MeshUVFunctions::CopyMeshUVLayerToMesh(
+	UDynamicMesh* CopyFromMesh,
+	int UVSetIndex,
+	UPARAM(DisplayName = "Copy To UV Mesh", ref) UDynamicMesh* CopyToUVMesh,
+	UPARAM(DisplayName = "Copy To UV Mesh") UDynamicMesh*& CopyToUVMeshOut,
+	bool& bInvalidTopology,
+	bool& bIsValidUVSet,
+	UGeometryScriptDebug* Debug)
+{
+	if (CopyFromMesh == nullptr)
+	{
+		UE::Geometry::AppendError(Debug, EGeometryScriptErrorType::InvalidInputs, LOCTEXT("CopyMeshUVLayerToMesh_InvalidInput", "CopyMeshUVLayerToMesh: CopyFromMesh is Null"));
+		return CopyFromMesh;
+	}
+	if (CopyToUVMesh == nullptr)
+	{
+		UE::Geometry::AppendError(Debug, EGeometryScriptErrorType::InvalidInputs, LOCTEXT("CopyMeshUVLayerToMesh_InvalidInput2", "CopyMeshUVLayerToMesh: CopyToUVMesh is Null"));
+		return CopyFromMesh;
+	}
+	if (CopyFromMesh == CopyToUVMesh)
+	{
+		// TODO: can actually support this but complicates the code below...
+		UE::Geometry::AppendError(Debug, EGeometryScriptErrorType::InvalidInputs, LOCTEXT("CopyMeshToUVMesh_SameMeshes", "CopyMeshUVLayerToMesh: CopyFromMesh and CopyToUVMesh are the same mesh, this is not supported"));
+		return CopyFromMesh;
+	}
+
+	FDynamicMesh3 UVMesh;
+	bIsValidUVSet = false;
+	bInvalidTopology = false;
+	CopyFromMesh->ProcessMesh([&](const FDynamicMesh3& FromMesh)
+	{
+		const FDynamicMeshUVOverlay* UVOverlay = (FromMesh.HasAttributes() && UVSetIndex < FromMesh.Attributes()->NumUVLayers()) ?
+			FromMesh.Attributes()->GetUVLayer(UVSetIndex) : nullptr;
+		if (UVOverlay == nullptr)
+		{
+			UE::Geometry::AppendError(Debug, EGeometryScriptErrorType::InvalidInputs, LOCTEXT("CopyMeshToUVMesh_InvalidUVSet", "CopyMeshUVLayerToMesh: UVSetIndex does not exist on CopyFromMesh"));
+			return;
+		}
+		bIsValidUVSet = true;
+
+		UVMesh.EnableTriangleGroups();
+		UVMesh.EnableAttributes();
+		UVMesh.Attributes()->SetNumUVLayers(0);
+
+		const FDynamicMeshMaterialAttribute* FromMaterialID = (FromMesh.HasAttributes() && FromMesh.Attributes()->HasMaterialID()) ?
+			FromMesh.Attributes()->GetMaterialID() : nullptr;
+		FDynamicMeshMaterialAttribute* ToMaterialID = nullptr;
+		if (FromMaterialID)
+		{
+			UVMesh.Attributes()->EnableMaterialID();
+			ToMaterialID = UVMesh.Attributes()->GetMaterialID();
+		}
+
+		UVMesh.BeginUnsafeVerticesInsert();
+		for (int32 elemid : UVOverlay->ElementIndicesItr())
+		{
+			FVector2f UV = UVOverlay->GetElement(elemid);
+			UVMesh.InsertVertex(elemid, FVector3d(UV.X, UV.Y, 0), true);
+		}
+		UVMesh.EndUnsafeVerticesInsert();
+		UVMesh.BeginUnsafeTrianglesInsert();
+		for (int32 tid : FromMesh.TriangleIndicesItr())
+		{
+			FIndex3i UVTri = UVOverlay->GetTriangle(tid);
+			int32 GroupID = FromMesh.GetTriangleGroup(tid);
+			EMeshResult Result = UVMesh.InsertTriangle(tid, UVTri, GroupID, true);
+			if (Result != EMeshResult::Ok)
+			{
+				bInvalidTopology = true;
+			}
+			else
+			{
+				if (FromMaterialID)
+				{
+					ToMaterialID->SetValue(tid, FromMaterialID->GetValue(tid));		// could we use Copy() here ?
+				}
+			}
+		}
+		UVMesh.EndUnsafeTrianglesInsert();
+	});
+
+	FMeshNormals::InitializeOverlayToPerVertexNormals(UVMesh.Attributes()->PrimaryNormals());
+
+	CopyToUVMesh->SetMesh(MoveTemp(UVMesh));
+	CopyToUVMeshOut = CopyToUVMesh;
+
+	return CopyFromMesh;
+}
+
+
+
+UDynamicMesh* UGeometryScriptLibrary_MeshUVFunctions::CopyMeshToMeshUVLayer(
+	UDynamicMesh* CopyFromUVMesh,
+	int ToUVSetIndex,
+	UDynamicMesh* CopyToMesh,
+	UDynamicMesh*& CopyToMeshOut,
+	bool& bFoundTopologyErrors,
+	bool& bIsValidUVSet,
+	bool bOnlyUVPositions,
+	UGeometryScriptDebug* Debug)
+{
+	if (CopyFromUVMesh == nullptr)
+	{
+		UE::Geometry::AppendError(Debug, EGeometryScriptErrorType::InvalidInputs, LOCTEXT("CopyMeshToMeshUVLayer_InvalidInput", "CopyMeshToMeshUVLayer: CopyFromUVMesh is Null"));
+		return CopyFromUVMesh;
+	}
+	if (CopyToMesh == nullptr)
+	{
+		UE::Geometry::AppendError(Debug, EGeometryScriptErrorType::InvalidInputs, LOCTEXT("CopyMeshToMeshUVLayer_InvalidInput2", "CopyMeshToMeshUVLayer: CopyToUVMesh is Null"));
+		return CopyFromUVMesh;
+	}
+	if (CopyFromUVMesh == CopyToMesh)
+	{
+		// TODO: can actually support this but complicates the code below...
+		UE::Geometry::AppendError(Debug, EGeometryScriptErrorType::InvalidInputs, LOCTEXT("CopyMeshToMeshUVLayer_SameMeshes", "CopyMeshToMeshUVLayer: CopyFromUVMesh and CopyToMesh are the same mesh, this is not supported"));
+		return CopyFromUVMesh;
+	}
+
+	bFoundTopologyErrors = false;
+	bIsValidUVSet = false;
+	CopyToMesh->EditMesh([&](FDynamicMesh3& EditMesh)
+	{
+		FDynamicMeshUVOverlay* UVOverlay = (EditMesh.HasAttributes() && ToUVSetIndex < EditMesh.Attributes()->NumUVLayers()) ?
+			EditMesh.Attributes()->GetUVLayer(ToUVSetIndex) : nullptr;
+		if (UVOverlay == nullptr)
+		{
+			UE::Geometry::AppendError(Debug, EGeometryScriptErrorType::InvalidInputs, LOCTEXT("CopyMeshToMeshUVLayer_InvalidUVSet", "CopyMeshToMeshUVLayer: ToUVSetIndex does not exist on CopyFromMesh"));
+			return;
+		}
+		bIsValidUVSet = true;
+
+		CopyFromUVMesh->ProcessMesh([&](const FDynamicMesh3& UVMesh)
+		{
+			if (bOnlyUVPositions)
+			{
+				if ( UVMesh.MaxVertexID() <= UVOverlay->MaxElementID() )
+				{
+					for (int32 vid : UVMesh.VertexIndicesItr())
+					{
+						if (UVOverlay->IsElement(vid))
+						{
+							FVector3d Pos = UVMesh.GetVertex(vid);
+							UVOverlay->SetElement(vid, FVector2f((float)Pos.X, (float)Pos.Y));
+						}
+						else
+						{
+							bFoundTopologyErrors = true;
+						}
+					}
+				}
+				else
+				{
+					bFoundTopologyErrors = true;
+				}
+			}
+			else
+			{
+				if (UVMesh.MaxTriangleID() <= EditMesh.MaxTriangleID())
+				{
+					UVOverlay->ClearElements();
+					UVOverlay->BeginUnsafeElementsInsert();
+					for (int32 vid : UVMesh.VertexIndicesItr())
+					{
+						FVector3d Pos = UVMesh.GetVertex(vid);
+						FVector2f UV((float)Pos.X, (float)Pos.Y);
+						UVOverlay->InsertElement(vid, &UV.X, true);
+					}
+					UVOverlay->EndUnsafeElementsInsert();
+					for (int32 tid : UVMesh.TriangleIndicesItr())
+					{
+						if (EditMesh.IsTriangle(tid))
+						{
+							FIndex3i Tri = UVMesh.GetTriangle(tid);
+							UVOverlay->SetTriangle(tid, Tri);
+						}
+						else
+						{
+							bFoundTopologyErrors = true;
+						}
+					}
+				}
+				else
+				{
+					bFoundTopologyErrors = true;
+				}
+			}
+		});
+	});
+
+
+	CopyToMeshOut = CopyToMesh;
+	return CopyFromUVMesh;
+}
+
+
+
 #undef LOCTEXT_NAMESPACE
\ No newline at end of file
diff --git a/Engine/Plugins/Experimental/GeometryScripting/Source/GeometryScriptingCore/Public/GeometryScript/MeshUVFunctions.h b/Engine/Plugins/Experimental/GeometryScripting/Source/GeometryScriptingCore/Public/GeometryScript/MeshUVFunctions.h
--- a/Engine/Plugins/Experimental/GeometryScripting/Source/GeometryScriptingCore/Public/GeometryScript/MeshUVFunctions.h
+++ b/Engine/Plugins/Experimental/GeometryScripting/Source/GeometryScriptingCore/Public/GeometryScript/MeshUVFunctions.h
@@ -286,4 +286,55 @@ public:
 		bool& bHasSplitUVs,
 		UGeometryScriptDebug* Debug = nullptr);
 
+
+
+	/**
+	 * Copy the 2D UVs from the given UVSetIndex in CopyFromMesh to the 3D vertex positions in CopyToUVMesh,
+	 * with the triangle mesh topolgoy defined by the UV Set. Generally this "UV Mesh" topolgoy will not
+	 * be the same as the 3D mesh topology. Polygroup IDs and Material IDs are preserved in the UVMesh.
+	 * 
+	 * 2D UV Positions are copied to 3D as (X, Y, 0) 
+	 * 
+	 * CopyMeshToMeshUVLayer will copy the 3D UV Mesh back to the UV Set. This pair of functions can
+	 * then be used to implement UV generation/editing via other mesh functions.
+	 * 
+	 * @param bInvalidTopology will be returned true if any topological issues were found
+	 * @param bIsValidUVSet will be returned false if UVSetIndex is not available
+	 */
+	UFUNCTION(BlueprintCallable, Category = "GeometryScript|UVs", meta=(ScriptMethod))
+	static UPARAM(DisplayName = "Copy From Mesh") UDynamicMesh* 
+	CopyMeshUVLayerToMesh(  
+		UDynamicMesh* CopyFromMesh, 
+		int UVSetIndex,
+		UPARAM(DisplayName = "Copy To UV Mesh", ref) UDynamicMesh* CopyToUVMesh, 
+		UPARAM(DisplayName = "Copy To UV Mesh") UDynamicMesh*& CopyToUVMeshOut,
+		bool& bInvalidTopology,
+		bool& bIsValidUVSet,
+		UGeometryScriptDebug* Debug = nullptr);
+
+	/**
+	 * Transfer the 3D vertex positions and triangles of CopyFromUVMesh to the given UV Layer identified by ToUVSetIndex of CopyToMesh.
+	 * 3D positions (X,Y,Z) will be copied as UV positions (X,Y), ie Z is ignored.
+	 * 
+	 * bOnlyUVPositions controls whether only UV positions will be updated, or if the UV topology will be fully replaced.
+	 * When false, CopyFromUVMesh must currently have a MaxVertexID <= that of the UV Layer MaxElementID
+	 * When true, CopyFromUVMesh must currently have a MaxTriangleID <= that of CopyToMesh
+	 * 
+	 * @param bInvalidTopology will be returned true if any topological inconsistencies are found (but the operation will generally continue)
+	 * @param bIsValidUVSet will be returned false if ToUVSetIndex is not available
+	 * @param bOnlyUVPositions if true, only (valid, matching) UV positions are updated, a full new UV topology is created
+	 */
+	UFUNCTION(BlueprintCallable, Category = "GeometryScript|UVs", meta=(ScriptMethod))
+	static UPARAM(DisplayName = "Copy From Mesh") UDynamicMesh* 
+	CopyMeshToMeshUVLayer(  
+		UDynamicMesh* CopyFromUVMesh, 
+		int ToUVSetIndex,
+		UPARAM(DisplayName = "Copy To Mesh", ref) UDynamicMesh* CopyToMesh, 
+		UPARAM(DisplayName = "Copy To Mesh") UDynamicMesh*& CopyToMeshOut,
+		bool& bFoundTopologyErrors,
+		bool& bIsValidUVSet,
+		bool bOnlyUVPositions = true,
+		UGeometryScriptDebug* Debug = nullptr);
+
+
 };
\ No newline at end of file
