ModelingTools: Add VertexColor output type support to BakeVertex tool and geometry script function.

diff --git a/Engine/Plugins/Experimental/GeometryScripting/Source/GeometryScriptingCore/Private/MeshBakeFunctions.cpp b/Engine/Plugins/Experimental/GeometryScripting/Source/GeometryScriptingCore/Private/MeshBakeFunctions.cpp
--- a/Engine/Plugins/Experimental/GeometryScripting/Source/GeometryScriptingCore/Private/MeshBakeFunctions.cpp
+++ b/Engine/Plugins/Experimental/GeometryScripting/Source/GeometryScriptingCore/Private/MeshBakeFunctions.cpp
@@ -688,15 +688,9 @@ namespace GeometryScriptBakeLocals
 
 		if (BakeTypes.OutputMode == EGeometryScriptBakeOutputMode::RGBA)
 		{
-			auto IsValidBakeType = [&BakeVertexPrefix](EGeometryScriptBakeTypes BakeType, TArray<FGeometryScriptDebugMessage>* Debug)
+			auto IsValidBakeType = [](EGeometryScriptBakeTypes, TArray<FGeometryScriptDebugMessage>*)
 			{
-				const bool bIsValid = BakeType != EGeometryScriptBakeTypes::VertexColor;
-				if (!bIsValid)
-				{
-					const FText BakeTypeName = FText::FromName(StaticEnum<EGeometryScriptBakeTypes>()->GetNameByIndex(static_cast<int32>(BakeType)));
-					UE::Geometry::AppendError(Debug, EGeometryScriptErrorType::InvalidInputs, FText::Format(LOCTEXT("BakeVertex_InvalidColorEval", "{0}: {1} bake type is not a supported RGBA evaluator."), BakeVertexPrefix, BakeTypeName));
-				}
-				return bIsValid;
+				return true;
 			};
 			
 			Baker.BakeMode = FMeshVertexBaker::EBakeMode::RGBA;
diff --git a/Engine/Plugins/Experimental/MeshModelingToolsetExp/Source/MeshModelingToolsExp/Private/BakeMeshAttributeVertexTool.cpp b/Engine/Plugins/Experimental/MeshModelingToolsetExp/Source/MeshModelingToolsExp/Private/BakeMeshAttributeVertexTool.cpp
--- a/Engine/Plugins/Experimental/MeshModelingToolsetExp/Source/MeshModelingToolsExp/Private/BakeMeshAttributeVertexTool.cpp
+++ b/Engine/Plugins/Experimental/MeshModelingToolsetExp/Source/MeshModelingToolsExp/Private/BakeMeshAttributeVertexTool.cpp
@@ -213,6 +213,13 @@ public:
 				Baker->ColorEvaluator = TextureEval;
 				break;
 			}
+			case EBakeMapType::VertexColor:
+			{
+				TSharedPtr<FMeshPropertyMapEvaluator, ESPMode::ThreadSafe> PropertyEval = MakeShared<FMeshPropertyMapEvaluator, ESPMode::ThreadSafe>();
+				PropertyEval->Property = EMeshPropertyMapType::VertexColor;
+				Baker->ColorEvaluator = PropertyEval;
+				break;
+			}
 			default:
 				break;
 			}
@@ -681,6 +688,7 @@ void UBakeMeshAttributeVertexTool::UpdateResult()
 		case EBakeMapType::FaceNormal:
 		case EBakeMapType::Position:
 		case EBakeMapType::MaterialID:
+		case EBakeMapType::VertexColor:
 			OpState |= UpdateResult_MeshProperty(Dimensions);
 			break;
 		case EBakeMapType::Texture:
diff --git a/Engine/Plugins/Experimental/MeshModelingToolsetExp/Source/MeshModelingToolsExp/Public/BakeMeshAttributeVertexTool.h b/Engine/Plugins/Experimental/MeshModelingToolsetExp/Source/MeshModelingToolsExp/Public/BakeMeshAttributeVertexTool.h
--- a/Engine/Plugins/Experimental/MeshModelingToolsetExp/Source/MeshModelingToolsExp/Public/BakeMeshAttributeVertexTool.h
+++ b/Engine/Plugins/Experimental/MeshModelingToolsetExp/Source/MeshModelingToolsExp/Public/BakeMeshAttributeVertexTool.h
@@ -65,7 +65,7 @@ public:
 
 	/** The bake output type to generate */
 	UPROPERTY(EditAnywhere, Category = BakeOutput, meta=(Bitmask, BitmaskEnum = "/Script/MeshModelingToolsExp.EBakeMapType",
-		ValidEnumValues="TangentSpaceNormal, AmbientOcclusion, BentNormal, Curvature, Texture, ObjectSpaceNormal, FaceNormal, Position, MaterialID, MultiTexture",
+		ValidEnumValues="TangentSpaceNormal, AmbientOcclusion, BentNormal, Curvature, Texture, ObjectSpaceNormal, FaceNormal, Position, MaterialID, MultiTexture, VertexColor",
 		EditCondition="OutputMode == EBakeVertexOutput::RGBA", EditConditionHides))
 	int32 OutputType = static_cast<int32>(EBakeMapType::TangentSpaceNormal);
 
