GeometryScript: Added vertex color transformation functions (SRGBToLinear, LinearToSRGB) for dynamic meshes.

diff --git a/Engine/Plugins/Experimental/GeometryScripting/Source/GeometryScriptingCore/Private/MeshVertexColorFunctions.cpp b/Engine/Plugins/Experimental/GeometryScripting/Source/GeometryScriptingCore/Private/MeshVertexColorFunctions.cpp
--- a/Engine/Plugins/Experimental/GeometryScripting/Source/GeometryScriptingCore/Private/MeshVertexColorFunctions.cpp
+++ b/Engine/Plugins/Experimental/GeometryScripting/Source/GeometryScriptingCore/Private/MeshVertexColorFunctions.cpp
@@ -4,6 +4,7 @@
 
 #include "DynamicMesh/DynamicMesh3.h"
 #include "UDynamicMesh.h"
+#include "Util/ColorConstants.h"
 
 using namespace UE::Geometry;
 
@@ -219,6 +220,70 @@ UDynamicMesh* UGeometryScriptLibrary_MeshVertexColorFunctions::GetMeshPerVertexC
 	return TargetMesh;
 }
 
+UDynamicMesh* UGeometryScriptLibrary_MeshVertexColorFunctions::ConvertMeshVertexColorsSRGBToLinear(
+	UDynamicMesh* TargetMesh,
+	UGeometryScriptDebug* Debug)
+{
+	if (TargetMesh == nullptr)
+	{
+		UE::Geometry::AppendError(Debug, EGeometryScriptErrorType::InvalidInputs, LOCTEXT("ConvertMeshVertexColorsSRGBToLinear_InvalidInput", "ConvertMeshVertexColorsSRGBToLinear: TargetMesh is Null"));
+		return TargetMesh;
+	}
+
+	TargetMesh->EditMesh([&](FDynamicMesh3& EditMesh) 
+	{
+		if (EditMesh.HasAttributes() == false || EditMesh.Attributes()->HasPrimaryColors() == false)
+		{
+			return;
+		}
+
+		FDynamicMeshColorOverlay* Colors = EditMesh.Attributes()->PrimaryColors();
+		if (Colors->ElementCount() > 0)
+		{
+			for (int32 ElementID : Colors->ElementIndicesItr())
+			{
+				FVector4f Existing = Colors->GetElement(ElementID);
+				LinearColors::SRGBToLinear(Existing);
+				Colors->SetElement(ElementID, Existing);
+			}
+		}
+	}, EDynamicMeshChangeType::GeneralEdit, EDynamicMeshAttributeChangeFlags::Unknown, false);
+	
+	return TargetMesh;
+}
+
+UDynamicMesh* UGeometryScriptLibrary_MeshVertexColorFunctions::ConvertMeshVertexColorsLinearToSRGB(
+	UDynamicMesh* TargetMesh,
+	UGeometryScriptDebug* Debug)
+{
+	if (TargetMesh == nullptr)
+	{
+		UE::Geometry::AppendError(Debug, EGeometryScriptErrorType::InvalidInputs, LOCTEXT("ConvertMeshVertexColorsLinearToSRGB_InvalidInput", "ConvertMeshVertexColorsLinearToSRGB: TargetMesh is Null"));
+		return TargetMesh;
+	}
+
+	TargetMesh->EditMesh([&](FDynamicMesh3& EditMesh) 
+	{
+		if (EditMesh.HasAttributes() == false || EditMesh.Attributes()->HasPrimaryColors() == false)
+		{
+			return;
+		}
+
+		FDynamicMeshColorOverlay* Colors = EditMesh.Attributes()->PrimaryColors();
+		if (Colors->ElementCount() > 0)
+		{
+			for (int32 ElementID : Colors->ElementIndicesItr())
+			{
+				FVector4f Existing = Colors->GetElement(ElementID);
+				LinearColors::LinearToSRGB(Existing);
+				Colors->SetElement(ElementID, Existing);
+			}
+		}
+	}, EDynamicMeshChangeType::GeneralEdit, EDynamicMeshAttributeChangeFlags::Unknown, false);
+	
+	return TargetMesh;
+}
+
 
 
 
diff --git a/Engine/Plugins/Experimental/GeometryScripting/Source/GeometryScriptingCore/Public/GeometryScript/MeshVertexColorFunctions.h b/Engine/Plugins/Experimental/GeometryScripting/Source/GeometryScriptingCore/Public/GeometryScript/MeshVertexColorFunctions.h
--- a/Engine/Plugins/Experimental/GeometryScripting/Source/GeometryScriptingCore/Public/GeometryScript/MeshVertexColorFunctions.h
+++ b/Engine/Plugins/Experimental/GeometryScripting/Source/GeometryScriptingCore/Public/GeometryScript/MeshVertexColorFunctions.h
@@ -59,5 +59,25 @@ public:
 		bool& bIsValidColorSet,
 		bool& bHasVertexIDGaps,
 		bool bBlendSplitVertexValues = true);
+
+	/**
+	 * Apply a SRGB to Linear color transformation on all vertex colors
+	 * on the mesh.
+	 */
+	UFUNCTION(BlueprintCallable, Category = "GeometryScript|VertexColor", meta=(ScriptMethod))
+	static UPARAM(DisplayName = "Target Mesh") UDynamicMesh*
+	ConvertMeshVertexColorsSRGBToLinear(
+		UDynamicMesh* TargetMesh,
+		UGeometryScriptDebug* Debug = nullptr);
+
+	/**
+	 * Apply a Linear to SRGB color transformation on all vertex colors
+	 * on the mesh.
+	 */
+	UFUNCTION(BlueprintCallable, Category = "GeometryScript|VertexColor", meta=(ScriptMethod))
+	static UPARAM(DisplayName = "Target Mesh") UDynamicMesh*
+	ConvertMeshVertexColorsLinearToSRGB(
+		UDynamicMesh* TargetMesh,
+		UGeometryScriptDebug* Debug = nullptr);
 };
 
