GeometryScript: add enum to be used to indicate how operations should behave when input Selection is empty

diff --git a/Engine/Plugins/Experimental/GeometryScripting/Source/GeometryScriptingCore/Public/GeometryScript/GeometryScriptSelectionTypes.h b/Engine/Plugins/Experimental/GeometryScripting/Source/GeometryScriptingCore/Public/GeometryScript/GeometryScriptSelectionTypes.h
--- a/Engine/Plugins/Experimental/GeometryScripting/Source/GeometryScriptingCore/Public/GeometryScript/GeometryScriptSelectionTypes.h
+++ b/Engine/Plugins/Experimental/GeometryScripting/Source/GeometryScriptingCore/Public/GeometryScript/GeometryScriptSelectionTypes.h
@@ -42,6 +42,15 @@ enum class EGeometryScriptCombineSelectionMode : uint8
 	Intersection
 };
 
+/**
+ * Behavior of operations when a MeshSelection is empty
+ */
+UENUM(BlueprintType)
+enum class EGeometryScriptEmptySelectionBehavior : uint8
+{
+	FullMeshSelection = 0,
+	EmptySelection = 1
+};
 
 /**
  * FGeometryScriptMeshSelection is a container for a Mesh Selection used in Geometry Script.
