GeometryScript: replace usage of IndexList with new MeshSelection in ApplyAdaptiveTessellation

diff --git a/Engine/Plugins/Experimental/GeometryScripting/Source/GeometryScriptingCore/Private/MeshSubdivideFunctions.cpp b/Engine/Plugins/Experimental/GeometryScripting/Source/GeometryScriptingCore/Private/MeshSubdivideFunctions.cpp
--- a/Engine/Plugins/Experimental/GeometryScripting/Source/GeometryScriptingCore/Private/MeshSubdivideFunctions.cpp
+++ b/Engine/Plugins/Experimental/GeometryScripting/Source/GeometryScriptingCore/Private/MeshSubdivideFunctions.cpp
@@ -131,8 +131,8 @@ UDynamicMesh* UGeometryScriptLibrary_MeshSubdivideFunctions::ApplyUniformTessell
 
 UDynamicMesh* UGeometryScriptLibrary_MeshSubdivideFunctions::ApplyAdaptiveTessellation(
 	UDynamicMesh* TargetMesh,
+	FGeometryScriptMeshSelection Selection,
 	FGeometryScriptAdaptiveTessellateOptions Options,
-	FGeometryScriptIndexList IndexList,
 	int TessellationLevel,
 	EAdaptiveTessellatePatternType PatternType,
 	UGeometryScriptDebug* Debug)
@@ -148,71 +148,37 @@ UDynamicMesh* UGeometryScriptLibrary_MeshSubdivideFunctions::ApplyAdaptiveTessel
 	}
 	if (PatternType != EAdaptiveTessellatePatternType::ConcentricRings)
 	{
-		UE::Geometry::AppendError(Debug, EGeometryScriptErrorType::InvalidInputs, LOCTEXT("ApplyAdapativeTessellation_InvalidPatternType", "Only ConcentricRings pattern is currently supported"));
+		UE::Geometry::AppendError(Debug, EGeometryScriptErrorType::InvalidInputs, LOCTEXT("ApplyAdapativeTessellation_InvalidPatternType", "ApplyAdapativeTessellation: Only ConcentricRings pattern is currently supported"));
 		return TargetMesh;
 	} 
-	if (IndexList.IndexType != EGeometryScriptIndexType::Triangle)
+	if (Selection.IsEmpty() && Options.EmptyBehavior != EGeometryScriptEmptySelectionBehavior::FullMeshSelection )
 	{
-		UE::Geometry::AppendError(Debug, EGeometryScriptErrorType::InvalidInputs, LOCTEXT("ApplyAdapativeTessellation_InvalidSelectionType", "Only Triangle selection is currently supported"));
 		return TargetMesh;
-	} 
+	}
 
 	TargetMesh->EditMesh([&](FDynamicMesh3& EditMesh) 
 	{	
 		FDynamicMesh3 TessellatedMesh;
 		FAdaptiveTessellate Tessellator(&EditMesh, &TessellatedMesh);
 
-		TSharedPtr<TArray<int>> List = IndexList.List;
+		TArray<int32> Triangles;
+		Selection.ConvertToMeshIndexArray(EditMesh, Triangles, EGeometryScriptIndexType::Triangle);
+		if (Triangles.Num() == 0 && Options.EmptyBehavior != EGeometryScriptEmptySelectionBehavior::FullMeshSelection)
+		{
+			return;
+		}
+
 		TUniquePtr<FTessellationPattern> Pattern; 
-		if (List == nullptr) //if list is not provided then tessellate the whole mesh
-		{	
-			switch(PatternType) 
-			{
-				case EAdaptiveTessellatePatternType::Uniform:
-				{
-					Pattern = nullptr; //TODO: implement
-					break;
-				}
-				case EAdaptiveTessellatePatternType::InnerUniform: 
-				{
-					Pattern = nullptr; //TODO: implement
-					break;
-				}
-				case EAdaptiveTessellatePatternType::ConcentricRings:
-				{
-					Pattern = FAdaptiveTessellate::CreateConcentricRingsTessellationPattern(&EditMesh, TessellationLevel);
-					break;
-				} 
-				default:
-				{
-					checkSlow(false);
-				}
-			}
+		if (PatternType == EAdaptiveTessellatePatternType::ConcentricRings)
+		{
+			Pattern = (Triangles.Num() > 0) ?
+				FAdaptiveTessellate::CreateConcentricRingsTessellationPattern(&EditMesh, TessellationLevel, Triangles)
+				: FAdaptiveTessellate::CreateConcentricRingsTessellationPattern(&EditMesh, TessellationLevel);
 		}
 		else
 		{
-			switch(PatternType) 
-			{
-				case EAdaptiveTessellatePatternType::Uniform:
-				{
-					Pattern = nullptr; //TODO: implement
-					break;
-				}
-				case EAdaptiveTessellatePatternType::InnerUniform: 
-				{
-					Pattern = nullptr; //TODO: implement
-					break;
-				}
-				case EAdaptiveTessellatePatternType::ConcentricRings:
-				{
-					Pattern = FAdaptiveTessellate::CreateConcentricRingsTessellationPattern(&EditMesh, TessellationLevel, *List);
-					break;
-				} 
-				default:
-				{
-					checkSlow(false);
-				}
-			}
+			checkSlow(false);
+			return;
 		}
 
 		Tessellator.SetPattern(Pattern.Get());
diff --git a/Engine/Plugins/Experimental/GeometryScripting/Source/GeometryScriptingCore/Public/GeometryScript/MeshSubdivideFunctions.h b/Engine/Plugins/Experimental/GeometryScripting/Source/GeometryScriptingCore/Public/GeometryScript/MeshSubdivideFunctions.h
--- a/Engine/Plugins/Experimental/GeometryScripting/Source/GeometryScriptingCore/Public/GeometryScript/MeshSubdivideFunctions.h
+++ b/Engine/Plugins/Experimental/GeometryScripting/Source/GeometryScriptingCore/Public/GeometryScript/MeshSubdivideFunctions.h
@@ -5,6 +5,7 @@
 #include "CoreMinimal.h"
 #include "Kismet/BlueprintFunctionLibrary.h"
 #include "GeometryScript/GeometryScriptTypes.h"
+#include "GeometryScript/GeometryScriptSelectionTypes.h"
 #include "MeshSubdivideFunctions.generated.h"
 
 class UDynamicMesh;
@@ -28,19 +29,23 @@ public:
 UENUM(BlueprintType)
 enum class EAdaptiveTessellatePatternType : uint8
 {
-	Uniform,
-	InnerUniform,
-	ConcentricRings
+	Uniform = 0,
+	InnerUniform = 1,
+	ConcentricRings = 2
 };
 
 
-USTRUCT(BlueprintType, meta = (DisplayName = "Additional Adaptive Tessellate Options"))
+USTRUCT(BlueprintType, meta = (DisplayName = "Adaptive Tessellate Options"))
 struct GEOMETRYSCRIPTINGCORE_API FGeometryScriptAdaptiveTessellateOptions
 {
 	GENERATED_BODY()
 public:
-	UPROPERTY(BlueprintReadWrite, Category = Options)
+	UPROPERTY(EditAnywhere, BlueprintReadWrite, Category = Options)
 	bool bEnableMultithreading = true;
+
+	/** EmptyBehavior Defines how an empty input selection should be interpreted */
+	UPROPERTY(EditAnywhere, BlueprintReadWrite, Category = Options)
+	EGeometryScriptEmptySelectionBehavior EmptyBehavior = EGeometryScriptEmptySelectionBehavior::FullMeshSelection;
 };
 
 
@@ -78,8 +83,8 @@ public:
 	static UPARAM(DisplayName = "Target Mesh") UDynamicMesh* 
 	ApplyAdaptiveTessellation(
 		UDynamicMesh* TargetMesh,
+		FGeometryScriptMeshSelection Selection,
 		FGeometryScriptAdaptiveTessellateOptions Options,
-		FGeometryScriptIndexList IndexList,
 		int TessellationLevel = 1,
 		EAdaptiveTessellatePatternType PatternType = EAdaptiveTessellatePatternType::ConcentricRings,
 		UGeometryScriptDebug* Debug = nullptr);
