work around loctext issue

diff --git a/Engine/Plugins/Experimental/GeometryScripting/Source/GeometryScriptingCore/Private/MeshSubdivideFunctions.cpp b/Engine/Plugins/Experimental/GeometryScripting/Source/GeometryScriptingCore/Private/MeshSubdivideFunctions.cpp
--- a/Engine/Plugins/Experimental/GeometryScripting/Source/GeometryScriptingCore/Private/MeshSubdivideFunctions.cpp
+++ b/Engine/Plugins/Experimental/GeometryScripting/Source/GeometryScriptingCore/Private/MeshSubdivideFunctions.cpp
@@ -148,7 +148,7 @@ UDynamicMesh* UGeometryScriptLibrary_MeshSubdivideFunctions::ApplyAdaptiveTessel
 	}
 	if (PatternType != EAdaptiveTessellatePatternType::ConcentricRings)
 	{
-		UE::Geometry::AppendError(Debug, EGeometryScriptErrorType::InvalidInputs, LOCTEXT("ApplyAdapativeTessellation_InvalidPatternType", "ApplyAdapativeTessellation: Only ConcentricRings pattern is currently supported"));
+		UE::Geometry::AppendError(Debug, EGeometryScriptErrorType::InvalidInputs, LOCTEXT("ApplyAdapativeTessellation_InvalidPatternTypeNew", "ApplyAdapativeTessellation: Only ConcentricRings pattern is currently supported"));
 		return TargetMesh;
 	} 
 	if (Selection.IsEmpty() && Options.EmptyBehavior != EGeometryScriptEmptySelectionBehavior::FullMeshSelection )
