GeometryScript: add AppendMeshTransformed function that appends N copies of a mesh to another mesh, with a custom Transform applied for each

diff --git a/Engine/Plugins/Experimental/GeometryScripting/Source/GeometryScriptingCore/Private/MeshBasicEditFunctions.cpp b/Engine/Plugins/Experimental/GeometryScripting/Source/GeometryScriptingCore/Private/MeshBasicEditFunctions.cpp
--- a/Engine/Plugins/Experimental/GeometryScripting/Source/GeometryScriptingCore/Private/MeshBasicEditFunctions.cpp
+++ b/Engine/Plugins/Experimental/GeometryScripting/Source/GeometryScriptingCore/Private/MeshBasicEditFunctions.cpp
@@ -7,6 +7,7 @@
 #include "DynamicMesh/DynamicMeshAttributeSet.h"
 #include "DynamicMesh/MeshTransforms.h"
 #include "DynamicMeshEditor.h"
+#include "TransformSequence.h"
 
 using namespace UE::Geometry;
 
@@ -371,6 +372,79 @@ UDynamicMesh* UGeometryScriptLibrary_MeshBasicEditFunctions::AppendMesh(
 
 
 
+UDynamicMesh* UGeometryScriptLibrary_MeshBasicEditFunctions::AppendMeshTransformed(
+	UDynamicMesh* TargetMesh,
+	UDynamicMesh* AppendMesh,
+	const TArray<FTransform>& AppendTransforms, 
+	FTransform ConstantTransform,
+	bool bConstantTransformIsRelative,
+	bool bDeferChangeNotifications,
+	UGeometryScriptDebug* Debug)
+{
+	if (TargetMesh == nullptr)
+	{
+		UE::Geometry::AppendError(Debug, EGeometryScriptErrorType::InvalidInputs, LOCTEXT("AppendMeshTransformed_InvalidInput1", "AppendMeshTransformed: TargetMesh is Null"));
+		return TargetMesh;
+	}
+	if (AppendMesh == nullptr)
+	{
+		UE::Geometry::AppendError(Debug, EGeometryScriptErrorType::InvalidInputs, LOCTEXT("AppendMeshTransformed_InvalidInput2", "AppendMeshTransformed: AppendMesh is Null"));
+		return TargetMesh;
+	}
+	if (AppendTransforms.IsEmpty())
+	{
+		UE::Geometry::AppendError(Debug, EGeometryScriptErrorType::InvalidInputs, LOCTEXT("AppendMeshTransformed_NoTransforms", "AppendMeshTransformed: AppendTransforms array is empty"));
+		return TargetMesh;
+	}
+
+	TargetMesh->EditMesh([&](FDynamicMesh3& AppendToMesh)
+	{
+		AppendMesh->ProcessMesh([&](const FDynamicMesh3& OtherMesh)
+		{
+			FMeshIndexMappings TmpMappings;
+			FDynamicMeshEditor Editor(&AppendToMesh);
+			const FDynamicMesh3* UseOtherMesh = &OtherMesh;
+			FDynamicMesh3 TmpMesh;
+			if (UseOtherMesh == &AppendToMesh)
+			{
+				TmpMesh = OtherMesh;	// need  to make a copy if we are appending to ourself
+				UseOtherMesh = &TmpMesh;
+			}
+			for (FTransform AppendTransform : AppendTransforms)
+			{
+				FTransformSequence3d TransformSequence;
+
+				if (bConstantTransformIsRelative)
+				{
+					TransformSequence.Append(ConstantTransform);
+					TransformSequence.Append(AppendTransform);
+				}
+				else
+				{
+					// want to apply the constant transform's rotate/scale after
+					// the main transform rotate/scale, so the main positioning 
+					// translation has to be deferred until after that
+					FVector Translation = AppendTransform.GetLocation();
+					AppendTransform.SetTranslation(FVector::Zero());
+
+					TransformSequence.Append(AppendTransform);
+					TransformSequence.Append(ConstantTransform);
+					TransformSequence.Append(FTransform(Translation));
+				}
+
+				Editor.AppendMesh(UseOtherMesh, TmpMappings,
+					[&](int, const FVector3d& Position) { return TransformSequence.TransformPosition(Position); },
+					[&](int, const FVector3d& Normal) { return TransformSequence.TransformNormal(Normal); });
+				TmpMappings.Reset();
+			}
+		});
+	}, EDynamicMeshChangeType::GeneralEdit, EDynamicMeshAttributeChangeFlags::Unknown, bDeferChangeNotifications);
+
+	return TargetMesh;
+}
+
+
+
 UDynamicMesh* UGeometryScriptLibrary_MeshBasicEditFunctions::AppendMeshRepeated(
 	UDynamicMesh* TargetMesh,
 	UDynamicMesh* AppendMesh,
diff --git a/Engine/Plugins/Experimental/GeometryScripting/Source/GeometryScriptingCore/Public/GeometryScript/MeshBasicEditFunctions.h b/Engine/Plugins/Experimental/GeometryScripting/Source/GeometryScriptingCore/Public/GeometryScript/MeshBasicEditFunctions.h
--- a/Engine/Plugins/Experimental/GeometryScripting/Source/GeometryScriptingCore/Public/GeometryScript/MeshBasicEditFunctions.h
+++ b/Engine/Plugins/Experimental/GeometryScripting/Source/GeometryScriptingCore/Public/GeometryScript/MeshBasicEditFunctions.h
@@ -153,6 +153,9 @@ public:
 		int& NumDeleted,
 		bool bDeferChangeNotifications = false );
 
+	/**
+	 * Apply AppendTransform to AppendMesh and then add it's geometry to the TargetMesh
+	 */
 	UFUNCTION(BlueprintCallable, Category = "GeometryScript|MeshEdits", meta=(ScriptMethod))
 	static UPARAM(DisplayName = "Target Mesh") UDynamicMesh* 
 	AppendMesh( 
@@ -162,6 +165,27 @@ public:
 		bool bDeferChangeNotifications = false,
 		UGeometryScriptDebug* Debug = nullptr);
 
+	/**
+	 * For each transform in AppendTransforms, apply the transform to AppendMesh and then add it's geometry to the TargetMesh.
+	 * @param ConstantTransform the Constant transform will be applied after each Append transform
+	 * @param bConstantTransformIsRelative if true, the Constant transform is applied "in the frame" of the Append Transform, otherwise it is applied as a second transform in local coordinates (ie rotate around the AppendTransform X axis, vs around the local X axis)
+	 */
+	UFUNCTION(BlueprintCallable, Category = "GeometryScript|MeshEdits", meta=(ScriptMethod))
+	static UPARAM(DisplayName = "Target Mesh") UDynamicMesh* 
+	AppendMeshTransformed( 
+		UDynamicMesh* TargetMesh, 
+		UDynamicMesh* AppendMesh, 
+		const TArray<FTransform>& AppendTransforms, 
+		FTransform ConstantTransform,
+		bool bConstantTransformIsRelative = true,
+		bool bDeferChangeNotifications = false,
+		UGeometryScriptDebug* Debug = nullptr);
+
+	/**
+	 * Repeatedly apply AppendTransform to the AppendMesh, each time adding the geometry to TargetMesh.
+	 * @param RepeatCount number of times to repeat the transform-append cycle
+	 * @param bApplyTransformToFirstInstance if true, the AppendTransform is applied before the first mesh append, otherwise it is applied after
+	 */
 	UFUNCTION(BlueprintCallable, Category = "GeometryScript|MeshEdits", meta=(ScriptMethod))
 	static UPARAM(DisplayName = "Target Mesh") UDynamicMesh* 
 	AppendMeshRepeated( 
