GeometryScript: add SetMeshSelectionVertexColor function, sets vertex color on selected faces/vertices, optionally creating a "crease" in the vertex color overlay

diff --git a/Engine/Plugins/Experimental/GeometryScripting/Source/GeometryScriptingCore/Private/MeshVertexColorFunctions.cpp b/Engine/Plugins/Experimental/GeometryScripting/Source/GeometryScriptingCore/Private/MeshVertexColorFunctions.cpp
--- a/Engine/Plugins/Experimental/GeometryScripting/Source/GeometryScriptingCore/Private/MeshVertexColorFunctions.cpp
+++ b/Engine/Plugins/Experimental/GeometryScripting/Source/GeometryScriptingCore/Private/MeshVertexColorFunctions.cpp
@@ -13,10 +13,10 @@ using namespace UE::Geometry;
 
 static FLinearColor CombineColors(FLinearColor ExistingColor, FLinearColor NewColor, FGeometryScriptColorFlags Flags)
 {
-	ExistingColor.R = (Flags.bRed) ? ExistingColor.R : NewColor.R;
-	ExistingColor.G = (Flags.bBlue) ? ExistingColor.G : NewColor.G;
-	ExistingColor.B = (Flags.bGreen) ? ExistingColor.B : NewColor.B;
-	ExistingColor.A = (Flags.bAlpha) ? ExistingColor.A : NewColor.A;
+	ExistingColor.R = (Flags.bRed) ? NewColor.R : ExistingColor.R;
+	ExistingColor.G = (Flags.bBlue) ? NewColor.G : ExistingColor.G;
+	ExistingColor.B = (Flags.bGreen) ? NewColor.B : ExistingColor.B;
+	ExistingColor.A = (Flags.bAlpha) ? NewColor.A : ExistingColor.A;
 	return ExistingColor;
 }
 
@@ -82,6 +82,111 @@ UDynamicMesh* UGeometryScriptLibrary_MeshVertexColorFunctions::SetMeshConstantVe
 
 
 
+
+UDynamicMesh* UGeometryScriptLibrary_MeshVertexColorFunctions::SetMeshSelectionVertexColor(
+	UDynamicMesh* TargetMesh,
+	FGeometryScriptMeshSelection Selection,
+	FLinearColor Color,
+	FGeometryScriptColorFlags Flags,
+	bool bCreateColorSeam,
+	UGeometryScriptDebug* Debug)
+{
+	if (TargetMesh == nullptr)
+	{
+		UE::Geometry::AppendError(Debug, EGeometryScriptErrorType::InvalidInputs, LOCTEXT("SetMeshSelectionVertexColor_InvalidInput", "SetMeshSelectionVertexColor: TargetMesh is Null"));
+		return TargetMesh;
+	}
+
+	TargetMesh->EditMesh([&](FDynamicMesh3& EditMesh) 
+	{
+		bool bCreated = false;
+		if (EditMesh.HasAttributes() == false)
+		{
+			EditMesh.EnableAttributes();
+			bCreated = true;
+		}
+		if (EditMesh.Attributes()->HasPrimaryColors() == false)
+		{
+			EditMesh.Attributes()->EnablePrimaryColors();
+			bCreated = true;
+		}
+		FDynamicMeshColorOverlay* Colors = EditMesh.Attributes()->PrimaryColors();
+
+		if (bCreateColorSeam)
+		{
+			FLinearColor NewColor = CombineColors(FLinearColor::Black, Color, Flags);
+
+			TArray<int32> Triangles;
+			Selection.ConvertToMeshIndexArray(EditMesh, Triangles, EGeometryScriptIndexType::Triangle);
+			TMap<int32, int32> VerticesToElements;
+			for (int32 TriangleID : Triangles)
+			{
+				FIndex3i TriVerts = EditMesh.GetTriangle(TriangleID);
+				for (int32 j = 0; j < 3; ++j)
+				{
+					const int32* FoundElemID = VerticesToElements.Find(TriVerts[j]);
+					if (FoundElemID == nullptr)
+					{
+						int32 NewElemID = Colors->AppendElement(FVector4f(NewColor.R, NewColor.G, NewColor.B, NewColor.A));
+						VerticesToElements.Add(TriVerts[j], NewElemID);
+					}
+				}
+			}
+
+			for (int32 TriangleID : Triangles)
+			{
+				FIndex3i TriVerts = EditMesh.GetTriangle(TriangleID);
+				TriVerts[0] = VerticesToElements[TriVerts[0]];
+				TriVerts[1] = VerticesToElements[TriVerts[1]];
+				TriVerts[2] = VerticesToElements[TriVerts[2]];
+				Colors->SetTriangle(TriangleID, TriVerts);
+			}
+		}
+		else if (Selection.GetSelectionType() == EGeometryScriptMeshSelectionType::Vertices)
+		{
+			Selection.ProcessByVertexID(EditMesh, [&](int32 VertexID)
+			{
+				EditMesh.EnumerateVertexTriangles(VertexID, [&](int32 TriangleID)
+				{
+					if (Colors->IsSetTriangle(TriangleID))
+					{
+						FIndex3i TriElems = Colors->GetTriangle(TriangleID);
+						for (int32 j = 0; j < 3; ++j)
+						{
+							if (Colors->GetParentVertex(TriElems[j]) == VertexID)
+							{
+								FLinearColor Existing = ToLinearColor(Colors->GetElement(TriElems[j]));
+								FLinearColor NewColor = CombineColors(Existing, Color, Flags);
+								Colors->SetElement(TriElems[j], ToVector4<float>(NewColor));
+							}
+						}
+					}
+				});
+			});
+		}
+		else
+		{
+			Selection.ProcessByTriangleID(EditMesh, [&](int32 TriangleID)
+			{
+				if (Colors->IsSetTriangle(TriangleID))
+				{
+					FIndex3i TriElems = Colors->GetTriangle(TriangleID);
+					for (int32 j = 0; j < 3; ++j)
+					{
+						FLinearColor Existing = ToLinearColor(Colors->GetElement(TriElems[j]));
+						FLinearColor NewColor = CombineColors(Existing, Color, Flags);
+						Colors->SetElement(TriElems[j], ToVector4<float>(NewColor));
+					}
+				}
+			});
+		}
+
+	}, EDynamicMeshChangeType::GeneralEdit, EDynamicMeshAttributeChangeFlags::Unknown, false);
+
+	return TargetMesh;
+}
+
+
 UDynamicMesh* UGeometryScriptLibrary_MeshVertexColorFunctions::SetMeshPerVertexColors(
 	UDynamicMesh* TargetMesh,
 	FGeometryScriptColorList VertexColorList,
diff --git a/Engine/Plugins/Experimental/GeometryScripting/Source/GeometryScriptingCore/Public/GeometryScript/MeshVertexColorFunctions.h b/Engine/Plugins/Experimental/GeometryScripting/Source/GeometryScriptingCore/Public/GeometryScript/MeshVertexColorFunctions.h
--- a/Engine/Plugins/Experimental/GeometryScripting/Source/GeometryScriptingCore/Public/GeometryScript/MeshVertexColorFunctions.h
+++ b/Engine/Plugins/Experimental/GeometryScripting/Source/GeometryScriptingCore/Public/GeometryScript/MeshVertexColorFunctions.h
@@ -5,6 +5,7 @@
 #include "CoreMinimal.h"
 #include "Kismet/BlueprintFunctionLibrary.h"
 #include "GeometryScript/GeometryScriptTypes.h"
+#include "GeometryScript/GeometryScriptSelectionTypes.h"
 #include "MeshVertexColorFunctions.generated.h"
 
 class UDynamicMesh;
@@ -31,6 +32,26 @@ public:
 		bool bClearExisting = false,
 		UGeometryScriptDebug* Debug = nullptr);
 
+
+	/**
+	 * Set the colors in the TargetMesh VertexColor Overlay identified by the Selection to a constant value.
+	 * For a Vertex Selection, each existing VertexColor Overlay Element for the vertex is updated.
+	 * For a Triangle or Polygroup Selection, all Overlay Elements in the identified Triangles are updated.
+	 * @param Color the constant color to set
+	 * @param Flags specify which RGBA channels to set (default all channels)
+	 * @param bCreateColorSeam if true, a "hard edge" in the vertex colors is created, by creating new Elements for all the triangles in the selection. If enabled, Vertex selections are converted to Triangle selections, and Flags is ignored. 
+	 */
+	UFUNCTION(BlueprintCallable, Category = "GeometryScript|VertexColor", meta=(ScriptMethod))
+	static UPARAM(DisplayName = "Target Mesh") UDynamicMesh*
+	SetMeshSelectionVertexColor(
+		UDynamicMesh* TargetMesh,
+		FGeometryScriptMeshSelection Selection,
+		FLinearColor Color,
+		FGeometryScriptColorFlags Flags,
+		bool bCreateColorSeam = false,
+		UGeometryScriptDebug* Debug = nullptr);
+
+
 	/**
 	 * Set all vertex colors in the TargetMesh VertexColor Overlay to the specified per-vertex colors
 	 * @param VertexColorList per-vertex colors. Size must be less than or equal to the MaxVertexID of TargetMesh  (ie gaps are supported)
