GeometryScript: add SampleSplineToTransforms function, produces full FTransforms based on sampling a USplineComponent, instead of just positions.

diff --git a/Engine/Plugins/Experimental/GeometryScripting/Source/GeometryScriptingCore/Private/PolyPathFunctions.cpp b/Engine/Plugins/Experimental/GeometryScripting/Source/GeometryScriptingCore/Private/PolyPathFunctions.cpp
--- a/Engine/Plugins/Experimental/GeometryScripting/Source/GeometryScriptingCore/Private/PolyPathFunctions.cpp
+++ b/Engine/Plugins/Experimental/GeometryScripting/Source/GeometryScriptingCore/Private/PolyPathFunctions.cpp
@@ -224,6 +224,52 @@ void UGeometryScriptLibrary_PolyPathFunctions::ConvertSplineToPolyPath(const USp
 	}
 }
 
+
+void UGeometryScriptLibrary_PolyPathFunctions::SampleSplineToTransforms(
+	const USplineComponent* Spline, 
+	TArray<FTransform>& Frames, 
+	TArray<double>& FrameTimes,
+	FGeometryScriptSplineSamplingOptions SamplingOptions,
+	FTransform RelativeTransform,
+	bool bIncludeScale)
+{
+	Frames.Reset();
+	FrameTimes.Reset();
+
+	// Currently ErrorTolerance sampling can only be done via Spline->ConvertSplineToPolyLine, which only returns a list of points.
+	// To convert to Transforms we would have to reverse-engineer the Time at each Point which could be very expensive...
+	if (SamplingOptions.SampleSpacing == EGeometryScriptSampleSpacing::ErrorTolerance)
+	{
+		UE_LOG(LogGeometry, Warning, TEXT("SampleSplineToTransforms: ErrorTolerance sampling mode is currently not supported, falling back to UniformDistance"));
+		SamplingOptions.SampleSpacing = EGeometryScriptSampleSpacing::UniformDistance;
+	}
+
+	if (Spline != nullptr )
+	{
+		bool bIsLoop = Spline->IsClosedLoop();
+
+		float Duration = Spline->Duration;
+
+		bool bUseConstantVelocity = SamplingOptions.SampleSpacing == EGeometryScriptSampleSpacing::UniformDistance;
+		int32 UseSamples = FMath::Max(2, SamplingOptions.NumSamples); // Always use at least 2 samples
+																		// In non-loops, we adjust DivNum so we exactly sample the end of the spline
+																		// In loops we don't sample the endpoint, by convention, as it's the same as the start
+		float DivNum = float(UseSamples - (int32)!bIsLoop);
+		Frames.Reserve(UseSamples);
+		FrameTimes.Reserve(UseSamples);
+		for (int32 Idx = 0; Idx < UseSamples; Idx++)
+		{
+			float Time = Duration * ((float)Idx / DivNum);
+			FTransform Transform = Spline->GetTransformAtTime(Time, SamplingOptions.CoordinateSpace, bUseConstantVelocity, bIncludeScale);
+			FTransform::Multiply(&Transform, &RelativeTransform, &Transform);
+			Frames.Add(Transform);
+			FrameTimes.Add(Time);
+		}
+	}
+}
+
+
+
 TArray<FVector> UGeometryScriptLibrary_PolyPathFunctions::Conv_GeometryScriptPolyPathToArray(FGeometryScriptPolyPath PolyPath)
 {
 	TArray<FVector> PathVertices;
diff --git a/Engine/Plugins/Experimental/GeometryScripting/Source/GeometryScriptingCore/Public/GeometryScript/PolyPathFunctions.h b/Engine/Plugins/Experimental/GeometryScripting/Source/GeometryScriptingCore/Public/GeometryScript/PolyPathFunctions.h
--- a/Engine/Plugins/Experimental/GeometryScripting/Source/GeometryScriptingCore/Public/GeometryScript/PolyPathFunctions.h
+++ b/Engine/Plugins/Experimental/GeometryScripting/Source/GeometryScriptingCore/Public/GeometryScript/PolyPathFunctions.h
@@ -94,9 +94,29 @@ public:
 	UFUNCTION(BlueprintCallable, Category = "GeometryScript|PolyPath")
 	static UPARAM(DisplayName = "Poly Path") FGeometryScriptPolyPath CreateArcPath2D(FVector2D Center = FVector2D(0, 0), float Radius = 10, int NumPoints = 10, float StartAngle = 0, float EndAngle = 90);
 
+	/**
+	 * Sample positions from a USplineComponent into a PolyPath, based on the given SamplingOptions
+	 */
 	UFUNCTION(BlueprintCallable, Category = "GeometryScript|PolyPath")
 	static void ConvertSplineToPolyPath(const USplineComponent* Spline, FGeometryScriptPolyPath& PolyPath, FGeometryScriptSplineSamplingOptions SamplingOptions);
 
+	/**
+	 * Sample a USplineComponent into a list of FTransforms, based on the given SamplingOptions.
+	 * @param Frames Transforms are returned here, with X axis oriented along spline Tangent and Z as the 'Up' vector.
+	 * @param FrameTimes the spline Time value used for each Frame
+	 * @param RelativeTransform a constant Transform applied to each sample Transform in its local frame of reference. So, eg, an X Rotation will rotate each frame around the local spline Tangent vector
+	 * @param bIncludeScale if true, the Scale of each FTransform is taken from the Spline, otherwise the Transforms have unit scale
+	 */
+	UFUNCTION(BlueprintCallable, Category = "GeometryScript|PolyPath")
+	static void SampleSplineToTransforms(
+		const USplineComponent* Spline, 
+		TArray<FTransform>& Frames, 
+		TArray<double>& FrameTimes,
+		FGeometryScriptSplineSamplingOptions SamplingOptions,
+		FTransform RelativeTransform,
+		bool bIncludeScale = true);
+
+
 	UFUNCTION(BlueprintCallable, Category = "GeometryScript|PolyPath", meta=(ScriptMethod))
 	static void ConvertPolyPathToArray(FGeometryScriptPolyPath PolyPath, TArray<FVector>& VertexArray);
 
