GeometryScript: add ResolveMeshTJunctions function

diff --git a/Engine/Plugins/Experimental/GeometryScripting/Source/GeometryScriptingCore/Private/MeshRepairFunctions.cpp b/Engine/Plugins/Experimental/GeometryScripting/Source/GeometryScriptingCore/Private/MeshRepairFunctions.cpp
--- a/Engine/Plugins/Experimental/GeometryScripting/Source/GeometryScriptingCore/Private/MeshRepairFunctions.cpp
+++ b/Engine/Plugins/Experimental/GeometryScripting/Source/GeometryScriptingCore/Private/MeshRepairFunctions.cpp
@@ -10,6 +10,7 @@
 #include "Polygroups/PolygroupSet.h"
 #include "MeshBoundaryLoops.h"
 #include "DynamicMeshEditor.h"
+#include "Operations/MeshResolveTJunctions.h"
 #include "Operations/RemoveOccludedTriangles.h"
 #include "Selections/MeshConnectedComponents.h"
 #include "Selections/MeshFaceSelection.h"
@@ -42,6 +43,33 @@ UDynamicMesh* UGeometryScriptLibrary_MeshRepairFunctions::CompactMesh(
 }
 
 
+UDynamicMesh* UGeometryScriptLibrary_MeshRepairFunctions::ResolveMeshTJunctions(
+	UDynamicMesh* TargetMesh,
+	FGeometryScriptResolveTJunctionOptions ResolveOptions,
+	UGeometryScriptDebug* Debug)
+{
+	if (TargetMesh == nullptr)
+	{
+		UE::Geometry::AppendError(Debug, EGeometryScriptErrorType::InvalidInputs, LOCTEXT("ResolveMeshTJunctions_InvalidInput", "ResolveMeshTJunctions: TargetMesh is Null"));
+		return TargetMesh;
+	}
+
+	TargetMesh->EditMesh([&](FDynamicMesh3& EditMesh) 
+	{
+		FMeshResolveTJunctions Resolver(&EditMesh);
+		Resolver.DistanceTolerance = 2 * ResolveOptions.Tolerance;
+		bool bResolveOK = Resolver.Apply();
+		if (!bResolveOK)
+		{
+			UE::Geometry::AppendError(Debug, EGeometryScriptErrorType::OperationFailed, LOCTEXT("ResolveMeshTJunctions_Error", "ResolveMeshTJunctions: ResolveTJunctions Operation returned error flag"));
+		}
+
+	}, EDynamicMeshChangeType::GeneralEdit, EDynamicMeshAttributeChangeFlags::Unknown, false);
+
+	return TargetMesh;
+}
+
+
 
 UDynamicMesh* UGeometryScriptLibrary_MeshRepairFunctions::WeldMeshEdges(
 	UDynamicMesh* TargetMesh,
diff --git a/Engine/Plugins/Experimental/GeometryScripting/Source/GeometryScriptingCore/Public/GeometryScript/MeshRepairFunctions.h b/Engine/Plugins/Experimental/GeometryScripting/Source/GeometryScriptingCore/Public/GeometryScript/MeshRepairFunctions.h
--- a/Engine/Plugins/Experimental/GeometryScripting/Source/GeometryScriptingCore/Public/GeometryScript/MeshRepairFunctions.h
+++ b/Engine/Plugins/Experimental/GeometryScripting/Source/GeometryScriptingCore/Public/GeometryScript/MeshRepairFunctions.h
@@ -24,6 +24,18 @@ public:
 };
 
 
+
+USTRUCT(BlueprintType)
+struct GEOMETRYSCRIPTINGCORE_API FGeometryScriptResolveTJunctionOptions
+{
+	GENERATED_BODY()
+public:
+	/**  */
+	UPROPERTY(EditAnywhere, BlueprintReadWrite, Category = Options)
+	float Tolerance = 1e-03f;
+};
+
+
 UENUM(BlueprintType)
 enum class EGeometryScriptFillHolesMethod : uint8
 {
@@ -126,6 +138,13 @@ public:
 		UDynamicMesh* TargetMesh, 
 		UGeometryScriptDebug* Debug = nullptr);
 
+	UFUNCTION(BlueprintCallable, Category = "GeometryScript|Repair", meta=(ScriptMethod))
+	static UPARAM(DisplayName = "Target Mesh") UDynamicMesh* 
+	ResolveMeshTJunctions(  
+		UDynamicMesh* TargetMesh, 
+		FGeometryScriptResolveTJunctionOptions ResolveOptions,
+		UGeometryScriptDebug* Debug = nullptr);
+
 	UFUNCTION(BlueprintCallable, Category = "GeometryScript|Repair", meta=(ScriptMethod))
 	static UPARAM(DisplayName = "Target Mesh") UDynamicMesh* 
 	WeldMeshEdges(  
