GeometryScript: improve runtime support and error handling in CopyMeshFromStaticMesh

diff --git a/Engine/Plugins/Experimental/GeometryScripting/Source/GeometryScriptingCore/Private/MeshAssetFunctions.cpp b/Engine/Plugins/Experimental/GeometryScripting/Source/GeometryScriptingCore/Private/MeshAssetFunctions.cpp
--- a/Engine/Plugins/Experimental/GeometryScripting/Source/GeometryScriptingCore/Private/MeshAssetFunctions.cpp
+++ b/Engine/Plugins/Experimental/GeometryScripting/Source/GeometryScriptingCore/Private/MeshAssetFunctions.cpp
@@ -41,7 +41,6 @@ static UDynamicMesh* CopyMeshFromStaticMesh_SourceData(
 	FGeometryScriptMeshReadLOD RequestedLOD,
 	TEnumAsByte<EGeometryScriptOutcomePins>& Outcome,
 	UGeometryScriptDebug* Debug
-
 )
 {
 	if (RequestedLOD.LODType != EGeometryScriptLODType::MaxAvailable && RequestedLOD.LODType != EGeometryScriptLODType::SourceModel)
@@ -119,19 +118,34 @@ static UDynamicMesh* CopyMeshFromStaticMesh_RenderData(
 	FGeometryScriptMeshReadLOD RequestedLOD,
 	TEnumAsByte<EGeometryScriptOutcomePins>& Outcome,
 	UGeometryScriptDebug* Debug
-
 )
 {
-	if (RequestedLOD.LODType != EGeometryScriptLODType::RenderData)
+	if (RequestedLOD.LODType != EGeometryScriptLODType::MaxAvailable && RequestedLOD.LODType != EGeometryScriptLODType::RenderData)
 	{
 		UE::Geometry::AppendError(Debug, EGeometryScriptErrorType::InvalidInputs, LOCTEXT("CopyMeshFromStaticMesh_LODNotAvailable", "CopyMeshFromStaticMesh: Requested LOD Type is not available"));
 		return ToDynamicMesh;
 	}
 
-	int32 UseLODIndex = FMath::Clamp(RequestedLOD.LODIndex, 0, FromStaticMeshAsset->GetNumLODs() - 1);
+#if !WITH_EDITOR
+	if (FromStaticMeshAsset->bAllowCPUAccess == false)
+	{
+		UE::Geometry::AppendError(Debug, EGeometryScriptErrorType::InvalidInputs, LOCTEXT("CopyMeshFromStaticMesh_CPUAccess", "CopyMeshFromStaticMesh: StaticMesh bAllowCPUAccess must be set to true to read mesh data at Runtime"));
+		return ToDynamicMesh;
+	}
+#endif
 
+	int32 UseLODIndex = FMath::Clamp(RequestedLOD.LODIndex, 0, FromStaticMeshAsset->GetNumLODs() - 1);
 
-	const FStaticMeshLODResources* LODResources = &FromStaticMeshAsset->GetRenderData()->LODResources[UseLODIndex];
+	const FStaticMeshLODResources* LODResources = nullptr;
+	if (FStaticMeshRenderData* RenderData = FromStaticMeshAsset->GetRenderData())
+	{
+		LODResources = &RenderData->LODResources[UseLODIndex];
+	}
+	if (LODResources == nullptr)
+	{
+		UE::Geometry::AppendError(Debug, EGeometryScriptErrorType::InvalidInputs, LOCTEXT("CopyMeshFromStaticMesh_NoLODResources", "CopyMeshFromStaticMesh: LOD Data is not available"));
+		return ToDynamicMesh;
+	}
 
 	FStaticMeshLODResourcesToDynamicMesh::ConversionOptions ConvertOptions;
 #if WITH_EDITOR
