GeometryScript: add support for input selection to ComputeMeshConvexHull. Add ComputeMeshConvexDecomposition that exposes new FConvexDecomposition3 algorithm, via FMeshSimpleShapeApproximation

diff --git a/Engine/Plugins/Experimental/GeometryScripting/Source/GeometryScriptingCore/Private/ContainmentFunctions.cpp b/Engine/Plugins/Experimental/GeometryScripting/Source/GeometryScriptingCore/Private/ContainmentFunctions.cpp
--- a/Engine/Plugins/Experimental/GeometryScripting/Source/GeometryScriptingCore/Private/ContainmentFunctions.cpp
+++ b/Engine/Plugins/Experimental/GeometryScripting/Source/GeometryScriptingCore/Private/ContainmentFunctions.cpp
@@ -4,8 +4,10 @@
 
 #include "DynamicMesh/DynamicMesh3.h"
 #include "DynamicMesh/MeshNormals.h"
+#include "DynamicMeshEditor.h"
 #include "Operations/MeshConvexHull.h"
 #include "Operations/MeshProjectionHull.h"
+#include "ShapeApproximation/MeshSimpleShapeApproximation.h"
 #include "UDynamicMesh.h"
 
 using namespace UE::Geometry;
@@ -17,6 +19,7 @@ UDynamicMesh* UGeometryScriptLibrary_ContainmentFunctions::ComputeMeshConvexHull
 	UDynamicMesh* TargetMesh,
 	UPARAM(DisplayName = "Hull Mesh", ref) UDynamicMesh* CopyToMesh, 
 	UPARAM(DisplayName = "Hull Mesh") UDynamicMesh*& CopyToMeshOut, 
+	FGeometryScriptMeshSelection Selection,
 	FGeometryScriptConvexHullOptions Options,
 	UGeometryScriptDebug* Debug)
 {
@@ -37,7 +40,13 @@ UDynamicMesh* UGeometryScriptLibrary_ContainmentFunctions::ComputeMeshConvexHull
 	{
 		FMeshConvexHull Hull(&EditMesh);
 
-		if (Options.bPrefilterVertices)
+		if (Selection.IsEmpty() == false)
+		{
+			Selection.ProcessByVertexID(EditMesh, [&](int32 VertexID) {
+				Hull.VertexSet.Add(VertexID);
+			});
+		}
+		else if (Options.bPrefilterVertices)
 		{
 			FMeshConvexHull::GridSample(EditMesh, FMath::Max(32, Options.PrefilterGridResolution), Hull.VertexSet);
 		}
@@ -120,4 +129,76 @@ UDynamicMesh* UGeometryScriptLibrary_ContainmentFunctions::ComputeMeshSweptHull(
 }
 
 
+
+
+
+
+
+
+UDynamicMesh* UGeometryScriptLibrary_ContainmentFunctions::ComputeMeshConvexDecomposition(
+	UDynamicMesh* TargetMesh,
+	UPARAM(DisplayName = "Hull Mesh", ref) UDynamicMesh* CopyToMesh, 
+	UPARAM(DisplayName = "Hull Mesh") UDynamicMesh*& CopyToMeshOut, 
+	FGeometryScriptConvexDecompositionOptions Options,
+	UGeometryScriptDebug* Debug)
+{
+	if (TargetMesh == nullptr)
+	{
+		UE::Geometry::AppendError(Debug, EGeometryScriptErrorType::InvalidInputs, LOCTEXT("ComputeMeshConvexDecomposition_InvalidInput", "ComputeMeshConvexDecomposition: TargetMesh is Null"));
+		return TargetMesh;
+	}
+	if (CopyToMesh == nullptr)
+	{
+		UE::Geometry::AppendError(Debug, EGeometryScriptErrorType::InvalidInputs, LOCTEXT("ComputeMeshConvexDecomposition_InvalidInput2", "ComputeMeshConvexDecomposition: CopyToMesh is Null"));
+		return TargetMesh;
+	}
+	CopyToMeshOut = CopyToMesh;
+
+	FDynamicMesh3 HullsMesh;
+	TargetMesh->ProcessMesh([&](const FDynamicMesh3& EditMesh) 
+	{
+		TArray<const FDynamicMesh3*> InputMeshes;
+		InputMeshes.Add(&EditMesh);
+		FMeshSimpleShapeApproximation Approximator;
+		Approximator.InitializeSourceMeshes(InputMeshes);
+
+		Approximator.bSimplifyHulls = (Options.SimplifyToFaceCount > 4);
+		Approximator.HullTargetFaceCount = Options.SimplifyToFaceCount;
+
+		Approximator.ConvexDecompositionMaxPieces = Options.NumHulls;
+		Approximator.ConvexDecompositionSearchFactor = Options.SearchFactor;
+		Approximator.ConvexDecompositionErrorTolerance = Options.ErrorTolerance;
+		Approximator.ConvexDecompositionMinPartThickness = Options.MinPartThickness;
+
+		FSimpleShapeSet3d Shapes;
+		Approximator.Generate_ConvexHullDecompositions(Shapes);
+
+		HullsMesh.EnableMatchingAttributes(EditMesh);
+		FDynamicMeshEditor Editor(&HullsMesh);
+		for (FConvexShape3d& Convex : Shapes.Convexes)
+		{
+			FMeshIndexMappings Mappings;
+			Editor.AppendMesh(&Convex.Mesh, Mappings);
+		}
+	});
+
+	if ( HullsMesh.TriangleCount() == 0 )
+	{
+		// todo: replace output with bounding box mesh?
+		UE::Geometry::AppendError(Debug, EGeometryScriptErrorType::OperationFailed, LOCTEXT("ComputeMeshConvexDecomposition_Failed", "ComputeMeshConvexDecomposition: Hull Computation Failed"));
+		CopyToMesh->ResetToCube();
+	}
+	else
+	{
+		FMeshNormals::InitializeOverlayToPerTriangleNormals(HullsMesh.Attributes()->PrimaryNormals());
+		CopyToMesh->SetMesh(MoveTemp(HullsMesh));
+	}
+
+	return TargetMesh;
+}
+
+
+
+
+
 #undef LOCTEXT_NAMESPACE
\ No newline at end of file
diff --git a/Engine/Plugins/Experimental/GeometryScripting/Source/GeometryScriptingCore/Public/GeometryScript/ContainmentFunctions.h b/Engine/Plugins/Experimental/GeometryScripting/Source/GeometryScriptingCore/Public/GeometryScript/ContainmentFunctions.h
--- a/Engine/Plugins/Experimental/GeometryScripting/Source/GeometryScriptingCore/Public/GeometryScript/ContainmentFunctions.h
+++ b/Engine/Plugins/Experimental/GeometryScripting/Source/GeometryScriptingCore/Public/GeometryScript/ContainmentFunctions.h
@@ -5,6 +5,7 @@
 #include "CoreMinimal.h"
 #include "Kismet/BlueprintFunctionLibrary.h"
 #include "GeometryScript/GeometryScriptTypes.h"
+#include "GeometryScript/GeometryScriptSelectionTypes.h"
 #include "ContainmentFunctions.generated.h"
 
 class UDynamicMesh;
@@ -20,7 +21,7 @@ public:
 	UPROPERTY(EditAnywhere, BlueprintReadWrite, Category = Options)
 	int PrefilterGridResolution = 128;
 
-	/** */
+	/** Try to simplify each convex hull to this triangle count. If 0, no simplification */
 	UPROPERTY(EditAnywhere, BlueprintReadWrite, Category = Options)
 	int SimplifyToFaceCount = 0;
 };
@@ -54,6 +55,33 @@ public:
 
 
 
+USTRUCT(BlueprintType)
+struct GEOMETRYSCRIPTINGCORE_API FGeometryScriptConvexDecompositionOptions
+{
+	GENERATED_BODY()
+public:
+	/** How many convex pieces to target per mesh when creating convex decompositions.  If ErrorTolerance is set, can create fewer pieces. */
+	UPROPERTY(EditAnywhere, BlueprintReadWrite, Category = DecompositionOptions)
+	int NumHulls = 2;
+
+	/** How much additional decomposition decomposition + merging to do, as a fraction of max pieces.  Larger values can help better-cover small features, while smaller values create a cleaner decomposition with less overlap between hulls. */
+	UPROPERTY(EditAnywhere, BlueprintReadWrite, Category = DecompositionOptions)
+	double SearchFactor = .5;
+
+	/** Error tolerance to guide convex decomposition (in cm); we stop adding new parts if the volume error is below the threshold.  For volumetric errors, value will be cubed. */
+	UPROPERTY(EditAnywhere, BlueprintReadWrite, Category = DecompositionOptions)
+	double ErrorTolerance = 0;
+
+	/** Minimum part thickness for convex decomposition (in cm); hulls thinner than this will be merged into adjacent hulls, if possible. */
+	UPROPERTY(EditAnywhere, BlueprintReadWrite, Category = DecompositionOptions)
+	double MinPartThickness = .1;
+
+	/** Try to simplify each convex hull to this triangle count. If 0, no simplification */
+	UPROPERTY(EditAnywhere, BlueprintReadWrite, Category = ConvexHullOptions)
+	int SimplifyToFaceCount = 0;
+};
+
+
 
 
 UCLASS(meta = (ScriptName = "GeometryScript_Containment"))
@@ -62,15 +90,24 @@ class GEOMETRYSCRIPTINGCORE_API UGeometryScriptLibrary_ContainmentFunctions : pu
 	GENERATED_BODY()
 public:
 
+	/**
+	 * Compute the Convex Hull of a given Mesh, or part of the mesh if an optional Selection is provided
+	 * @param Selection selection of mesh faces/vertices to contain in the convex hull. If not provided, entire mesh is used.
+	 */
 	UFUNCTION(BlueprintCallable, Category = "GeometryScript|Containment", meta=(ScriptMethod))
 	static UPARAM(DisplayName = "Target Mesh") UDynamicMesh* 
 	ComputeMeshConvexHull(  
 		UDynamicMesh* TargetMesh, 
 		UPARAM(DisplayName = "Hull Mesh", ref) UDynamicMesh* CopyToMesh, 
 		UPARAM(DisplayName = "Hull Mesh") UDynamicMesh*& CopyToMeshOut, 
+		FGeometryScriptMeshSelection Selection,
 		FGeometryScriptConvexHullOptions Options,
 		UGeometryScriptDebug* Debug = nullptr);
 
+	/**
+	* Compute the Swept Hull of a given Mesh for a given 3D Plane defined by ProjectionFrame.
+	* The Swept Hull is a linear sweep of the 2D convex hull of the mesh vertices projected onto the plane (the sweep precisely contains the mesh extents along the plane normal)
+	*/
 	UFUNCTION(BlueprintCallable, Category = "GeometryScript|Containment", meta=(ScriptMethod))
 	static UPARAM(DisplayName = "Target Mesh") UDynamicMesh* 
 	ComputeMeshSweptHull(  
@@ -81,4 +118,19 @@ public:
 		FGeometryScriptSweptHullOptions Options,
 		UGeometryScriptDebug* Debug = nullptr);
 
+	/**
+	 * Compute a Convex Hull Decomposition of the given TargetMesh. Assuming more than one hull is requested,
+	 * multiple hulls will be returned that attempt to approximate the mesh. There is no guarantee that the entire
+	 * mesh is contained in the hulls.
+	 * @warning this function can be quite expensive, and the results are expected to change in the future as the Convex Decomposition algorithm is improved
+	 */
+	UFUNCTION(BlueprintCallable, Category = "GeometryScript|Containment|Experimental", meta=(ScriptMethod))
+	static UPARAM(DisplayName = "Target Mesh") UDynamicMesh* 
+	ComputeMeshConvexDecomposition(  
+		UDynamicMesh* TargetMesh, 
+		UPARAM(DisplayName = "Hull Mesh", ref) UDynamicMesh* CopyToMesh, 
+		UPARAM(DisplayName = "Hull Mesh") UDynamicMesh*& CopyToMeshOut, 
+		FGeometryScriptConvexDecompositionOptions Options,
+		UGeometryScriptDebug* Debug = nullptr);
+
 };
\ No newline at end of file
diff --git a/Engine/Plugins/Experimental/GeometryScripting/Source/GeometryScriptingCore/Public/GeometryScript/MeshRemeshFunctions.h b/Engine/Plugins/Experimental/GeometryScripting/Source/GeometryScriptingCore/Public/GeometryScript/MeshRemeshFunctions.h
--- a/Engine/Plugins/Experimental/GeometryScripting/Source/GeometryScriptingCore/Public/GeometryScript/MeshRemeshFunctions.h
+++ b/Engine/Plugins/Experimental/GeometryScripting/Source/GeometryScriptingCore/Public/GeometryScript/MeshRemeshFunctions.h
@@ -145,6 +145,7 @@ public:
 
 	/**
 	 * Apply Uniform Remeshing to the TargetMesh. 
+	 * @warning this function can be quite expensive. The results may be non-deterministic, and are expected to change in future versions.
 	 */
 	UFUNCTION(BlueprintCallable, Category = "GeometryScript|Simplification", meta=(ScriptMethod))
 	static UPARAM(DisplayName = "Target Mesh") UDynamicMesh* 
diff --git a/Engine/Plugins/Runtime/GeometryProcessing/Source/DynamicMesh/Private/ShapeApproximation/MeshSimpleShapeApproximation.cpp b/Engine/Plugins/Runtime/GeometryProcessing/Source/DynamicMesh/Private/ShapeApproximation/MeshSimpleShapeApproximation.cpp
--- a/Engine/Plugins/Runtime/GeometryProcessing/Source/DynamicMesh/Private/ShapeApproximation/MeshSimpleShapeApproximation.cpp
+++ b/Engine/Plugins/Runtime/GeometryProcessing/Source/DynamicMesh/Private/ShapeApproximation/MeshSimpleShapeApproximation.cpp
@@ -331,7 +331,7 @@ void FMeshSimpleShapeApproximation::Generate_ConvexHullDecompositions(FSimpleSha
 		for (int32 HullIdx = 0; HullIdx < Decomposition.NumHulls(); HullIdx++)
 		{
 			FDynamicMesh3 HullMesh = Decomposition.GetHullMesh(HullIdx);
-			if (!FMeshConvexHull::SimplifyHull(HullMesh, HullTargetFaceCount, Progress))
+			if (bSimplifyHulls && FMeshConvexHull::SimplifyHull(HullMesh, HullTargetFaceCount, Progress) == false)
 			{
 				return;
 			}
