GeometryScript:
Add optional Selection to ApplyPerlinNoiseToMesh, ApplyIterativeSmoothingToMesh, ApplyDisplaceFromTextureMap, ApplyDisplaceFromPerVertexVectors.
Add new functions ApplyMeshDisconnectFaces and ApplyMeshDuplicateFaces
Add new ApplyMeshLinearExtrudeFaces function, supports linear extrusion of entire mesh or selection. Redirect previous ApplyMeshExtrude function to ApplyMeshExtrude_Compatibility_5p0
Add new ApplyMeshBevelSelection function, supports various ways to interpret mesh selection as set of faces or edges to bevel

diff --git a/Engine/Plugins/Experimental/GeometryScripting/Source/GeometryScriptingCore/Private/GeometryScriptingCoreModule.cpp b/Engine/Plugins/Experimental/GeometryScripting/Source/GeometryScriptingCore/Private/GeometryScriptingCoreModule.cpp
--- a/Engine/Plugins/Experimental/GeometryScripting/Source/GeometryScriptingCore/Private/GeometryScriptingCoreModule.cpp
+++ b/Engine/Plugins/Experimental/GeometryScripting/Source/GeometryScriptingCore/Private/GeometryScriptingCoreModule.cpp
@@ -13,7 +13,11 @@ void FGeometryScriptingCoreModule::StartupModule()
 	// UE5.0 - AppendRectangle and AppendRoundRectangle shipped in Preview 1 interpreted dimensions incorrectly
 	Redirects.Emplace(ECoreRedirectFlags::Type_Function, TEXT("AppendRectangle"), TEXT("AppendRectangle_Compatibility_5_0"));
 	Redirects.Emplace(ECoreRedirectFlags::Type_Function, TEXT("AppendRoundRectangle"), TEXT("AppendRoundRectangle_Compatibility_5_0"));
-	
+
+	// UE5.1 - added selection argument to various functions
+	Redirects.Emplace(ECoreRedirectFlags::Type_Function, TEXT("ApplyMeshExtrude"), TEXT("ApplyMeshExtrude_Compatibility_5p0"));
+
+
 	FCoreRedirects::AddRedirectList(Redirects, TEXT("GeometryScriptingCore"));
 }
 
diff --git a/Engine/Plugins/Experimental/GeometryScripting/Source/GeometryScriptingCore/Private/MeshDeformFunctions.cpp b/Engine/Plugins/Experimental/GeometryScripting/Source/GeometryScriptingCore/Private/MeshDeformFunctions.cpp
--- a/Engine/Plugins/Experimental/GeometryScripting/Source/GeometryScriptingCore/Private/MeshDeformFunctions.cpp
+++ b/Engine/Plugins/Experimental/GeometryScripting/Source/GeometryScriptingCore/Private/MeshDeformFunctions.cpp
@@ -223,6 +223,7 @@ UDynamicMesh* UGeometryScriptLibrary_MeshDeformFunctions::ApplyMathWarpToMesh(
 
 UDynamicMesh* UGeometryScriptLibrary_MeshDeformFunctions::ApplyPerlinNoiseToMesh(
 	UDynamicMesh* TargetMesh,
+	FGeometryScriptMeshSelection Selection,
 	FGeometryScriptPerlinNoiseOptions Options,
 	UGeometryScriptDebug* Debug)
 {
@@ -231,6 +232,10 @@ UDynamicMesh* UGeometryScriptLibrary_MeshDeformFunctions::ApplyPerlinNoiseToMesh
 		UE::Geometry::AppendError(Debug, EGeometryScriptErrorType::InvalidInputs, LOCTEXT("ApplyPerlinNoiseToMesh_InvalidInput", "ApplyPerlinNoiseToMesh: TargetMesh is Null"));
 		return TargetMesh;
 	}
+	if (Selection.IsEmpty() && Options.EmptyBehavior != EGeometryScriptEmptySelectionBehavior::FullMeshSelection )
+	{
+		return TargetMesh;
+	}
 
 	TargetMesh->EditMesh([&](FDynamicMesh3& EditMesh) 
 	{
@@ -250,30 +255,45 @@ UDynamicMesh* UGeometryScriptLibrary_MeshDeformFunctions::ApplyPerlinNoiseToMesh
 			Normals.ComputeVertexNormals();
 		}
 
-		ParallelFor(EditMesh.MaxVertexID(), [&](int32 vid)
+		auto GetDisplacedPosition = [&EditMesh, &Options, &Offsets, &Normals, bAlignWithNormal](int32 VertexID)
 		{
-			if (EditMesh.IsVertex(vid))
+			FVector3d Pos = EditMesh.GetVertex(VertexID);
+			if (bAlignWithNormal)
 			{
-				FVector3d Pos = EditMesh.GetVertex(vid);
-				if (bAlignWithNormal)
+				FVector NoisePos = (FVector)((double)Options.BaseLayer.Frequency * (Pos + Offsets[0]));
+				float Displacement = Options.BaseLayer.Magnitude * FMath::PerlinNoise3D(Options.BaseLayer.Frequency * NoisePos);
+				Pos += Displacement * Normals[VertexID];
+			}
+			else
+			{
+				FVector3d Displacement;
+				for (int32 k = 0; k < 3; ++k)
 				{
-					FVector NoisePos = (FVector)( (double)Options.BaseLayer.Frequency * (Pos + Offsets[0]) );
-					float Displacement = Options.BaseLayer.Magnitude * FMath::PerlinNoise3D(Options.BaseLayer.Frequency * NoisePos);
-					Pos += Displacement * Normals[vid];
+					FVector NoisePos = (FVector)((double)Options.BaseLayer.Frequency * (Pos + Offsets[k]));
+					Displacement[k] = Options.BaseLayer.Magnitude * FMath::PerlinNoise3D(Options.BaseLayer.Frequency * NoisePos);
 				}
-				else
+				Pos += Displacement;
+			}
+			return Pos;
+		};
+
+		if (Selection.IsEmpty())
+		{
+			ParallelFor(EditMesh.MaxVertexID(), [&](int32 VertexID)
+			{
+				if (EditMesh.IsVertex(VertexID))
 				{
-					FVector3d Displacement;
-					for (int32 k = 0; k < 3; ++k)
-					{
-						FVector NoisePos = (FVector)( (double)Options.BaseLayer.Frequency * (Pos + Offsets[k]) );
-						Displacement[k] = Options.BaseLayer.Magnitude * FMath::PerlinNoise3D(Options.BaseLayer.Frequency * NoisePos);
-					}
-					Pos += Displacement;
+					EditMesh.SetVertex(VertexID, GetDisplacedPosition(VertexID));
 				}
-				EditMesh.SetVertex(vid, Pos);
-			}
-		});
+			});
+		}
+		else
+		{
+			Selection.ProcessByVertexID(EditMesh, [&](int32 VertexID)
+			{
+				EditMesh.SetVertex(VertexID, GetDisplacedPosition(VertexID));
+			});
+		}
 
 	}, EDynamicMeshChangeType::GeneralEdit, EDynamicMeshAttributeChangeFlags::Unknown, false);
 
@@ -288,6 +308,7 @@ UDynamicMesh* UGeometryScriptLibrary_MeshDeformFunctions::ApplyPerlinNoiseToMesh
 
 UDynamicMesh* UGeometryScriptLibrary_MeshDeformFunctions::ApplyIterativeSmoothingToMesh(
 	UDynamicMesh* TargetMesh,
+	FGeometryScriptMeshSelection Selection,
 	FGeometryScriptIterativeMeshSmoothingOptions Options,
 	UGeometryScriptDebug* Debug)
 {
@@ -296,6 +317,11 @@ UDynamicMesh* UGeometryScriptLibrary_MeshDeformFunctions::ApplyIterativeSmoothin
 		UE::Geometry::AppendError(Debug, EGeometryScriptErrorType::InvalidInputs, LOCTEXT("ApplyIterativeSmoothingToMesh_InvalidInput", "ApplyIterativeSmoothingToMesh: TargetMesh is Null"));
 		return TargetMesh;
 	}
+	if (Selection.IsEmpty() && Options.EmptyBehavior != EGeometryScriptEmptySelectionBehavior::FullMeshSelection )
+	{
+		return TargetMesh;
+	}
+
 	double ClampAlpha = FMathd::Clamp((double)Options.Alpha, 0.0, 1.0);
 	int ClampIters = FMath::Clamp(Options.NumIterations, 0, 100);
 	if (ClampIters == 0 || ClampAlpha == 0.0)
@@ -309,25 +335,48 @@ UDynamicMesh* UGeometryScriptLibrary_MeshDeformFunctions::ApplyIterativeSmoothin
 		TArray<FVector3d> SmoothPositions;
 		SmoothPositions.SetNum(MaxVertices);
 
-		for (int32 k = 0; k < ClampIters; ++k)
+		if (Selection.GetNumSelected() > 0)
 		{
-			ParallelFor(MaxVertices, [&](int32 VertexID)
+			TArray<int32> Vertices;
+			Selection.ConvertToMeshIndexArray(EditMesh, Vertices, EGeometryScriptIndexType::Vertex);
+			int32 NumVertices = Vertices.Num();
+			for (int32 k = 0; k < ClampIters; ++k)
 			{
-				if (EditMesh.IsVertex(VertexID))
+				ParallelFor(NumVertices, [&](int32 i)
 				{
+					int32 VertexID = Vertices[i];
 					FVector3d Centroid;
 					EditMesh.GetVtxOneRingCentroid(VertexID, Centroid);
 					SmoothPositions[VertexID] = Lerp(EditMesh.GetVertex(VertexID), Centroid, ClampAlpha);
+				});
+				for (int32 VertexID : Vertices)
+				{
+					EditMesh.SetVertex(VertexID, SmoothPositions[VertexID]);
 				}
-			});
-			for (int32 VertexID = 0; VertexID < MaxVertices; ++VertexID)
+			}
+		}
+		else
+		{
+			for (int32 k = 0; k < ClampIters; ++k)
 			{
-				if (EditMesh.IsVertex(VertexID))
+				ParallelFor(MaxVertices, [&](int32 VertexID)
 				{
-					EditMesh.SetVertex(VertexID, SmoothPositions[VertexID]);
+					if (EditMesh.IsVertex(VertexID))
+					{
+						FVector3d Centroid;
+						EditMesh.GetVtxOneRingCentroid(VertexID, Centroid);
+						SmoothPositions[VertexID] = Lerp(EditMesh.GetVertex(VertexID), Centroid, ClampAlpha);
+					}
+				});
+				for (int32 VertexID = 0; VertexID < MaxVertices; ++VertexID)
+				{
+					if (EditMesh.IsVertex(VertexID))
+					{
+						EditMesh.SetVertex(VertexID, SmoothPositions[VertexID]);
+					}
 				}
 			}
-		}		
+		}
 
 	}, EDynamicMeshChangeType::GeneralEdit, EDynamicMeshAttributeChangeFlags::Unknown, false);
 
@@ -344,6 +393,7 @@ UDynamicMesh* UGeometryScriptLibrary_MeshDeformFunctions::ApplyIterativeSmoothin
 UDynamicMesh* UGeometryScriptLibrary_MeshDeformFunctions::ApplyDisplaceFromTextureMap(
 	UDynamicMesh* TargetMesh,
 	UTexture2D* TextureAsset,
+	FGeometryScriptMeshSelection Selection,
 	FGeometryScriptDisplaceFromTextureOptions Options,
 	int32 UVLayer,
 	UGeometryScriptDebug* Debug)
@@ -359,6 +409,11 @@ UDynamicMesh* UGeometryScriptLibrary_MeshDeformFunctions::ApplyDisplaceFromTextu
 		return TargetMesh;
 	}
 
+	if (Selection.IsEmpty() && Options.EmptyBehavior != EGeometryScriptEmptySelectionBehavior::FullMeshSelection)
+	{
+		return TargetMesh;
+	}
+
 	TImageBuilder<FVector4f> ImageData;
 	if (UE::AssetUtils::ReadTexture(TextureAsset, ImageData, false) == false)
 	{
@@ -432,15 +487,30 @@ UDynamicMesh* UGeometryScriptLibrary_MeshDeformFunctions::ApplyDisplaceFromTextu
 			}
 		}
 
-		for (int32 vid : EditMesh.VertexIndicesItr())
+		// This is not necessarily the most efficient strategy, as we have computed the full-mesh displacement and
+		// then are potentially discarding most of that work. However if we filter by triangles then we might
+		// get different results when averaging per-triangle normals.
+		if (Selection.IsEmpty() == false)
 		{
-			if (Counts[vid] != 0)
+			Selection.ProcessByVertexID(EditMesh, [&](int32 VertexID)
 			{
-				EditMesh.SetVertex(vid, DisplacedPositions[vid] / (double)Counts[vid]);
+				if (Counts[VertexID] != 0)
+				{
+					EditMesh.SetVertex(VertexID, DisplacedPositions[VertexID] / (double)Counts[VertexID]);
+				}
+			});
+		}
+		else
+		{
+			for (int32 vid : EditMesh.VertexIndicesItr())
+			{
+				if (Counts[vid] != 0)
+				{
+					EditMesh.SetVertex(vid, DisplacedPositions[vid] / (double)Counts[vid]);
+				}
 			}
 		}
 
-
 	}, EDynamicMeshChangeType::GeneralEdit, EDynamicMeshAttributeChangeFlags::Unknown, false);
 
 	return TargetMesh;
@@ -453,6 +523,7 @@ UDynamicMesh* UGeometryScriptLibrary_MeshDeformFunctions::ApplyDisplaceFromTextu
 
 UDynamicMesh* UGeometryScriptLibrary_MeshDeformFunctions::ApplyDisplaceFromPerVertexVectors(
 	UDynamicMesh* TargetMesh,
+	FGeometryScriptMeshSelection Selection,
 	const FGeometryScriptVectorList& VectorList, 
 	float Magnitude,
 	UGeometryScriptDebug* Debug)
@@ -473,15 +544,27 @@ UDynamicMesh* UGeometryScriptLibrary_MeshDeformFunctions::ApplyDisplaceFromPerVe
 			return;
 		}
 
-		for (int32 vid : EditMesh.VertexIndicesItr())
+		auto UpdateVertex = [&EditMesh, &Magnitude, &Vectors](int32 VertexID)
+		{
+			FVector Position = EditMesh.GetVertex(VertexID);
+			FVector Vector = Magnitude * Vectors[VertexID];
+			if (Vector.Length() < UE_FLOAT_HUGE_DISTANCE)		// this is a bit arbitrary but should avoid disasters
+			{
+				EditMesh.SetVertex(VertexID, Position + Vector);
+			}
+		};
+
+		if (Selection.IsEmpty())
 		{
-			FVector Position = EditMesh.GetVertex(vid);
-			FVector Vector = Magnitude * Vectors[vid];
-			if (Vector.Length() < HALF_WORLD_MAX)		// this is a bit arbitrary but should avoid disasters
+			for (int32 vid : EditMesh.VertexIndicesItr())
 			{
-				EditMesh.SetVertex(vid, Position + Vector);
+				UpdateVertex(vid);
 			}
 		}
+		else
+		{
+			Selection.ProcessByVertexID(EditMesh, [&](int32 VertexID) { UpdateVertex(VertexID); });
+		}
 
 	}, EDynamicMeshChangeType::GeneralEdit, EDynamicMeshAttributeChangeFlags::Unknown, false);
 
diff --git a/Engine/Plugins/Experimental/GeometryScripting/Source/GeometryScriptingCore/Private/MeshModelingFunctions.cpp b/Engine/Plugins/Experimental/GeometryScripting/Source/GeometryScriptingCore/Private/MeshModelingFunctions.cpp
--- a/Engine/Plugins/Experimental/GeometryScripting/Source/GeometryScriptingCore/Private/MeshModelingFunctions.cpp
+++ b/Engine/Plugins/Experimental/GeometryScripting/Source/GeometryScriptingCore/Private/MeshModelingFunctions.cpp
@@ -16,6 +16,7 @@
 #include "Operations/OffsetMeshRegion.h"
 #include "Operations/MeshBevel.h"
 #include "UDynamicMesh.h"
+#include "GeometryScript/MeshSelectionFunctions.h"
 
 using namespace UE::Geometry;
 
@@ -255,6 +256,129 @@ protected:
 
 
 
+namespace UELocal
+{
+	static void ApplyPolygroupModeToNewTriangles(
+		FDynamicMesh3& EditMesh,
+		TArray<int32>& NewTriangles,
+		FGeometryScriptMeshEditPolygroupOptions GroupOptions)
+	{
+		if (GroupOptions.GroupMode == EGeometryScriptMeshEditPolygroupMode::PreserveExisting)
+		{
+			return;
+		}
+		else if (GroupOptions.GroupMode == EGeometryScriptMeshEditPolygroupMode::SetConstant)
+		{
+			for (int32 tid : NewTriangles)
+			{
+				EditMesh.SetTriangleGroup(tid, GroupOptions.ConstantGroup);
+			}
+		}
+		else  // EGeometryScriptMeshEditPolygroupMode::AutoGenerateNew
+		{
+			TMap<int32, int32> NewGroups;
+			for (int32 tid : NewTriangles)
+			{
+				int32 GroupID = EditMesh.GetTriangleGroup(tid);
+				int32* Found = NewGroups.Find(GroupID);
+				if (Found == nullptr)
+				{
+					int32 NewGroupID = EditMesh.AllocateTriangleGroup();
+					NewGroups.Add(GroupID, NewGroupID);
+					EditMesh.SetTriangleGroup(tid, NewGroupID);
+				}
+				else
+				{
+					EditMesh.SetTriangleGroup(tid, *Found);
+				}
+			}
+		}
+	}
+
+}
+
+
+
+
+
+UDynamicMesh* UGeometryScriptLibrary_MeshModelingFunctions::ApplyMeshDisconnectFaces(
+	UDynamicMesh* TargetMesh,
+	FGeometryScriptMeshSelection Selection,
+	bool bAllowBowtiesInOutput,
+	UGeometryScriptDebug* Debug)
+{
+	if (TargetMesh == nullptr)
+	{
+		UE::Geometry::AppendError(Debug, EGeometryScriptErrorType::InvalidInputs, LOCTEXT("ApplyMeshDisconnectFaces_InvalidInput", "ApplyMeshDisconnectFaces: TargetMesh is Null"));
+		return TargetMesh;
+	}
+
+	TargetMesh->EditMesh([&](FDynamicMesh3& EditMesh) 
+	{
+		TArray<int32> Triangles;
+		Selection.ConvertToMeshIndexArray(EditMesh, Triangles, EGeometryScriptIndexType::Triangle);
+		if (Triangles.Num() == 0)
+		{
+			UE::Geometry::AppendError(Debug, EGeometryScriptErrorType::InvalidInputs, LOCTEXT("ApplyMeshDisconnectFaces_NoTriangles", "ApplyMeshDisconnectFaces: Disconnect Area contains 0 triangles"));
+		}
+		else
+		{
+			FDynamicMeshEditor Editor(&EditMesh);
+			Editor.DisconnectTriangles(Triangles, !bAllowBowtiesInOutput);
+		}
+
+	}, EDynamicMeshChangeType::GeneralEdit, EDynamicMeshAttributeChangeFlags::Unknown, false);
+
+	return TargetMesh;
+}
+
+
+UDynamicMesh* UGeometryScriptLibrary_MeshModelingFunctions::ApplyMeshDuplicateFaces(
+	UDynamicMesh* TargetMesh,
+	FGeometryScriptMeshSelection Selection,
+	FGeometryScriptMeshSelection& NewSelection,
+	FGeometryScriptMeshEditPolygroupOptions GroupOptions,
+	UGeometryScriptDebug* Debug)
+{
+	if (TargetMesh == nullptr)
+	{
+		UE::Geometry::AppendError(Debug, EGeometryScriptErrorType::InvalidInputs, LOCTEXT("ApplyMeshDuplicateFaces_InvalidInput", "ApplyMeshDuplicateFaces: TargetMesh is Null"));
+		return TargetMesh;
+	}
+
+	TArray<int32> NewTriangles;
+	TargetMesh->EditMesh([&](FDynamicMesh3& EditMesh) 
+	{
+		TArray<int32> Triangles;
+		Selection.ConvertToMeshIndexArray(EditMesh, Triangles, EGeometryScriptIndexType::Triangle);
+		if (Triangles.Num() == 0)
+		{
+			UE::Geometry::AppendError(Debug, EGeometryScriptErrorType::InvalidInputs, LOCTEXT("ApplyMeshDuplicateFaces_NoTriangles", "ApplyMeshDuplicateFaces: Disconnect Area contains 0 triangles"));
+		}
+		else
+		{
+			FDynamicMeshEditor Editor(&EditMesh);
+			FMeshIndexMappings Mappings; 
+			FDynamicMeshEditResult EditResult;
+			Editor.AppendTriangles(&EditMesh, Triangles, Mappings, EditResult);
+			NewTriangles = MoveTemp(EditResult.NewTriangles);
+
+			UELocal::ApplyPolygroupModeToNewTriangles(EditMesh, NewTriangles, GroupOptions);
+		}
+
+	}, EDynamicMeshChangeType::GeneralEdit, EDynamicMeshAttributeChangeFlags::Unknown, false);
+
+	UGeometryScriptLibrary_MeshSelectionFunctions::ConvertIndexArrayToMeshSelection(
+		TargetMesh, NewTriangles, EGeometryScriptMeshSelectionType::Triangles, NewSelection);
+
+	return TargetMesh;
+}
+
+
+
+
+
+
 
 UDynamicMesh* UGeometryScriptLibrary_MeshModelingFunctions::ApplyMeshOffset(
 	UDynamicMesh* TargetMesh,
@@ -385,38 +509,183 @@ UDynamicMesh* UGeometryScriptLibrary_MeshModelingFunctions::ApplyMeshShell(
 
 
 
+UDynamicMesh* UGeometryScriptLibrary_MeshModelingFunctions::ApplyMeshExtrude_Compatibility_5p0(
+	UDynamicMesh* TargetMesh,
+	FGeometryScriptMeshExtrudeOptions Options,
+	UGeometryScriptDebug* Debug)
+{
+	FGeometryScriptMeshLinearExtrudeOptions ConvertOptions;
+	ConvertOptions.Distance = Options.ExtrudeDistance;
+	ConvertOptions.Direction = Options.ExtrudeDirection;
+	ConvertOptions.UVScale = Options.UVScale;
+	ConvertOptions.bSolidsToShells = Options.bSolidsToShells;
 
+	return ApplyMeshLinearExtrudeFaces(TargetMesh, ConvertOptions, FGeometryScriptMeshSelection(), Debug);
+}
 
-
-UDynamicMesh* UGeometryScriptLibrary_MeshModelingFunctions::ApplyMeshExtrude(
+UDynamicMesh* UGeometryScriptLibrary_MeshModelingFunctions::ApplyMeshLinearExtrudeFaces(
 	UDynamicMesh* TargetMesh,
-	FGeometryScriptMeshExtrudeOptions Options,
+	FGeometryScriptMeshLinearExtrudeOptions Options,
+	FGeometryScriptMeshSelection Selection,
 	UGeometryScriptDebug* Debug)
 {
 	if (TargetMesh == nullptr)
 	{
-		UE::Geometry::AppendError(Debug, EGeometryScriptErrorType::InvalidInputs, LOCTEXT("ApplyMeshExtrude_InvalidInput", "ApplyMeshExtrude: TargetMesh is Null"));
+		UE::Geometry::AppendError(Debug, EGeometryScriptErrorType::InvalidInputs, LOCTEXT("ApplyMeshLinearExtrudeFaces_InvalidInput", "ApplyMeshLinearExtrudeFaces: TargetMesh is Null"));
 		return TargetMesh;
 	}
-	FVector3d ExtrudeVec = (double)Options.ExtrudeDistance * (FVector3d)Options.ExtrudeDirection;
 
 	TargetMesh->EditMesh([&](FDynamicMesh3& EditMesh) 
 	{
-		FOffsetMeshRegion Extruder(&EditMesh);
-		for (int32 tid : EditMesh.TriangleIndicesItr())
+		auto ApplyExtrudeArea = [&EditMesh, &Options](TArray<int32>&& Triangles)
+		{
+			FOffsetMeshRegion Extruder(&EditMesh);
+			Extruder.Triangles = MoveTemp(Triangles);
+			if (Extruder.Triangles.Num() > 0)
+			{
+				FVector3d UseDirection = Options.Direction;
+				if (Options.DirectionMode == EGeometryScriptLinearExtrudeDirection::AverageFaceNormal)
+				{
+					UseDirection = FVector3d::Zero();
+					for (int32 tid : Extruder.Triangles)
+					{
+						FVector3d Normal, Centroid; double Area;
+						EditMesh.GetTriInfo(tid, Normal, Area, Centroid);
+						UseDirection += Area * Normal;
+					}
+					if (Normalize(UseDirection) <= 0)
+					{
+						UseDirection = Options.Direction;
+					}
+				}
+				UseDirection *= Options.Distance;
+
+				Extruder.OffsetPositionFunc = [UseDirection](const FVector3d& Position, const FVector3d& VertexVector, int VertexID)
+				{
+					return Position + UseDirection;
+				};
+				Extruder.bIsPositiveOffset = (Options.Distance > 0);
+				Extruder.UVScaleFactor = Options.UVScale;
+				Extruder.bOffsetFullComponentsAsSolids = Options.bSolidsToShells;
+				Extruder.Apply();
+				for (FOffsetMeshRegion::FOffsetInfo& OffsetInfo : Extruder.OffsetRegions)
+				{
+					UELocal::ApplyPolygroupModeToNewTriangles(EditMesh, OffsetInfo.OffsetTids, Options.GroupOptions);
+				}
+			}
+		};
+
+		TArray<int32> Triangles;
+		if (Selection.GetNumSelected() == 0)
 		{
-			Extruder.Triangles.Add(tid);
+			for (int32 tid : EditMesh.TriangleIndicesItr())
+			{
+				Triangles.Add(tid);
+			}
 		}
-		Extruder.OffsetPositionFunc = [ExtrudeVec](const FVector3d& Position, const FVector3d& VertexVector, int VertexID)
+		else
 		{
-			return Position + ExtrudeVec;
-		};
-		Extruder.bIsPositiveOffset = (Options.ExtrudeDistance > 0);
-		
-		Extruder.UVScaleFactor = Options.UVScale;
-		Extruder.bOffsetFullComponentsAsSolids = Options.bSolidsToShells;
+			Selection.ConvertToMeshIndexArray(EditMesh, Triangles, EGeometryScriptIndexType::Triangle);
+		}
+
+		if (Options.AreaMode == EGeometryScriptLinearExtrudeArea::PerPolygroup)
+		{
+			TMap<int32, TArray<int32>> PolygroupTris;
+			for (int32 tid : Triangles)
+			{
+				int32 gid = EditMesh.GetTriangleGroup(tid);
+				TArray<int32>* Found = PolygroupTris.Find(gid);
+				if (Found != nullptr)
+				{
+					Found->Add(tid);
+				}
+				else
+				{
+					TArray<int32> Temp;
+					Temp.Add(tid);
+					PolygroupTris.Add(gid, Temp);
+				}
+			}
+
+			for (TPair<int32, TArray<int32>> Pair : PolygroupTris)
+			{
+				ApplyExtrudeArea(MoveTemp(Pair.Value));
+			}
+		}
+		else if (Options.AreaMode == EGeometryScriptLinearExtrudeArea::PerTriangle)
+		{
+			for (int32 tid : Triangles)
+			{
+				TArray<int32> Temp;
+				Temp.Add(tid);
+				ApplyExtrudeArea(MoveTemp(Temp));
+			}
+		}
+		else
+		{
+			ApplyExtrudeArea(MoveTemp(Triangles));
+		}
+
+	}, EDynamicMeshChangeType::GeneralEdit, EDynamicMeshAttributeChangeFlags::Unknown, false);
+
+	return TargetMesh;
+}
+
 
-		Extruder.Apply();
+
+
+
+
+UDynamicMesh* UGeometryScriptLibrary_MeshModelingFunctions::ApplyMeshBevelSelection(
+	UDynamicMesh* TargetMesh,
+	FGeometryScriptMeshSelection Selection,
+	EGeometryScriptMeshBevelSelectionMode BevelMode,
+	FGeometryScriptMeshBevelSelectionOptions Options,
+	UGeometryScriptDebug* Debug)
+{
+	if (TargetMesh == nullptr)
+	{
+		UE::Geometry::AppendError(Debug, EGeometryScriptErrorType::InvalidInputs, LOCTEXT("ApplyMeshBevelSelection_InvalidInput", "ApplyMeshBevelSelection: TargetMesh is Null"));
+		return TargetMesh;
+	}
+
+	TargetMesh->EditMesh([&](FDynamicMesh3& EditMesh)
+	{
+		FMeshBevel Bevel;
+		Bevel.InsetDistance = Options.BevelDistance;
+		Bevel.MaterialIDMode = (Options.bInferMaterialID) ? FMeshBevel::EMaterialIDMode::InferMaterialID : FMeshBevel::EMaterialIDMode::ConstantMaterialID;
+		Bevel.SetConstantMaterialID = Options.SetMaterialID;
+
+		TUniquePtr<FGroupTopology> GroupTopology;
+
+		if (BevelMode == EGeometryScriptMeshBevelSelectionMode::TriangleArea)
+		{
+			TArray<int32> Triangles;
+			Selection.ConvertToMeshIndexArray(EditMesh, Triangles, EGeometryScriptIndexType::Triangle);
+			Bevel.InitializeFromTriangleSet(EditMesh, Triangles);
+		}
+		else
+		{
+			int32 NumRequiredAdjacentGroups = (BevelMode == EGeometryScriptMeshBevelSelectionMode::SharedPolygroupEdges) ? 2 : 1;
+
+			TArray<int32> GroupIDs;
+			Selection.ConvertToMeshIndexArray(EditMesh, GroupIDs, EGeometryScriptIndexType::PolygroupID);
+			GroupTopology = MakeUnique<FGroupTopology>(&EditMesh, true);
+			TArray<int32> GroupEdgeIDs;
+			int32 NumEdges = GroupTopology->Edges.Num();
+			for ( int32 GroupEdgeID = 0; GroupEdgeID < NumEdges; ++GroupEdgeID )
+			{
+				FIndex2i EdgeGroupIDs = GroupTopology->Edges[GroupEdgeID].Groups;
+				int32 NumAdjacentGroups = (GroupIDs.Contains(EdgeGroupIDs.A) ? 1 : 0) + (GroupIDs.Contains(EdgeGroupIDs.B) ? 1 : 0);
+				if (NumAdjacentGroups >= NumRequiredAdjacentGroups )
+				{
+					GroupEdgeIDs.Add(GroupEdgeID);
+				}
+			}
+			Bevel.InitializeFromGroupTopologyEdges(EditMesh, *GroupTopology, GroupEdgeIDs);
+		}
+
+		Bevel.Apply(EditMesh, nullptr);
 
 	}, EDynamicMeshChangeType::GeneralEdit, EDynamicMeshAttributeChangeFlags::Unknown, false);
 
diff --git a/Engine/Plugins/Experimental/GeometryScripting/Source/GeometryScriptingCore/Public/GeometryScript/MeshDeformFunctions.h b/Engine/Plugins/Experimental/GeometryScripting/Source/GeometryScriptingCore/Public/GeometryScript/MeshDeformFunctions.h
--- a/Engine/Plugins/Experimental/GeometryScripting/Source/GeometryScriptingCore/Public/GeometryScript/MeshDeformFunctions.h
+++ b/Engine/Plugins/Experimental/GeometryScripting/Source/GeometryScriptingCore/Public/GeometryScript/MeshDeformFunctions.h
@@ -5,6 +5,7 @@
 #include "CoreMinimal.h"
 #include "Kismet/BlueprintFunctionLibrary.h"
 #include "GeometryScript/GeometryScriptTypes.h"
+#include "GeometryScript/GeometryScriptSelectionTypes.h"
 #include "MeshDeformFunctions.generated.h"
 
 class UDynamicMesh;
@@ -138,6 +139,10 @@ public:
 
 	UPROPERTY(EditAnywhere, BlueprintReadWrite, Category = Options)
 	bool bApplyAlongNormal = true;
+
+	/** EmptyBehavior Defines how an empty input selection should be interpreted */
+	UPROPERTY(EditAnywhere, BlueprintReadWrite, Category = Options)
+	EGeometryScriptEmptySelectionBehavior EmptyBehavior = EGeometryScriptEmptySelectionBehavior::FullMeshSelection;
 };
 
 
@@ -151,6 +156,10 @@ public:
 
 	UPROPERTY(EditAnywhere, BlueprintReadWrite, Category = Options)
 	float Alpha = 0.2;
+
+	/** EmptyBehavior Defines how an empty input selection should be interpreted */
+	UPROPERTY(EditAnywhere, BlueprintReadWrite, Category = Options)
+	EGeometryScriptEmptySelectionBehavior EmptyBehavior = EGeometryScriptEmptySelectionBehavior::FullMeshSelection;
 };
 
 
@@ -175,6 +184,10 @@ public:
 
 	UPROPERTY(EditAnywhere, BlueprintReadWrite, Category = Options)
 	int ImageChannel = 0;
+
+	/** EmptyBehavior Defines how an empty input selection should be interpreted */
+	UPROPERTY(EditAnywhere, BlueprintReadWrite, Category = Options)
+	EGeometryScriptEmptySelectionBehavior EmptyBehavior = EGeometryScriptEmptySelectionBehavior::FullMeshSelection;
 };
 
 
@@ -233,6 +246,7 @@ public:
 	static UPARAM(DisplayName = "Target Mesh") UDynamicMesh* 
 	ApplyPerlinNoiseToMesh(  
 		UDynamicMesh* TargetMesh, 
+		FGeometryScriptMeshSelection Selection,
 		FGeometryScriptPerlinNoiseOptions Options,
 		UGeometryScriptDebug* Debug = nullptr);
 
@@ -240,6 +254,7 @@ public:
 	static UPARAM(DisplayName = "Target Mesh") UDynamicMesh* 
 	ApplyIterativeSmoothingToMesh(  
 		UDynamicMesh* TargetMesh, 
+		FGeometryScriptMeshSelection Selection,
 		FGeometryScriptIterativeMeshSmoothingOptions Options,
 		UGeometryScriptDebug* Debug = nullptr);
 
@@ -249,15 +264,21 @@ public:
 	ApplyDisplaceFromTextureMap(  
 		UDynamicMesh* TargetMesh, 
 		UTexture2D* Texture,
+		FGeometryScriptMeshSelection Selection,
 		FGeometryScriptDisplaceFromTextureOptions Options,
 		int32 UVLayer = 0,
 		UGeometryScriptDebug* Debug = nullptr);
 
-	
+	/**
+	 * Add the vectors in VectorList, scaled by Magnitude, to the vertex positions in TargetMesh.
+	 * VectorList Length must be >= the as TargetMesh MaxVertexID.
+	 * @param Selection if non-empty, only the vertices identified by the selection will be displaced. The VectorList must still be the same size as the whole mesh, this is just a filter on which vertices are updated.
+	 */
 	UFUNCTION(BlueprintCallable, Category = "GeometryScript|Deformations", meta=(ScriptMethod))
 	static UPARAM(DisplayName = "Target Mesh") UDynamicMesh* 
 	ApplyDisplaceFromPerVertexVectors(  
 		UDynamicMesh* TargetMesh,
+		FGeometryScriptMeshSelection Selection,
 		const FGeometryScriptVectorList& VectorList, 
 		float Magnitude = 5.0,
 		UGeometryScriptDebug* Debug = nullptr);
diff --git a/Engine/Plugins/Experimental/GeometryScripting/Source/GeometryScriptingCore/Public/GeometryScript/MeshModelingFunctions.h b/Engine/Plugins/Experimental/GeometryScripting/Source/GeometryScriptingCore/Public/GeometryScript/MeshModelingFunctions.h
--- a/Engine/Plugins/Experimental/GeometryScripting/Source/GeometryScriptingCore/Public/GeometryScript/MeshModelingFunctions.h
+++ b/Engine/Plugins/Experimental/GeometryScripting/Source/GeometryScriptingCore/Public/GeometryScript/MeshModelingFunctions.h
@@ -5,10 +5,31 @@
 #include "CoreMinimal.h"
 #include "Kismet/BlueprintFunctionLibrary.h"
 #include "GeometryScript/GeometryScriptTypes.h"
+#include "GeometryScript/GeometryScriptSelectionTypes.h"
 #include "MeshModelingFunctions.generated.h"
 
 class UDynamicMesh;
 
+UENUM(BlueprintType)
+enum class EGeometryScriptMeshEditPolygroupMode : uint8
+{
+	PreserveExisting = 0,
+	AutoGenerateNew = 1,
+	SetConstant = 2
+};
+
+USTRUCT(BlueprintType)
+struct GEOMETRYSCRIPTINGCORE_API FGeometryScriptMeshEditPolygroupOptions
+{
+	GENERATED_BODY()
+
+	UPROPERTY(EditAnywhere, BlueprintReadWrite, Category = Options)
+	EGeometryScriptMeshEditPolygroupMode GroupMode = EGeometryScriptMeshEditPolygroupMode::PreserveExisting;
+
+	UPROPERTY(EditAnywhere, BlueprintReadWrite, Category = Options)
+	int ConstantGroup = 0;
+};
+
 
 USTRUCT(BlueprintType)
 struct GEOMETRYSCRIPTINGCORE_API FGeometryScriptMeshOffsetOptions
@@ -56,6 +77,50 @@ public:
 };
 
 
+UENUM(BlueprintType)
+enum class EGeometryScriptLinearExtrudeDirection : uint8
+{
+	FixedDirection = 0,
+	AverageFaceNormal = 1
+};
+
+UENUM(BlueprintType)
+enum class EGeometryScriptLinearExtrudeArea : uint8
+{
+	EntireSelection = 0,
+	PerPolygroup = 1,
+	PerTriangle = 2
+};
+
+USTRUCT(BlueprintType)
+struct GEOMETRYSCRIPTINGCORE_API FGeometryScriptMeshLinearExtrudeOptions
+{
+	GENERATED_BODY()
+public:
+	UPROPERTY(EditAnywhere, BlueprintReadWrite, Category = Options)
+	float Distance = 1.0;
+
+	UPROPERTY(EditAnywhere, BlueprintReadWrite, Category = Options)
+	EGeometryScriptLinearExtrudeDirection DirectionMode = EGeometryScriptLinearExtrudeDirection::FixedDirection;
+
+	UPROPERTY(EditAnywhere, BlueprintReadWrite, Category = Options)
+	FVector Direction = FVector(0,0,1);
+
+	UPROPERTY(EditAnywhere, BlueprintReadWrite, Category = Options)
+	EGeometryScriptLinearExtrudeArea AreaMode = EGeometryScriptLinearExtrudeArea::EntireSelection;
+
+	UPROPERTY(EditAnywhere, BlueprintReadWrite, Category = Options)
+	FGeometryScriptMeshEditPolygroupOptions GroupOptions;
+
+	UPROPERTY(EditAnywhere, BlueprintReadWrite, Category = Options)
+	float UVScale = 1.0f;
+
+	UPROPERTY(EditAnywhere, BlueprintReadWrite, Category = Options)
+	bool bSolidsToShells = true;
+};
+
+
+
 
 USTRUCT(BlueprintType)
 struct GEOMETRYSCRIPTINGCORE_API FGeometryScriptMeshBevelOptions
@@ -101,6 +166,37 @@ public:
 };
 
 
+/**
+ * Mode passed to ApplyMeshBevelSelection to control how the input Selection should
+ * be interpreted as selecting an area of the mesh to Bevel
+ */
+UENUM(BlueprintType)
+enum class EGeometryScriptMeshBevelSelectionMode : uint8
+{
+	/** Convert the selection to Triangles and bevel the boundary edge loops of the triangle set */
+	TriangleArea = 0,
+	/** Convert the selection to Polygroups and bevel all the Polygroup Edges of the selected Polygroups */
+	AllPolygroupEdges = 1,
+	/** Convert the selection to Polygroups and bevel all the Polygroup Edges that are between selected Polygroups */
+	SharedPolygroupEdges = 2
+};
+
+
+USTRUCT(BlueprintType)
+struct GEOMETRYSCRIPTINGCORE_API FGeometryScriptMeshBevelSelectionOptions
+{
+	GENERATED_BODY()
+public:
+	UPROPERTY(EditAnywhere, BlueprintReadWrite, Category = Options)
+	float BevelDistance = 1.0;
+
+	UPROPERTY(EditAnywhere, BlueprintReadWrite, Category = Options)
+	bool bInferMaterialID = false;
+
+	UPROPERTY(EditAnywhere, BlueprintReadWrite, Category = Options)
+	int SetMaterialID = 0;
+};
+
 
 
 UCLASS(meta = (ScriptName = "GeometryScript_MeshModeling"))
@@ -109,6 +205,36 @@ class GEOMETRYSCRIPTINGCORE_API UGeometryScriptLibrary_MeshModelingFunctions : p
 	GENERATED_BODY()
 public:
 
+	/**
+	 * Disconnect the triangles of TargetMesh identified by the Selection.
+	 * The input Selection will still identify the same geometric elements after Disconnecting.
+	 * @param bAllowBowtiesInOutput if false, any bowtie vertices resulting created in the Duplicate area will be disconnected into unique vertices
+	 */
+	UFUNCTION(BlueprintCallable, Category = "GeometryScript|Modeling", meta=(ScriptMethod))
+	static UPARAM(DisplayName = "Target Mesh") UDynamicMesh* 
+	ApplyMeshDisconnectFaces(
+		UDynamicMesh* TargetMesh,
+		FGeometryScriptMeshSelection Selection,
+		bool bAllowBowtiesInOutput = true,
+		UGeometryScriptDebug* Debug = nullptr );
+
+	/**
+	 * Duplicate the triangles of TargetMesh identified by the Selection
+	 * @param NewTriangles a Mesh Selection of the duplicate triangles is returned here (with type Triangles)
+	 * @param bAllowBowtiesInOutput if false, any bowtie vertices resulting created in the Duplicate area will be disconnected into unique vertices
+	 */
+	UFUNCTION(BlueprintCallable, Category = "GeometryScript|Modeling", meta=(ScriptMethod))
+	static UPARAM(DisplayName = "Target Mesh") UDynamicMesh* 
+	ApplyMeshDuplicateFaces(
+		UDynamicMesh* TargetMesh,
+		FGeometryScriptMeshSelection Selection,
+		FGeometryScriptMeshSelection& NewTriangles,
+		FGeometryScriptMeshEditPolygroupOptions GroupOptions = FGeometryScriptMeshEditPolygroupOptions(),
+		UGeometryScriptDebug* Debug = nullptr );
+
+
+
+
 	UFUNCTION(BlueprintCallable, Category = "GeometryScript|Modeling", meta=(ScriptMethod))
 	static UPARAM(DisplayName = "Target Mesh") UDynamicMesh* 
 	ApplyMeshOffset(
@@ -123,13 +249,37 @@ public:
 		FGeometryScriptMeshOffsetOptions Options,
 		UGeometryScriptDebug* Debug = nullptr );
 
+
+	/**
+	 * Apply Linear Extrusion (ie extrusion in a single direction) to the triangles of TargetMesh identified by the Selection.
+	 * The input Selection will still identify the same geometric elements after the Extrusion
+	 */
 	UFUNCTION(BlueprintCallable, Category = "GeometryScript|Modeling", meta=(ScriptMethod))
 	static UPARAM(DisplayName = "Target Mesh") UDynamicMesh* 
-	ApplyMeshExtrude(
+	ApplyMeshLinearExtrudeFaces(
 		UDynamicMesh* TargetMesh,
-		FGeometryScriptMeshExtrudeOptions Options,
+		FGeometryScriptMeshLinearExtrudeOptions Options,
+		FGeometryScriptMeshSelection Selection,
 		UGeometryScriptDebug* Debug = nullptr );
 
+
+
+	/**
+	 * Apply a Mesh Bevel operation to parts of TargetMesh using the BevelOptions settings.
+	 * @param Selection specifies which mesh edges to Bevel
+	 * @param BevelMode specifies how Selection should be converted to a Triangle Region or set of Polygroup Edges
+	 * @param BevelOptions settings for the Bevel Operation
+	 */
+	UFUNCTION(BlueprintCallable, Category = "GeometryScript|Modeling", meta=(ScriptMethod))
+	static UPARAM(DisplayName = "Target Mesh") UDynamicMesh* 
+	ApplyMeshBevelSelection(
+		UDynamicMesh* TargetMesh,
+		FGeometryScriptMeshSelection Selection,
+		EGeometryScriptMeshBevelSelectionMode BevelMode,
+		FGeometryScriptMeshBevelSelectionOptions BevelOptions,
+		UGeometryScriptDebug* Debug = nullptr );
+
+
 	UFUNCTION(BlueprintCallable, Category = "GeometryScript|Modeling", meta=(ScriptMethod))
 	static UPARAM(DisplayName = "Target Mesh") UDynamicMesh* 
 	ApplyMeshPolygroupBevel(
@@ -137,4 +287,26 @@ public:
 		FGeometryScriptMeshBevelOptions Options,
 		UGeometryScriptDebug* Debug = nullptr );
 
+
+
+
+	//---------------------------------------------
+	// Backwards-Compatibility implementations
+	//---------------------------------------------
+	// These are versions/variants of the above functions that were released
+	// in previous UE 5.x versions, that have since been updated. 
+	// To avoid breaking user scripts, these previous versions are currently kept and 
+	// called via redirectors registered in GeometryScriptingCoreModule.cpp.
+	// 
+	// These functions may be deprecated in future UE releases.
+	//
+
+
+	UFUNCTION(BlueprintCallable, Category = "GeometryScript|Compatibility", meta=(ScriptMethod))
+	static UPARAM(DisplayName = "Target Mesh") UDynamicMesh* 
+	ApplyMeshExtrude_Compatibility_5p0(
+		UDynamicMesh* TargetMesh,
+		FGeometryScriptMeshExtrudeOptions Options,
+		UGeometryScriptDebug* Debug = nullptr );
+
 };
\ No newline at end of file
