GeometryScript: add new ComputePolygroupsFromPolygonDetection function that exposes new polygroups-from-polygon-detector algorithm

diff --git a/Engine/Plugins/Experimental/GeometryScripting/Source/GeometryScriptingCore/Private/MeshPolygroupFunctions.cpp b/Engine/Plugins/Experimental/GeometryScripting/Source/GeometryScriptingCore/Private/MeshPolygroupFunctions.cpp
--- a/Engine/Plugins/Experimental/GeometryScripting/Source/GeometryScriptingCore/Private/MeshPolygroupFunctions.cpp
+++ b/Engine/Plugins/Experimental/GeometryScripting/Source/GeometryScriptingCore/Private/MeshPolygroupFunctions.cpp
@@ -286,6 +286,46 @@ UDynamicMesh* UGeometryScriptLibrary_MeshPolygroupFunctions::ComputePolygroupsFr
 }
 
 
+UDynamicMesh* UGeometryScriptLibrary_MeshPolygroupFunctions::ComputePolygroupsFromPolygonDetection( 
+	UDynamicMesh* TargetMesh, 
+	FGeometryScriptGroupLayer GroupLayer,
+	bool bRespectUVSeams,
+	bool bRespectHardNormals,
+	double QuadAdjacencyWeight,
+	double QuadMetricClamp,
+	int MaxSearchRounds,
+	UGeometryScriptDebug* Debug)
+{
+	if (TargetMesh == nullptr)
+	{
+		UE::Geometry::AppendError(Debug, EGeometryScriptErrorType::InvalidInputs, LOCTEXT("ComputePolygroupsFromPolygonDetection_InvalidInput", "ComputePolygroupsFromPolygonDetection: TargetMesh is Null"));
+		return TargetMesh;
+	}
+
+	TargetMesh->EditMesh([&](FDynamicMesh3& EditMesh) 
+	{
+		TUniquePtr<FPolygroupSet> OutputGroups;
+		FPolygroupLayer InputGroupLayer{ GroupLayer.bDefaultLayer, GroupLayer.ExtendedLayerIndex };
+		if (InputGroupLayer.CheckExists(&EditMesh) == false)
+		{
+			UE::Geometry::AppendError(Debug, EGeometryScriptErrorType::InvalidInputs, LOCTEXT("ComputePolygroupsFromPolygonDetection_MissingGroups", "ComputePolygroupsFromPolygonDetection: Target Polygroup Layer does not exist"));
+			return;
+		}
+		OutputGroups = MakeUnique<FPolygroupSet>(&EditMesh, InputGroupLayer);
+
+		FPolygroupsGenerator Generator(&EditMesh);
+		Generator.bApplyPostProcessing = false;
+		Generator.bCopyToMesh = false;
+		Generator.FindSourceMeshPolygonPolygroups(
+			bRespectUVSeams, bRespectHardNormals,
+			QuadAdjacencyWeight, QuadMetricClamp, MaxSearchRounds);
+		Generator.CopyPolygroupsToPolygroupSet(*OutputGroups, EditMesh);
+
+	}, EDynamicMeshChangeType::GeneralEdit, EDynamicMeshAttributeChangeFlags::Unknown, false);
+
+	return TargetMesh;
+}
+
 
 
 int32 UGeometryScriptLibrary_MeshPolygroupFunctions::GetTrianglePolygroupID( 
diff --git a/Engine/Plugins/Experimental/GeometryScripting/Source/GeometryScriptingCore/Public/GeometryScript/MeshPolygroupFunctions.h b/Engine/Plugins/Experimental/GeometryScripting/Source/GeometryScriptingCore/Public/GeometryScript/MeshPolygroupFunctions.h
--- a/Engine/Plugins/Experimental/GeometryScripting/Source/GeometryScriptingCore/Public/GeometryScript/MeshPolygroupFunctions.h
+++ b/Engine/Plugins/Experimental/GeometryScripting/Source/GeometryScriptingCore/Public/GeometryScript/MeshPolygroupFunctions.h
@@ -69,6 +69,17 @@ public:
 		int MinGroupSize = 2,
 		UGeometryScriptDebug* Debug = nullptr );
 
+	UFUNCTION(BlueprintCallable, Category = "GeometryScript|Polygroups", meta=(ScriptMethod))
+	static UPARAM(DisplayName = "Target Mesh") UDynamicMesh* 
+	ComputePolygroupsFromPolygonDetection( 
+		UDynamicMesh* TargetMesh, 
+		FGeometryScriptGroupLayer GroupLayer,
+		bool bRespectUVSeams = true,
+		bool bRespectHardNormals = false,
+		double QuadAdjacencyWeight = 1.0,
+		double QuadMetricClamp = 1.0,
+		int MaxSearchRounds = 1,
+		UGeometryScriptDebug* Debug = nullptr );
 
 	UFUNCTION(BlueprintPure, Category = "GeometryScript|Polygroups", meta=(ScriptMethod))
 	static UPARAM(DisplayName = "Polygroup ID") int32
