GeometryCore: make FDynamicMesh3::GetMeshInfoString() const
GeometryScript: add SplitMeshBowties, RepairMeshDegenerateGeometry, and GetMeshInfoString functions

diff --git a/Engine/Plugins/Experimental/GeometryScripting/Source/GeometryScriptingCore/Private/MeshQueryFunctions.cpp b/Engine/Plugins/Experimental/GeometryScripting/Source/GeometryScriptingCore/Private/MeshQueryFunctions.cpp
--- a/Engine/Plugins/Experimental/GeometryScripting/Source/GeometryScriptingCore/Private/MeshQueryFunctions.cpp
+++ b/Engine/Plugins/Experimental/GeometryScripting/Source/GeometryScriptingCore/Private/MeshQueryFunctions.cpp
@@ -57,6 +57,11 @@ ReturnType SimpleMeshUVSetQuery(UDynamicMesh* Mesh, int32 UVSetIndex, bool& bIsV
 
 
 
+FString UGeometryScriptLibrary_MeshQueryFunctions::GetMeshInfoString(UDynamicMesh* TargetMesh)
+{
+	return SimpleMeshQuery<FString>(TargetMesh, FString(TEXT("Mesh is Null")), [&](const FDynamicMesh3& Mesh) { return Mesh.MeshInfoString(); });
+}
+
 
 bool UGeometryScriptLibrary_MeshQueryFunctions::GetIsDenseMesh( UDynamicMesh* TargetMesh )
 {
diff --git a/Engine/Plugins/Experimental/GeometryScripting/Source/GeometryScriptingCore/Private/MeshRepairFunctions.cpp b/Engine/Plugins/Experimental/GeometryScripting/Source/GeometryScriptingCore/Private/MeshRepairFunctions.cpp
--- a/Engine/Plugins/Experimental/GeometryScripting/Source/GeometryScriptingCore/Private/MeshRepairFunctions.cpp
+++ b/Engine/Plugins/Experimental/GeometryScripting/Source/GeometryScriptingCore/Private/MeshRepairFunctions.cpp
@@ -15,6 +15,8 @@
 #include "Selections/MeshConnectedComponents.h"
 #include "Selections/MeshFaceSelection.h"
 #include "UDynamicMesh.h"
+#include "MeshSimplification.h"
+#include "MeshConstraintsUtil.h"
 
 #include "CleaningOps/HoleFillOp.h"
 
@@ -250,4 +252,136 @@ UDynamicMesh* UGeometryScriptLibrary_MeshRepairFunctions::RemoveHiddenTriangles(
 
 
 
+UDynamicMesh* UGeometryScriptLibrary_MeshRepairFunctions::SplitMeshBowties(
+	UDynamicMesh* TargetMesh,
+	bool bMeshBowties,
+	bool bAttributeBowties,
+	UGeometryScriptDebug* Debug)
+{
+	if (TargetMesh == nullptr)
+	{
+		UE::Geometry::AppendError(Debug, EGeometryScriptErrorType::InvalidInputs, LOCTEXT("SplitMeshBowties_InvalidInput", "SplitMeshBowties: TargetMesh is Null"));
+		return TargetMesh;
+	}
+
+	TargetMesh->EditMesh([&](FDynamicMesh3& EditMesh) 
+	{
+		if (bMeshBowties)
+		{
+			FDynamicMeshEditor Editor(&EditMesh);
+			FDynamicMeshEditResult EditResult;
+			Editor.SplitBowties(EditResult);
+		}
+
+		if (bAttributeBowties)
+		{
+			if (EditMesh.HasAttributes())
+			{
+				EditMesh.Attributes()->SplitAllBowties(true);
+			}
+		}
+
+	}, EDynamicMeshChangeType::GeneralEdit, EDynamicMeshAttributeChangeFlags::Unknown, false);
+
+	return TargetMesh;
+}
+
+
+
+UDynamicMesh* UGeometryScriptLibrary_MeshRepairFunctions::RepairMeshDegenerateGeometry(
+	UDynamicMesh* TargetMesh,
+	FGeometryScriptDegenerateTriangleOptions Options,
+	UGeometryScriptDebug* Debug)
+{
+	if (TargetMesh == nullptr)
+	{
+		UE::Geometry::AppendError(Debug, EGeometryScriptErrorType::InvalidInputs, LOCTEXT("RepairMeshDegenerateGeometry_InvalidInput", "RepairMeshDegenerateGeometry: TargetMesh is Null"));
+		return TargetMesh;
+	}
+
+	bool bModified = false;
+
+	auto DeleteByAreaPass = [&bModified, &Options](FDynamicMesh3& EditMesh)
+	{
+		TArray<int32> ToDelete;
+		for (int32 tid : EditMesh.TriangleIndicesItr())
+		{
+			if (EditMesh.GetTriArea(tid) < Options.MinTriangleArea)
+			{
+				ToDelete.Add(tid);
+			}
+		}
+		FDynamicMeshEditor Editor(&EditMesh);
+		Editor.RemoveTriangles(ToDelete, true);
+		bModified = true;
+	};
+
+	auto DeleteByEdgePass = [&bModified, &Options](FDynamicMesh3& EditMesh)
+	{
+		TArray<int32> ToDelete;
+		for (int32 eid : EditMesh.EdgeIndicesItr())
+		{
+			FDynamicMesh3::FEdge Edge = EditMesh.GetEdge(eid);
+			if ( Distance(EditMesh.GetVertex(Edge.Vert.A), EditMesh.GetVertex(Edge.Vert.B)) < Options.MinEdgeLength )
+			{
+				ToDelete.AddUnique(Edge.Tri.A);
+				if (Edge.Tri.B != IndexConstants::InvalidID)
+				{
+					ToDelete.AddUnique(Edge.Tri.B);
+				}
+			}
+		}
+		FDynamicMeshEditor Editor(&EditMesh);
+		Editor.RemoveTriangles(ToDelete, true);
+		bModified = true;
+	};
+
+
+	TargetMesh->EditMesh([&](FDynamicMesh3& EditMesh) 
+	{
+		bool bTryRepair = (Options.Mode != EGeometryScriptRepairMeshMode::DeleteOnly);
+		bool bDoDelete = (Options.Mode != EGeometryScriptRepairMeshMode::RepairOrSkip);
+
+		if (bTryRepair)
+		{
+			FQEMSimplification Simplifier(&EditMesh);
+			Simplifier.ProjectionMode = FQEMSimplification::ETargetProjectionMode::NoProjection;
+			Simplifier.DEBUG_CHECK_LEVEL = 0;
+			Simplifier.bRetainQuadricMemory = false;
+			if (EditMesh.HasAttributes())
+			{
+				Simplifier.bAllowSeamCollapse = true;
+				Simplifier.SetEdgeFlipTolerance(1.e-5);
+				EditMesh.Attributes()->SplitAllBowties();	// eliminate any bowties that might have formed on attribute seams.
+			}
+			FMeshConstraints Constraints;
+			FMeshConstraintsUtil::ConstrainAllBoundariesAndSeams(Constraints, EditMesh,
+				EEdgeRefineFlags::NoFlip, EEdgeRefineFlags::NoConstraint, EEdgeRefineFlags::NoConstraint,
+				false, false, true);
+			Simplifier.SetExternalConstraints(MoveTemp(Constraints));
+			Simplifier.SimplifyToEdgeLength(Options.MinEdgeLength);
+
+			//TODO: this would not resolve sliver triangles that have all long-edges...
+			//TODO: could we do explicit collapse passes here? we would need to still respect constraints...
+
+			bModified = true;
+		}
+
+		if (bDoDelete)
+		{
+			DeleteByAreaPass(EditMesh);
+			DeleteByEdgePass(EditMesh);
+		}
+
+		if (bModified && Options.bCompactOnCompletion)
+		{
+			EditMesh.CompactInPlace();
+		}
+
+	}, EDynamicMeshChangeType::GeneralEdit, EDynamicMeshAttributeChangeFlags::Unknown, false);
+
+	return TargetMesh;
+}
+
+
 #undef LOCTEXT_NAMESPACE
diff --git a/Engine/Plugins/Experimental/GeometryScripting/Source/GeometryScriptingCore/Public/GeometryScript/MeshQueryFunctions.h b/Engine/Plugins/Experimental/GeometryScripting/Source/GeometryScriptingCore/Public/GeometryScript/MeshQueryFunctions.h
--- a/Engine/Plugins/Experimental/GeometryScripting/Source/GeometryScriptingCore/Public/GeometryScript/MeshQueryFunctions.h
+++ b/Engine/Plugins/Experimental/GeometryScripting/Source/GeometryScriptingCore/Public/GeometryScript/MeshQueryFunctions.h
@@ -16,6 +16,9 @@ class GEOMETRYSCRIPTINGCORE_API UGeometryScriptLibrary_MeshQueryFunctions : publ
 	GENERATED_BODY()
 public:
 
+	UFUNCTION(BlueprintPure, Category = "GeometryScript|MeshQueries", meta=(ScriptMethod))
+	static FString GetMeshInfoString( UDynamicMesh* TargetMesh );
+
 	UFUNCTION(BlueprintPure, Category = "GeometryScript|MeshQueries", meta=(ScriptMethod))
 	static bool GetIsDenseMesh( UDynamicMesh* TargetMesh );
 
diff --git a/Engine/Plugins/Experimental/GeometryScripting/Source/GeometryScriptingCore/Public/GeometryScript/MeshRepairFunctions.h b/Engine/Plugins/Experimental/GeometryScripting/Source/GeometryScriptingCore/Public/GeometryScript/MeshRepairFunctions.h
--- a/Engine/Plugins/Experimental/GeometryScripting/Source/GeometryScriptingCore/Public/GeometryScript/MeshRepairFunctions.h
+++ b/Engine/Plugins/Experimental/GeometryScripting/Source/GeometryScriptingCore/Public/GeometryScript/MeshRepairFunctions.h
@@ -125,6 +125,38 @@ public:
 
 
 
+UENUM(BlueprintType)
+enum class EGeometryScriptRepairMeshMode : uint8
+{
+	DeleteOnly = 0,
+	RepairOrDelete = 1,
+	RepairOrSkip = 2
+};
+
+
+
+USTRUCT(BlueprintType)
+struct GEOMETRYSCRIPTINGCORE_API FGeometryScriptDegenerateTriangleOptions
+{
+	GENERATED_BODY()
+public:
+	/**  */
+	UPROPERTY(EditAnywhere, BlueprintReadWrite, Category = Options)
+	EGeometryScriptRepairMeshMode Mode = EGeometryScriptRepairMeshMode::RepairOrDelete;
+
+	/**  */
+	UPROPERTY(EditAnywhere, BlueprintReadWrite, Category = Options)
+	double MinTriangleArea = 0.001;
+
+	/**  */
+	UPROPERTY(EditAnywhere, BlueprintReadWrite, Category = Options)
+	double MinEdgeLength = 0.0001;
+
+	UPROPERTY(EditAnywhere, BlueprintReadWrite, Category = Options)
+	bool bCompactOnCompletion = true;
+};
+
+
 
 UCLASS(meta = (ScriptName = "GeometryScript_MeshRepair"))
 class GEOMETRYSCRIPTINGCORE_API UGeometryScriptLibrary_MeshRepairFunctions : public UBlueprintFunctionLibrary
@@ -175,4 +207,18 @@ public:
 		FGeometryScriptRemoveHiddenTrianglesOptions Options,
 		UGeometryScriptDebug* Debug = nullptr);
 
+	UFUNCTION(BlueprintCallable, Category = "GeometryScript|Repair", meta=(ScriptMethod))
+	static UPARAM(DisplayName = "Target Mesh") UDynamicMesh* 
+	SplitMeshBowties(  
+		UDynamicMesh* TargetMesh, 
+		bool bMeshBowties = true,
+		bool bAttributeBowties = true,
+		UGeometryScriptDebug* Debug = nullptr);
+
+	UFUNCTION(BlueprintCallable, Category = "GeometryScript|Repair", meta=(ScriptMethod))
+	static UPARAM(DisplayName = "Target Mesh") UDynamicMesh* 
+	RepairMeshDegenerateGeometry(  
+		UDynamicMesh* TargetMesh, 
+		FGeometryScriptDegenerateTriangleOptions Options,
+		UGeometryScriptDebug* Debug = nullptr);
 };
\ No newline at end of file
diff --git a/Engine/Source/Runtime/GeometryCore/Private/DynamicMesh/DynamicMesh3.cpp b/Engine/Source/Runtime/GeometryCore/Private/DynamicMesh/DynamicMesh3.cpp
--- a/Engine/Source/Runtime/GeometryCore/Private/DynamicMesh/DynamicMesh3.cpp
+++ b/Engine/Source/Runtime/GeometryCore/Private/DynamicMesh/DynamicMesh3.cpp
@@ -676,7 +676,7 @@ void FDynamicMesh3::EnumerateEdgeTriangles(int32 EdgeID, TFunctionRef<void(int32
 }
 
 
-FString FDynamicMesh3::MeshInfoString()
+FString FDynamicMesh3::MeshInfoString() const
 {
 	FString VtxString = FString::Printf(TEXT("Vertices count %d max %d  %s  VtxEdges %s"),
 		VertexCount(), MaxVertexID(), *VertexRefCounts.UsageStats(), *(VertexEdgeLists.MemoryUsage()));
@@ -686,8 +686,8 @@ FString FDynamicMesh3::MeshInfoString()
 		EdgeCount(), MaxEdgeID(), *EdgeRefCounts.UsageStats());
 	FString AttribString = FString::Printf(TEXT("VtxNormals %d  VtxColors %d  VtxUVs %d  TriGroups %d  Attributes %d"),
 		HasVertexNormals(), HasVertexColors(), HasVertexUVs(), HasTriangleGroups(), HasAttributes());
-	FString InfoString = FString::Printf(TEXT("Closed %d  Compact %d  ShapeChangeStamp %d  TopologyChangeStamp %d  MaxGroupID %d"),
-		IsClosed(), IsCompact(), GetShapeChangeStamp(), GetTopologyChangeStamp(), MaxGroupID());
+	FString InfoString = FString::Printf(TEXT("Closed %d  Compact %d  TopologyChangeStamp %d  MaxGroupID %d"),
+		IsClosed(), IsCompact(), GetTopologyChangeStamp(), MaxGroupID());
 
 	return VtxString + "\n" + TriString + "\n" + EdgeString + "\n" + AttribString + "\n" + InfoString;
 }
diff --git a/Engine/Source/Runtime/GeometryCore/Public/DynamicMesh/DynamicMesh3.h b/Engine/Source/Runtime/GeometryCore/Public/DynamicMesh/DynamicMesh3.h
--- a/Engine/Source/Runtime/GeometryCore/Public/DynamicMesh/DynamicMesh3.h
+++ b/Engine/Source/Runtime/GeometryCore/Public/DynamicMesh/DynamicMesh3.h
@@ -1320,7 +1320,7 @@ public:
 	/**
 	 * Returns a debug string that contains mesh statistics and other information
 	 */
-	virtual FString MeshInfoString();
+	virtual FString MeshInfoString() const;
 
 	/**
 	 * Options for the IsSameAs check
