GeometryScript: add ExpandContractMeshSelection. Add option to fallback to full mesh processing for FGeometryScriptMeshSelection::ProcessByTriangleID() and ProcessByVertexID(), to simplify handling "empty selection is full mesh" cases
GeometryProcessing: add basic implementation of FMeshVertexSelection::ContractByBorderVertices()

diff --git a/Engine/Plugins/Experimental/GeometryScripting/Source/GeometryScriptingCore/Private/GeometryScriptSelectionTypes.cpp b/Engine/Plugins/Experimental/GeometryScripting/Source/GeometryScriptingCore/Private/GeometryScriptSelectionTypes.cpp
--- a/Engine/Plugins/Experimental/GeometryScripting/Source/GeometryScriptingCore/Private/GeometryScriptSelectionTypes.cpp
+++ b/Engine/Plugins/Experimental/GeometryScripting/Source/GeometryScriptingCore/Private/GeometryScriptSelectionTypes.cpp
@@ -286,8 +286,18 @@ EGeometryScriptIndexType FGeometryScriptMeshSelection::ConvertToMeshIndexArray(
 
 
 void FGeometryScriptMeshSelection::ProcessByTriangleID(const UE::Geometry::FDynamicMesh3& Mesh,
-	TFunctionRef<void(int32)> PerTriangleFunc) const
+	TFunctionRef<void(int32)> PerTriangleFunc,
+	bool bProcessAllTrisIfSelectionEmpty) const
 {
+	if (IsEmpty() && bProcessAllTrisIfSelectionEmpty)
+	{
+		for (int32 TriangleID : Mesh.TriangleIndicesItr())
+		{
+			PerTriangleFunc(TriangleID);
+		}
+		return;
+	}
+
 	if (GetSelectionType() == EGeometryScriptMeshSelectionType::Triangles)
 	{
 		for (uint64 Item : GeoSelection->Selection)
@@ -338,8 +348,18 @@ void FGeometryScriptMeshSelection::ProcessByTriangleID(const UE::Geometry::FDyna
 
 
 void FGeometryScriptMeshSelection::ProcessByVertexID(const UE::Geometry::FDynamicMesh3& Mesh,
-	TFunctionRef<void(int32)> PerVertexFunc) const
+	TFunctionRef<void(int32)> PerVertexFunc,
+	bool bProcessAllVertsIfSelectionEmpty) const
 {
+	if (IsEmpty() && bProcessAllVertsIfSelectionEmpty)
+	{
+		for (int32 VertexID : Mesh.VertexIndicesItr())
+		{
+			PerVertexFunc(VertexID);
+		}
+		return;
+	}
+
 	if (GetSelectionType() == EGeometryScriptMeshSelectionType::Vertices)
 	{
 		for (uint64 Item : GeoSelection->Selection)
diff --git a/Engine/Plugins/Experimental/GeometryScripting/Source/GeometryScriptingCore/Private/MeshSelectionFunctions.cpp b/Engine/Plugins/Experimental/GeometryScripting/Source/GeometryScriptingCore/Private/MeshSelectionFunctions.cpp
--- a/Engine/Plugins/Experimental/GeometryScripting/Source/GeometryScriptingCore/Private/MeshSelectionFunctions.cpp
+++ b/Engine/Plugins/Experimental/GeometryScripting/Source/GeometryScriptingCore/Private/MeshSelectionFunctions.cpp
@@ -10,6 +10,8 @@
 #include "DynamicMesh/DynamicMeshAABBTree3.h"
 #include "Spatial/FastWinding.h"
 #include "Selections/MeshConnectedComponents.h"
+#include "Selections/MeshFaceSelection.h"
+#include "Selections/MeshVertexSelection.h"
 
 using namespace UE::Geometry;
 
@@ -809,4 +811,127 @@ UDynamicMesh* UGeometryScriptLibrary_MeshSelectionFunctions::ExpandMeshSelection
 
 
 
+UDynamicMesh* UGeometryScriptLibrary_MeshSelectionFunctions::ExpandContractMeshSelection(
+	UDynamicMesh* TargetMesh,
+	FGeometryScriptMeshSelection Selection,
+	FGeometryScriptMeshSelection& NewSelection,
+	int32 Iterations,
+	bool bContract,
+	bool bOnlyExpandToFaceNeighbours)
+{
+	if (TargetMesh == nullptr)
+	{
+		UE_LOG(LogGeometry, Warning, TEXT("ExpandContractMeshSelection: TargetMesh is Null"));
+		return TargetMesh;
+	}
+	if (Selection.IsEmpty())
+	{
+		UE_LOG(LogGeometry, Warning, TEXT("ExpandContractMeshSelection: Initial Selection is Empty"));
+		return TargetMesh;
+	}
+	if (Iterations <= 0)
+	{
+		NewSelection.SetSelection(Selection);
+		return TargetMesh;
+	}
+	Iterations = FMath::Clamp(Iterations, 1, 100);
+
+	// TODO: when doing multiple iterations w/ polygroups, we cannot rely on the code below because it is only
+	// expanding/contracting by triangle rings. Need to expand to polygroups at each step which is currently not easy
+	// to do with a FMeshFaceSelection, need to convert to FMeshConnectedComponents/etc. So for now we will just
+	// recursively do it this way which is expensive...
+	if (Selection.GetSelectionType() == EGeometryScriptMeshSelectionType::Polygroups && Iterations > 1)
+	{
+		FGeometryScriptMeshSelection CurSelection = Selection;
+		FGeometryScriptMeshSelection NextSelection;
+		for (int32 k = 0; k < Iterations; ++k)
+		{
+			ExpandContractMeshSelection(TargetMesh, CurSelection, NextSelection, 1, bContract, bOnlyExpandToFaceNeighbours);
+			CurSelection = NextSelection;
+		}
+		NewSelection = CurSelection;
+		return TargetMesh;
+	}
+
+
+	FGeometrySelection NewGeoSelection;
+	if (Selection.GetSelectionType() == EGeometryScriptMeshSelectionType::Vertices)
+	{
+		NewGeoSelection.InitializeTypes(EGeometryElementType::Vertex, EGeometryTopologyType::Triangle);
+		TargetMesh->ProcessMesh([&](const FDynamicMesh3& ReadMesh)
+		{
+			FMeshVertexSelection VtxSelection(&ReadMesh);
+			Selection.ProcessByVertexID(ReadMesh, [&](int32 VertexID) { VtxSelection.Select(VertexID); } );
+			if (bContract)
+			{
+				VtxSelection.ContractByBorderVertices(Iterations);
+			}
+			else
+			{
+				VtxSelection.ExpandToOneRingNeighbours(Iterations);
+			}
+			for (int32 VertexID : VtxSelection)
+			{
+				NewGeoSelection.Selection.Add(FGeoSelectionID::MeshVertex(VertexID).Encoded());
+			}
+		});
+	}
+	else 
+	{
+		TargetMesh->ProcessMesh([&](const FDynamicMesh3& ReadMesh)
+		{
+			FMeshFaceSelection TriSelection(&ReadMesh);
+			Selection.ProcessByTriangleID(ReadMesh, [&](int32 TriangleID) { TriSelection.Select(TriangleID); });
+			if (bContract)
+			{
+				TriSelection.ContractBorderByOneRingNeighbours(Iterations, true);
+			}
+			else
+			{
+				if (bOnlyExpandToFaceNeighbours)
+				{
+					TriSelection.ExpandToFaceNeighbours(Iterations);
+				}
+				else
+				{
+					TriSelection.ExpandToOneRingNeighbours(Iterations);
+				}
+			}
+			if (Selection.GetSelectionType() == EGeometryScriptMeshSelectionType::Triangles)
+			{
+				NewGeoSelection.InitializeTypes(EGeometryElementType::Face, EGeometryTopologyType::Triangle);
+				for (int32 TriangleID : TriSelection)
+				{
+					NewGeoSelection.Selection.Add(FGeoSelectionID::MeshTriangle(TriangleID).Encoded());
+				}
+			}
+			else
+			{
+				NewGeoSelection.InitializeTypes(EGeometryElementType::Face, EGeometryTopologyType::Polygroup);
+				TargetMesh->ProcessMesh([&](const FDynamicMesh3& ReadMesh)
+				{
+					TSet<int32> UniqueGroupIDs;
+					for (int32 tid : TriSelection)
+					{
+						int32 GroupID = ReadMesh.GetTriangleGroup(tid);
+						if (UniqueGroupIDs.Contains(GroupID) == false)
+						{
+							NewGeoSelection.Selection.Add(FGeoSelectionID::GroupFace(tid, GroupID).Encoded());
+						}
+					}
+				});
+
+			}
+		});
+	}
+
+	NewSelection.SetSelection(MoveTemp(NewGeoSelection));
+
+	return TargetMesh;
+}
+
+
+
+
+
 #undef LOCTEXT_NAMESPACE
\ No newline at end of file
diff --git a/Engine/Plugins/Experimental/GeometryScripting/Source/GeometryScriptingCore/Public/GeometryScript/GeometryScriptSelectionTypes.h b/Engine/Plugins/Experimental/GeometryScripting/Source/GeometryScriptingCore/Public/GeometryScript/GeometryScriptSelectionTypes.h
--- a/Engine/Plugins/Experimental/GeometryScripting/Source/GeometryScriptingCore/Public/GeometryScript/GeometryScriptSelectionTypes.h
+++ b/Engine/Plugins/Experimental/GeometryScripting/Source/GeometryScriptingCore/Public/GeometryScript/GeometryScriptSelectionTypes.h
@@ -102,7 +102,8 @@ struct GEOMETRYSCRIPTINGCORE_API FGeometryScriptMeshSelection
 	 * For Polygroup Selections, a full mesh iteration is used to find all Triangles in the Groups.
 	 */
 	void ProcessByTriangleID(const UE::Geometry::FDynamicMesh3& Mesh,
-		TFunctionRef<void(int32)> PerTriangleFunc) const;
+		TFunctionRef<void(int32)> PerTriangleFunc,
+		bool bProcessAllTrisIfSelectionEmpty = false) const;
 
 	/**
 	 * Call PerVertexFunc for each VertexID in the Selection.
@@ -110,7 +111,8 @@ struct GEOMETRYSCRIPTINGCORE_API FGeometryScriptMeshSelection
 	 * For Polygroup Selections, a full mesh iteration is used to find all Triangle Vertices in the Groups (accumulated in a TSet)
 	 */
 	void ProcessByVertexID(const UE::Geometry::FDynamicMesh3& Mesh,
-		TFunctionRef<void(int32)> PerVertexFunc) const;
+		TFunctionRef<void(int32)> PerVertexFunc,
+		bool bProcessAllVertsIfSelectionEmpty = false) const;
 
 private:
 	// keeping this private for now in case it needs to be revised in 5.2+
diff --git a/Engine/Plugins/Experimental/GeometryScripting/Source/GeometryScriptingCore/Public/GeometryScript/MeshSelectionFunctions.h b/Engine/Plugins/Experimental/GeometryScripting/Source/GeometryScriptingCore/Public/GeometryScript/MeshSelectionFunctions.h
--- a/Engine/Plugins/Experimental/GeometryScripting/Source/GeometryScriptingCore/Public/GeometryScript/MeshSelectionFunctions.h
+++ b/Engine/Plugins/Experimental/GeometryScripting/Source/GeometryScriptingCore/Public/GeometryScript/MeshSelectionFunctions.h
@@ -160,7 +160,7 @@ public:
 	 * @param bInvert return a selection of all elements not in the Box
 	 * @param MinNumTrianglePoints number of vertices of a triangle that must be in the box for it to be selected (1,2, or 3)
 	 */
-	UFUNCTION(BlueprintCallable, Category = "GeometryScript|MeshSelection", meta=(ScriptMethod, AdvancedDisplay="NumPointsInBox"))
+	UFUNCTION(BlueprintCallable, Category = "GeometryScript|MeshSelection", meta=(ScriptMethod, AdvancedDisplay="MinNumTrianglePoints"))
 	static UPARAM(DisplayName = "Target Mesh") UDynamicMesh* 
 	SelectMeshElementsInBox(
 		UDynamicMesh* TargetMesh,
@@ -177,7 +177,7 @@ public:
 	* @param bInvert return a selection of all elements not in the Sphere
 	* @param MinNumTrianglePoints number of vertices of a triangle that must be in the Sphere for it to be selected (1,2, or 3)
 	*/
-	UFUNCTION(BlueprintCallable, Category = "GeometryScript|MeshSelection", meta=(ScriptMethod, AdvancedDisplay="NumPointsInBox"))
+	UFUNCTION(BlueprintCallable, Category = "GeometryScript|MeshSelection", meta=(ScriptMethod, AdvancedDisplay="MinNumTrianglePoints"))
 	static UPARAM(DisplayName = "Target Mesh") UDynamicMesh* 
 	SelectMeshElementsInSphere(
 		UDynamicMesh* TargetMesh,
@@ -195,7 +195,7 @@ public:
 	* @param bInvert return a selection of all elements on the other (negative) side of the Plane
 	* @param MinNumTrianglePoints number of vertices of a triangle that must be on the positive Plane side to be selected (1,2, or 3)
 	*/
-	UFUNCTION(BlueprintCallable, Category = "GeometryScript|MeshSelection", meta=(ScriptMethod, AdvancedDisplay="NumPointsInBox"))
+	UFUNCTION(BlueprintCallable, Category = "GeometryScript|MeshSelection", meta=(ScriptMethod, AdvancedDisplay="MinNumTrianglePoints"))
 	static UPARAM(DisplayName = "Target Mesh") UDynamicMesh* 
 	SelectMeshElementsWithPlane(
 		UDynamicMesh* TargetMesh,
@@ -215,7 +215,7 @@ public:
 	* @param bInvert return a selection of all elements not within the given deviation
 	* @param MinNumTrianglePoints number of vertices of a triangle that must be within the angular deviation for it to be selected (1,2, or 3)
 	*/
-	UFUNCTION(BlueprintCallable, Category = "GeometryScript|MeshSelection", meta=(ScriptMethod, AdvancedDisplay="NumPointsInBox"))
+	UFUNCTION(BlueprintCallable, Category = "GeometryScript|MeshSelection", meta=(ScriptMethod, AdvancedDisplay="MinNumTrianglePoints"))
 	static UPARAM(DisplayName = "Target Mesh") UDynamicMesh* 
 	SelectMeshElementsByNormalAngle(
 		UDynamicMesh* TargetMesh,
@@ -236,7 +236,7 @@ public:
 	* @param WindingThreshold Threshold used for Fast Mesh Winding Number inside/outside test (range is [0,1], with 1 being "inside")
 	* @param MinNumTrianglePoints number of vertices of a triangle that must be within the angular deviation for it to be selected (1,2, or 3)
 	*/
-	UFUNCTION(BlueprintCallable, Category = "GeometryScript|MeshSelection", meta=(ScriptMethod, AdvancedDisplay="NumPointsInBox"))
+	UFUNCTION(BlueprintCallable, Category = "GeometryScript|MeshSelection", meta=(ScriptMethod, AdvancedDisplay="MinNumTrianglePoints"))
 	static UPARAM(DisplayName = "Target Mesh") UDynamicMesh* 
 	SelectMeshElementsInsideMesh(
 		UDynamicMesh* TargetMesh,
@@ -253,7 +253,7 @@ public:
 	 * Expand the Selection on the TargetMesh to connected regions and return the NewSelection
 	 * @param ConnectionType defines what "connected" means, ie purely geometric connection, or some additional constraint like same MaterialIDs/etc 
 	 */
-	UFUNCTION(BlueprintCallable, Category = "GeometryScript|MeshSelection", meta=(ScriptMethod, AdvancedDisplay="NumPointsInBox"))
+	UFUNCTION(BlueprintCallable, Category = "GeometryScript|MeshSelection", meta=(ScriptMethod))
 	static UPARAM(DisplayName = "Target Mesh") UDynamicMesh* 
 	ExpandMeshSelectionToConnected(
 		UDynamicMesh* TargetMesh,
@@ -262,4 +262,26 @@ public:
 		EGeometryScriptTopologyConnectionType ConnectionType = EGeometryScriptTopologyConnectionType::Geometric
 	);
 
+
+	/**
+	 * Grow or Shrink the Selection on the TargetMesh to connected neighbours.
+	 * For Vertex selections, Expand includes vertices in one-ring of selected vertices, and Contract removes any vertices with a one-ring neighbour that is not selected
+	 * For Triangle selections, Add/Remove Triangles connected to selected Triangles
+	 * For Polygroup selections, Add/Remove Polygroups connected to selected Polygroups
+	 * @param Iterations number of times to Expand/Contract the Selection. Valid range is [0,100] where 0 is a no-op.
+	 * @param bContract if true selection contracts instead of growing
+	 * @param bOnlyExpandToFaceNeighbours if true, only adjacent Triangles/Polygroups directly connected by an edge are added, vs connected to any selected vertex
+	 */
+	UFUNCTION(BlueprintCallable, Category = "GeometryScript|MeshSelection", meta=(ScriptMethod))
+	static UPARAM(DisplayName = "Target Mesh") UDynamicMesh* 
+	ExpandContractMeshSelection(
+		UDynamicMesh* TargetMesh,
+		FGeometryScriptMeshSelection Selection, 
+		FGeometryScriptMeshSelection& NewSelection,
+		int32 Iterations = 1,
+		bool bContract = false,
+		bool bOnlyExpandToFaceNeighbours = false
+	);
+
+
 };
\ No newline at end of file
diff --git a/Engine/Plugins/Runtime/GeometryProcessing/Source/DynamicMesh/Public/Selections/MeshVertexSelection.h b/Engine/Plugins/Runtime/GeometryProcessing/Source/DynamicMesh/Public/Selections/MeshVertexSelection.h
--- a/Engine/Plugins/Runtime/GeometryProcessing/Source/DynamicMesh/Public/Selections/MeshVertexSelection.h
+++ b/Engine/Plugins/Runtime/GeometryProcessing/Source/DynamicMesh/Public/Selections/MeshVertexSelection.h
@@ -210,7 +210,8 @@ public:
 		remove(ev.A); remove(ev.B);
 	}
 
-	void DeselectEdges(TArrayView<const int> Edges) {
+	void DeselectEdges(TArrayView<const int> Edges) 
+	{
 		for (int EID : Edges)
 		{
 			FIndex2i ev = Mesh->GetEdgeV(EID);
@@ -230,7 +231,8 @@ public:
 	{
 		TArray<int> temp;
 
-		for (int vid : Selected) {
+		for (int vid : Selected) 
+		{
 			for (int nbr_vid : Mesh->VtxVerticesItr(vid))
 			{
 				if (FilterF && FilterF(nbr_vid) == false)
@@ -261,6 +263,37 @@ public:
 	}
 
 
+	/**
+	 * Remove all vertices in current selection set that have at least
+	 * one neighbour vertex that is not selected (ie vertices are on border of selection)
+	 */
+	void ContractByBorderVertices(int32 nRings = 1)
+	{
+		// find set of boundary vertices
+		TArray<int> BorderVertices;
+		for (int32 k = 0; k < nRings; ++k)
+		{
+			BorderVertices.Reset();
+
+			for (int vid : Selected) 
+			{
+				bool bAnyNeighbourDeselected = false;
+				for (int nbr_vid : Mesh->VtxVerticesItr(vid))
+				{
+					if (IsSelected(nbr_vid) == false)
+					{
+						bAnyNeighbourDeselected = true;
+						break;
+					}
+				}
+				if (bAnyNeighbourDeselected)
+				{
+					BorderVertices.Add(vid);
+				}
+			}
+			Deselect(BorderVertices);
+		}
+	}
 
 
 
