GeometryScript: add support for optional SetMeshUVsFromBoxProjection, Selection argument to TranslateMeshUVs, ScaleMeshUVs, RotateMeshUVs, SetMeshUVsFromPlanarProjection, SetMeshUVsFromBoxProjection, SetMeshUVsFromCylinderProjection, RecomputeMeshUVs. Add new GetMeshUVSizeInfo function that can be used to get UV Area/Bounds, and 3D Area/Bounds, for a given selection. All these functions fall back to using the full mesh if no Selection is provided.

diff --git a/Engine/Plugins/Experimental/GeometryScripting/Source/GeometryScriptingCore/Private/MeshUVFunctions.cpp b/Engine/Plugins/Experimental/GeometryScripting/Source/GeometryScriptingCore/Private/MeshUVFunctions.cpp
--- a/Engine/Plugins/Experimental/GeometryScripting/Source/GeometryScriptingCore/Private/MeshUVFunctions.cpp
+++ b/Engine/Plugins/Experimental/GeometryScripting/Source/GeometryScriptingCore/Private/MeshUVFunctions.cpp
@@ -64,7 +64,7 @@ UDynamicMesh* UGeometryScriptLibrary_MeshUVFunctions::CopyUVSet(
 	}
 	if (FromUVSet == ToUVSet)
 	{
-		UE::Geometry::AppendError(Debug, EGeometryScriptErrorType::InvalidInputs, LOCTEXT("SetNumUVSets_SameSet", "SetNumUVSets: From and To UV Sets have the same Index"));
+		UE::Geometry::AppendError(Debug, EGeometryScriptErrorType::InvalidInputs, LOCTEXT("CopyUVSet_SameSet", "CopyUVSet: From and To UV Sets have the same Index"));
 		return TargetMesh;
 	}
 	TargetMesh->EditMesh([&](FDynamicMesh3& EditMesh)
@@ -147,12 +147,49 @@ void ApplyMeshUVEditorOperation(UDynamicMesh* TargetMesh, int32 UVSetIndex, bool
 }
 
 
+namespace UELocal
+{
+void ApplyUVTransform(
+	FDynamicMesh3& EditMesh,
+	FDynamicMeshUVOverlay& UVOverlay,
+	FGeometryScriptMeshSelection& Selection,
+	TFunctionRef<FVector2f(FVector2f)> UVTransformFunc
+)
+{
+	if (Selection.IsEmpty())
+	{
+		for (int32 ElementID : UVOverlay.ElementIndicesItr())
+		{
+			FVector2f UV = UVOverlay.GetElement(ElementID);
+			UVOverlay.SetElement(ElementID, UVTransformFunc(UV));
+		}
+	}
+	else
+	{
+		TSet<int32> ElementSet;
+		Selection.ProcessByTriangleID(EditMesh, [&](int32 TriangleID)
+		{
+			if (UVOverlay.IsSetTriangle(TriangleID))
+			{
+				FIndex3i TriElems = UVOverlay.GetTriangle(TriangleID);
+				ElementSet.Add(TriElems.A); ElementSet.Add(TriElems.B); ElementSet.Add(TriElems.C);
+			}
+		});
+		for (int32 ElementID : ElementSet)
+		{
+			FVector2f UV = UVOverlay.GetElement(ElementID);
+			UVOverlay.SetElement(ElementID, UVTransformFunc(UV));
+		}
+	}
+}
+}
 
 
 UDynamicMesh* UGeometryScriptLibrary_MeshUVFunctions::TranslateMeshUVs(
 	UDynamicMesh* TargetMesh,
 	int UVSetIndex,
 	FVector2D Translation,
+	FGeometryScriptMeshSelection Selection,
 	UGeometryScriptDebug* Debug)
 {
 	if (TargetMesh == nullptr)
@@ -165,11 +202,8 @@ UDynamicMesh* UGeometryScriptLibrary_MeshUVFunctions::TranslateMeshUVs(
 	ApplyMeshUVEditorOperation(TargetMesh, UVSetIndex, bHasUVSet, Debug,
 		[&](FDynamicMesh3& EditMesh, FDynamicMeshUVOverlay* UVOverlay, FDynamicMeshUVEditor& UVEditor)
 	{
-		for (int32 elemid : UVOverlay->ElementIndicesItr())
-		{
-			FVector2f UV = UVOverlay->GetElement(elemid);
-			UVOverlay->SetElement(elemid, UV + (FVector2f)Translation);
-		}
+		UELocal::ApplyUVTransform(EditMesh, *UVOverlay, Selection,
+			[&](FVector2f UV) { return UV + (FVector2f)Translation; });
 	});
 	if (bHasUVSet == false)
 	{
@@ -185,6 +219,7 @@ UDynamicMesh* UGeometryScriptLibrary_MeshUVFunctions::ScaleMeshUVs(
 	int UVSetIndex,
 	FVector2D Scale,
 	FVector2D ScaleOrigin,
+	FGeometryScriptMeshSelection Selection,
 	UGeometryScriptDebug* Debug)
 {
 	if (TargetMesh == nullptr)
@@ -204,12 +239,8 @@ UDynamicMesh* UGeometryScriptLibrary_MeshUVFunctions::ScaleMeshUVs(
 	ApplyMeshUVEditorOperation(TargetMesh, UVSetIndex, bHasUVSet, Debug,
 		[&](FDynamicMesh3& EditMesh, FDynamicMeshUVOverlay* UVOverlay, FDynamicMeshUVEditor& UVEditor)
 	{
-		for (int32 elemid : UVOverlay->ElementIndicesItr())
-		{
-			FVector2f UV = UVOverlay->GetElement(elemid);
-			UV = (UV - UseOrigin) * UseScale + UseOrigin;
-			UVOverlay->SetElement(elemid, UV);
-		}
+		UELocal::ApplyUVTransform(EditMesh, *UVOverlay, Selection,
+			[&](FVector2f UV) { return (UV - UseOrigin) * UseScale + UseOrigin; });
 	});
 	if (bHasUVSet == false)
 	{
@@ -226,6 +257,7 @@ UDynamicMesh* UGeometryScriptLibrary_MeshUVFunctions::RotateMeshUVs(
 	int UVSetIndex,
 	float RotationAngle,
 	FVector2D RotationOrigin,
+	FGeometryScriptMeshSelection Selection,
 	UGeometryScriptDebug* Debug)
 {
 	if (TargetMesh == nullptr)
@@ -241,12 +273,8 @@ UDynamicMesh* UGeometryScriptLibrary_MeshUVFunctions::RotateMeshUVs(
 	ApplyMeshUVEditorOperation(TargetMesh, UVSetIndex, bHasUVSet, Debug,
 		[&](FDynamicMesh3& EditMesh, FDynamicMeshUVOverlay* UVOverlay, FDynamicMeshUVEditor& UVEditor)
 	{
-		for (int32 elemid : UVOverlay->ElementIndicesItr())
-		{
-			FVector2f UV = UVOverlay->GetElement(elemid);
-			UV = RotationMat* (UV - UseOrigin) + UseOrigin;
-			UVOverlay->SetElement(elemid, UV);
-		}
+		UELocal::ApplyUVTransform(EditMesh, *UVOverlay, Selection,
+			[&](FVector2f UV) { return RotationMat * (UV - UseOrigin) + UseOrigin; });
 	});
 	if (bHasUVSet == false)
 	{
@@ -263,6 +291,7 @@ UDynamicMesh* UGeometryScriptLibrary_MeshUVFunctions::SetMeshUVsFromPlanarProjec
 	UDynamicMesh* TargetMesh,
 	int UVSetIndex,
 	FTransform PlaneTransform,
+	FGeometryScriptMeshSelection Selection,
 	UGeometryScriptDebug* Debug)
 {
 	if (TargetMesh == nullptr)
@@ -275,17 +304,14 @@ UDynamicMesh* UGeometryScriptLibrary_MeshUVFunctions::SetMeshUVsFromPlanarProjec
 	ApplyMeshUVEditorOperation(TargetMesh, UVSetIndex, bHasUVSet, Debug,
 		[&](FDynamicMesh3& EditMesh, FDynamicMeshUVOverlay* UVOverlay, FDynamicMeshUVEditor& UVEditor)
 	{
-		TArray<int32> AllTriangles;
-		for (int32 tid : EditMesh.TriangleIndicesItr())
-		{
-			AllTriangles.Add(tid);
-		}
+		TArray<int32> TriangleROI;
+		Selection.ProcessByTriangleID(EditMesh, [&](int32 TriangleID) { TriangleROI.Add(TriangleID); }, true);
 
 		FFrame3d ProjectionFrame(PlaneTransform);
 		FVector Scale = PlaneTransform.GetScale3D();
 		FVector2d Dimensions(Scale.X, Scale.Y);
 
-		UVEditor.SetTriangleUVsFromPlanarProjection(AllTriangles, [&](const FVector3d& Pos) { return Pos; },
+		UVEditor.SetTriangleUVsFromPlanarProjection(TriangleROI, [&](const FVector3d& Pos) { return Pos; },
 			ProjectionFrame, Dimensions);
 	});
 	if (bHasUVSet == false)
@@ -303,6 +329,7 @@ UDynamicMesh* UGeometryScriptLibrary_MeshUVFunctions::SetMeshUVsFromBoxProjectio
 	UDynamicMesh* TargetMesh,
 	int UVSetIndex,
 	FTransform PlaneTransform,
+	FGeometryScriptMeshSelection Selection,
 	int MinIslandTriCount,
 	UGeometryScriptDebug* Debug)
 {
@@ -316,16 +343,13 @@ UDynamicMesh* UGeometryScriptLibrary_MeshUVFunctions::SetMeshUVsFromBoxProjectio
 	ApplyMeshUVEditorOperation(TargetMesh, UVSetIndex, bHasUVSet, Debug,
 		[&](FDynamicMesh3& EditMesh, FDynamicMeshUVOverlay* UVOverlay, FDynamicMeshUVEditor& UVEditor)
 	{
-		TArray<int32> AllTriangles;
-		for (int32 tid : EditMesh.TriangleIndicesItr())
-		{
-			AllTriangles.Add(tid);
-		}
+		TArray<int32> TriangleROI;
+		Selection.ProcessByTriangleID(EditMesh, [&](int32 TriangleID) { TriangleROI.Add(TriangleID); }, true);
 
 		FFrame3d ProjectionFrame(PlaneTransform);
 		FVector Scale = PlaneTransform.GetScale3D();
 		FVector3d Dimensions = (FVector)Scale;
-		UVEditor.SetTriangleUVsFromBoxProjection(AllTriangles, [&](const FVector3d& Pos) { return Pos; },
+		UVEditor.SetTriangleUVsFromBoxProjection(TriangleROI, [&](const FVector3d& Pos) { return Pos; },
 			ProjectionFrame, Dimensions, MinIslandTriCount);
 	});
 	if (bHasUVSet == false)
@@ -344,6 +368,7 @@ UDynamicMesh* UGeometryScriptLibrary_MeshUVFunctions::SetMeshUVsFromCylinderProj
 	UDynamicMesh* TargetMesh,
 	int UVSetIndex,
 	FTransform CylinderTransform,
+	FGeometryScriptMeshSelection Selection,
 	float SplitAngle,
 	UGeometryScriptDebug* Debug)
 {
@@ -357,16 +382,13 @@ UDynamicMesh* UGeometryScriptLibrary_MeshUVFunctions::SetMeshUVsFromCylinderProj
 	ApplyMeshUVEditorOperation(TargetMesh, UVSetIndex, bHasUVSet, Debug,
 		[&](FDynamicMesh3& EditMesh, FDynamicMeshUVOverlay* UVOverlay, FDynamicMeshUVEditor& UVEditor)
 	{
-		TArray<int32> AllTriangles;
-		for (int32 tid : EditMesh.TriangleIndicesItr())
-		{
-			AllTriangles.Add(tid);
-		}
+		TArray<int32> TriangleROI;
+		Selection.ProcessByTriangleID(EditMesh, [&](int32 TriangleID) { TriangleROI.Add(TriangleID); }, true);
 
 		FFrame3d ProjectionFrame(CylinderTransform);
 		FVector Scale = CylinderTransform.GetScale3D();
 		FVector3d Dimensions = (FVector)Scale;
-		UVEditor.SetTriangleUVsFromCylinderProjection(AllTriangles, [&](const FVector3d& Pos) { return Pos; },
+		UVEditor.SetTriangleUVsFromCylinderProjection(TriangleROI, [&](const FVector3d& Pos) { return Pos; },
 			ProjectionFrame, Dimensions, SplitAngle);
 	});
 	if (bHasUVSet == false)
@@ -383,6 +405,7 @@ UDynamicMesh* UGeometryScriptLibrary_MeshUVFunctions::RecomputeMeshUVs(
 	UDynamicMesh* TargetMesh, 
 	int UVSetIndex,
 	FGeometryScriptRecomputeUVsOptions Options,
+	FGeometryScriptMeshSelection Selection,
 	UGeometryScriptDebug* Debug)
 {
 	if (TargetMesh == nullptr)
@@ -395,10 +418,8 @@ UDynamicMesh* UGeometryScriptLibrary_MeshUVFunctions::RecomputeMeshUVs(
 	ApplyMeshUVEditorOperation(TargetMesh, UVSetIndex, bHasUVSet, Debug,
 		[&](FDynamicMesh3& EditMesh, FDynamicMeshUVOverlay* UVOverlay, FDynamicMeshUVEditor& UVEditor)
 	{
-
-
 		TUniquePtr<FPolygroupSet> IslandSourceGroups;
-		if (Options.IslandSource == EGeometryScriptUVIslandSource::PolyGroups)
+		if (Options.IslandSource == EGeometryScriptUVIslandSource::PolyGroups && Selection.IsEmpty())
 		{
 			FPolygroupLayer InputGroupLayer{ Options.GroupLayer.bDefaultLayer, Options.GroupLayer.ExtendedLayerIndex };
 			if (InputGroupLayer.CheckExists(&EditMesh))
@@ -415,15 +436,26 @@ UDynamicMesh* UGeometryScriptLibrary_MeshUVFunctions::RecomputeMeshUVs(
 
 		// find group-connected-components
 		FMeshConnectedComponents ConnectedComponents(&EditMesh);
-		if (Options.IslandSource == EGeometryScriptUVIslandSource::PolyGroups)
+		if (Selection.IsEmpty())
 		{
-			ConnectedComponents.FindConnectedTriangles([&](int32 CurTri, int32 NbrTri) {
-				return IslandSourceGroups->GetTriangleGroup(CurTri) == IslandSourceGroups->GetTriangleGroup(NbrTri);
-			});
+			if (Options.IslandSource == EGeometryScriptUVIslandSource::PolyGroups)
+			{
+				ConnectedComponents.FindConnectedTriangles([&](int32 CurTri, int32 NbrTri) {
+					return IslandSourceGroups->GetTriangleGroup(CurTri) == IslandSourceGroups->GetTriangleGroup(NbrTri);
+				});
+			}
+			else
+			{
+				ConnectedComponents.FindConnectedTriangles([&](int32 Triangle0, int32 Triangle1) {
+					return UVOverlay->AreTrianglesConnected(Triangle0, Triangle1);
+				});
+			}
 		}
 		else
 		{
-			ConnectedComponents.FindConnectedTriangles([&](int32 Triangle0, int32 Triangle1) {
+			TArray<int32> TriangleROI;
+			Selection.ConvertToMeshIndexArray(EditMesh, TriangleROI, EGeometryScriptIndexType::Triangle);
+			ConnectedComponents.FindConnectedTriangles(TriangleROI, [&](int32 Triangle0, int32 Triangle1) {
 				return UVOverlay->AreTrianglesConnected(Triangle0, Triangle1);
 			});
 		}
@@ -714,6 +746,75 @@ UDynamicMesh* UGeometryScriptLibrary_MeshUVFunctions::AutoGenerateXAtlasMeshUVs(
 
 
 
+UDynamicMesh* UGeometryScriptLibrary_MeshUVFunctions::GetMeshUVSizeInfo(
+	UDynamicMesh* TargetMesh,
+	int UVSetIndex,
+	FGeometryScriptMeshSelection Selection,
+	double& MeshArea,
+	double& UVArea,
+	FBox& MeshBounds,
+	FBox2D& UVBounds,
+	bool& bIsValidUVSet,
+	bool& bFoundUnsetUVs,
+	bool bOnlyIncludeValidUVTris,
+	UGeometryScriptDebug* Debug)
+{
+	MeshArea = UVArea = 0;
+	bIsValidUVSet = false;
+	bFoundUnsetUVs = false;
+	if (TargetMesh == nullptr)
+	{
+		UE::Geometry::AppendError(Debug, EGeometryScriptErrorType::InvalidInputs, LOCTEXT("GetMeshUVSizeInfo_InvalidInput", "GetMeshUVSizeInfo: TargetMesh is Null"));
+		return TargetMesh;
+	}
+	FAxisAlignedBox3d MeshBoundsTmp = FAxisAlignedBox3d::Empty();
+	FAxisAlignedBox2f UVBoundsTmp = FAxisAlignedBox2f::Empty();
+	TargetMesh->ProcessMesh([&](const FDynamicMesh3& ReadMesh)
+	{
+		const FDynamicMeshUVOverlay* UVOverlay = nullptr;
+		if (ReadMesh.HasAttributes())
+		{
+			UVOverlay = (UVSetIndex >= 0 && UVSetIndex < ReadMesh.Attributes()->NumUVLayers()) ? ReadMesh.Attributes()->GetUVLayer(UVSetIndex) : nullptr;
+		}
+		if (UVOverlay == nullptr)
+		{
+			UE::Geometry::AppendError(Debug, EGeometryScriptErrorType::InvalidInputs, LOCTEXT("GetMeshUVSizeInfo_InvalidUVSet", "GetMeshUVSizeInfo: UV Set does not Exist"));
+			return;
+		}
+		bIsValidUVSet = true;
+
+		Selection.ProcessByTriangleID(ReadMesh, [&](int32 TriangleID) { 
+		
+			bool bTriangleHasUVs = (UVOverlay->IsSetTriangle(TriangleID));
+			if (!bTriangleHasUVs)
+			{
+				bFoundUnsetUVs = true;
+			}
+
+			FVector3d Vertices[3];
+			FVector2f UVs[3];
+			if (bTriangleHasUVs || bOnlyIncludeValidUVTris == false)
+			{
+				ReadMesh.GetTriVertices(TriangleID, Vertices[0], Vertices[1], Vertices[2]);
+				MeshArea += VectorUtil::Area(Vertices[0], Vertices[1], Vertices[2]);
+				MeshBoundsTmp.Contain(Vertices[0]); MeshBoundsTmp.Contain(Vertices[1]); MeshBoundsTmp.Contain(Vertices[2]);
+			}
+			if (bTriangleHasUVs)
+			{
+				UVOverlay->GetTriElements(TriangleID, UVs[0], UVs[1], UVs[2]);
+				UVArea += (double)VectorUtil::Area(UVs[0], UVs[1], UVs[2]);
+				UVBoundsTmp.Contain(UVs[0]); UVBoundsTmp.Contain(UVs[1]); UVBoundsTmp.Contain(UVs[2]);
+			}
+		
+		}, true);
+	});
+
+	MeshBounds = (FBox)MeshBoundsTmp;
+	UVBounds = (FBox2D)UVBoundsTmp;
+	
+	return TargetMesh;
+}
+
 
 
 
diff --git a/Engine/Plugins/Experimental/GeometryScripting/Source/GeometryScriptingCore/Public/GeometryScript/MeshUVFunctions.h b/Engine/Plugins/Experimental/GeometryScripting/Source/GeometryScriptingCore/Public/GeometryScript/MeshUVFunctions.h
--- a/Engine/Plugins/Experimental/GeometryScripting/Source/GeometryScriptingCore/Public/GeometryScript/MeshUVFunctions.h
+++ b/Engine/Plugins/Experimental/GeometryScripting/Source/GeometryScriptingCore/Public/GeometryScript/MeshUVFunctions.h
@@ -5,6 +5,7 @@
 #include "CoreMinimal.h"
 #include "Kismet/BlueprintFunctionLibrary.h"
 #include "GeometryScript/GeometryScriptTypes.h"
+#include "GeometryScript/GeometryScriptSelectionTypes.h"
 #include "MeshUVFunctions.generated.h"
 
 class UDynamicMesh;
@@ -178,6 +179,7 @@ public:
 		UDynamicMesh* TargetMesh, 
 		int UVSetIndex,
 		FVector2D Translation,
+		FGeometryScriptMeshSelection Selection,
 		UGeometryScriptDebug* Debug = nullptr );
 
 	UFUNCTION(BlueprintCallable, Category = "GeometryScript|UVs", meta=(ScriptMethod))
@@ -187,6 +189,7 @@ public:
 		int UVSetIndex,
 		FVector2D Scale,
 		FVector2D ScaleOrigin,
+		FGeometryScriptMeshSelection Selection,
 		UGeometryScriptDebug* Debug = nullptr );
 
 	UFUNCTION(BlueprintCallable, Category = "GeometryScript|UVs", meta=(ScriptMethod))
@@ -196,6 +199,7 @@ public:
 		int UVSetIndex,
 		float RotationAngle,
 		FVector2D RotationOrigin,
+		FGeometryScriptMeshSelection Selection,
 		UGeometryScriptDebug* Debug = nullptr );
 
 
@@ -208,6 +212,7 @@ public:
 		UDynamicMesh* TargetMesh, 
 		int UVSetIndex,
 		FTransform PlaneTransform,
+		FGeometryScriptMeshSelection Selection,
 		UGeometryScriptDebug* Debug = nullptr );
 
 
@@ -217,6 +222,7 @@ public:
 		UDynamicMesh* TargetMesh, 
 		int UVSetIndex,
 		FTransform BoxTransform,
+		FGeometryScriptMeshSelection Selection,
 		int MinIslandTriCount = 2,
 		UGeometryScriptDebug* Debug = nullptr );
 
@@ -227,6 +233,7 @@ public:
 		UDynamicMesh* TargetMesh, 
 		int UVSetIndex,
 		FTransform CylinderTransform,
+		FGeometryScriptMeshSelection Selection,
 		float SplitAngle = 45.0,
 		UGeometryScriptDebug* Debug = nullptr );
 
@@ -238,8 +245,10 @@ public:
 		UDynamicMesh* TargetMesh, 
 		int UVSetIndex,
 		FGeometryScriptRecomputeUVsOptions Options,
+		FGeometryScriptMeshSelection Selection,
 		UGeometryScriptDebug* Debug = nullptr );
 
+
 	UFUNCTION(BlueprintCallable, Category = "GeometryScript|UVs", meta=(ScriptMethod))
 	static UPARAM(DisplayName = "Target Mesh") UDynamicMesh* 
 	RepackMeshUVs( 
@@ -264,6 +273,33 @@ public:
 		FGeometryScriptXAtlasOptions Options,
 		UGeometryScriptDebug* Debug = nullptr );
 
+	/**
+	 * Compute information about dimensions and areas for a UV Set of a Mesh, with an optional Mesh Selection
+	 * @param UVSetIndex index of UV Set to query
+	 * @param Selection subset of triangles to process, whole mesh is used if selection is not provided
+	 * @param MeshArea output 3D area of queried triangles
+	 * @param UVArea output 2D UV-space area of queried triangles
+	 * @param MeshBounds output 3D bounding box of queried triangles
+	 * @param UVBounds output 2D UV-space bounding box of queried triangles
+	 * @param bIsValidUVSet output flag set to false if UVSetIndex does not exist on the target mesh. In this case Areas and Bounds are not initialized.
+	 * @param bFoundUnsetUVs output flag set to true if any of the queried triangles do not have valid UVs set
+	 * @param bOnlyIncludeValidUVTris if true, only triangles with valid UVs are included in 3D Mesh Area/Bounds
+	 */
+	UFUNCTION(BlueprintCallable, Category = "GeometryScript|UVs", meta=(ScriptMethod))
+	static UPARAM(DisplayName = "Copy From Mesh") UDynamicMesh* 
+	GetMeshUVSizeInfo(  
+		UDynamicMesh* TargetMesh, 
+		int UVSetIndex,
+		FGeometryScriptMeshSelection Selection,
+		double& MeshArea,
+		double& UVArea,
+		FBox& MeshBounds,
+		FBox2D& UVBounds,
+		bool& bIsValidUVSet,
+		bool& bFoundUnsetUVs,
+		bool bOnlyIncludeValidUVTris = true,
+		UGeometryScriptDebug* Debug = nullptr);	
+
 
 	/**
 	 * Get a list of single vertex UVs for each mesh vertex in the TargetMesh, derived from the specified UV Overlay.
@@ -287,7 +323,6 @@ public:
 		UGeometryScriptDebug* Debug = nullptr);
 
 
-
 	/**
 	 * Copy the 2D UVs from the given UVSetIndex in CopyFromMesh to the 3D vertex positions in CopyToUVMesh,
 	 * with the triangle mesh topolgoy defined by the UV Set. Generally this "UV Mesh" topolgoy will not
