DynamicMesh: fix VertexNormals mode in FOffsetMeshRegion. It was not properly handling the case where input mesh did not have VertexNormals attribute, the FMeshNormals being computed and passed to ::ApplyOffset() were never used.
GeometryScript: add ApplyMeshOffsetFaces and ApplyMeshInsetOutsetFaces

diff --git a/Engine/Plugins/Experimental/GeometryScripting/Source/GeometryScriptingCore/Private/MeshModelingFunctions.cpp b/Engine/Plugins/Experimental/GeometryScripting/Source/GeometryScriptingCore/Private/MeshModelingFunctions.cpp
--- a/Engine/Plugins/Experimental/GeometryScripting/Source/GeometryScriptingCore/Private/MeshModelingFunctions.cpp
+++ b/Engine/Plugins/Experimental/GeometryScripting/Source/GeometryScriptingCore/Private/MeshModelingFunctions.cpp
@@ -14,6 +14,7 @@
 #include "GroupTopology.h"
 #include "DynamicMesh/DynamicMeshAABBTree3.h"
 #include "Operations/OffsetMeshRegion.h"
+#include "Operations/InsetMeshRegion.h"
 #include "Operations/MeshBevel.h"
 #include "UDynamicMesh.h"
 #include "GeometryScript/MeshSelectionFunctions.h"
@@ -588,7 +589,7 @@ UDynamicMesh* UGeometryScriptLibrary_MeshModelingFunctions::ApplyMeshLinearExtru
 			Selection.ConvertToMeshIndexArray(EditMesh, Triangles, EGeometryScriptIndexType::Triangle);
 		}
 
-		if (Options.AreaMode == EGeometryScriptLinearExtrudeArea::PerPolygroup)
+		if (Options.AreaMode == EGeometryScriptPolyOperationArea::PerPolygroup)
 		{
 			TMap<int32, TArray<int32>> PolygroupTris;
 			for (int32 tid : Triangles)
@@ -612,7 +613,7 @@ UDynamicMesh* UGeometryScriptLibrary_MeshModelingFunctions::ApplyMeshLinearExtru
 				ApplyExtrudeArea(MoveTemp(Pair.Value));
 			}
 		}
-		else if (Options.AreaMode == EGeometryScriptLinearExtrudeArea::PerTriangle)
+		else if (Options.AreaMode == EGeometryScriptPolyOperationArea::PerTriangle)
 		{
 			for (int32 tid : Triangles)
 			{
@@ -635,6 +636,206 @@ UDynamicMesh* UGeometryScriptLibrary_MeshModelingFunctions::ApplyMeshLinearExtru
 
 
 
+UDynamicMesh* UGeometryScriptLibrary_MeshModelingFunctions::ApplyMeshOffsetFaces(
+	UDynamicMesh* TargetMesh,
+	FGeometryScriptMeshOffsetFacesOptions Options,
+	FGeometryScriptMeshSelection Selection,
+	UGeometryScriptDebug* Debug)
+{
+	if (TargetMesh == nullptr)
+	{
+		UE::Geometry::AppendError(Debug, EGeometryScriptErrorType::InvalidInputs, LOCTEXT("ApplyMeshOffsetFaces_InvalidInput", "ApplyMeshOffsetFaces: TargetMesh is Null"));
+		return TargetMesh;
+	}
+
+	TargetMesh->EditMesh([&](FDynamicMesh3& EditMesh) 
+	{
+		auto ApplyOffsetArea = [&EditMesh, &Options](TArray<int32>&& Triangles)
+		{
+			FOffsetMeshRegion Extruder(&EditMesh);
+			Extruder.Triangles = MoveTemp(Triangles);
+			if (Extruder.Triangles.Num() > 0)
+			{
+				Extruder.OffsetPositionFunc = [Options](const FVector3d& Pos, const FVector3d& VertexVector, int32 VertexID) {
+					return Pos + Options.Distance * VertexVector;
+				};
+
+				if (Options.OffsetType == EGeometryScriptOffsetFacesType::VertexNormal)
+				{
+					Extruder.ExtrusionVectorType = FOffsetMeshRegion::EVertexExtrusionVectorType::VertexNormal;
+				}
+				else if (Options.OffsetType == EGeometryScriptOffsetFacesType::FaceNormal)
+				{
+					Extruder.ExtrusionVectorType = FOffsetMeshRegion::EVertexExtrusionVectorType::SelectionTriNormalsAngleWeightedAverage;
+				}
+				else
+				{
+					Extruder.ExtrusionVectorType = FOffsetMeshRegion::EVertexExtrusionVectorType::SelectionTriNormalsAngleWeightedAdjusted;
+				}
+
+				Extruder.DefaultOffsetDistance = Options.Distance;
+				Extruder.bIsPositiveOffset = (Options.Distance > 0);
+				Extruder.UVScaleFactor = Options.UVScale;
+				Extruder.bOffsetFullComponentsAsSolids = Options.bSolidsToShells;
+				Extruder.Apply();
+				for (FOffsetMeshRegion::FOffsetInfo& OffsetInfo : Extruder.OffsetRegions)
+				{
+					UELocal::ApplyPolygroupModeToNewTriangles(EditMesh, OffsetInfo.OffsetTids, Options.GroupOptions);
+				}
+			}
+		};
+
+		TArray<int32> Triangles;
+		if (Selection.GetNumSelected() == 0)
+		{
+			for (int32 tid : EditMesh.TriangleIndicesItr())
+			{
+				Triangles.Add(tid);
+			}
+		}
+		else
+		{
+			Selection.ConvertToMeshIndexArray(EditMesh, Triangles, EGeometryScriptIndexType::Triangle);
+		}
+
+		if (Options.AreaMode == EGeometryScriptPolyOperationArea::PerPolygroup)
+		{
+			TMap<int32, TArray<int32>> PolygroupTris;
+			for (int32 tid : Triangles)
+			{
+				int32 gid = EditMesh.GetTriangleGroup(tid);
+				TArray<int32>* Found = PolygroupTris.Find(gid);
+				if (Found != nullptr)
+				{
+					Found->Add(tid);
+				}
+				else
+				{
+					TArray<int32> Temp;
+					Temp.Add(tid);
+					PolygroupTris.Add(gid, Temp);
+				}
+			}
+
+			for (TPair<int32, TArray<int32>> Pair : PolygroupTris)
+			{
+				ApplyOffsetArea(MoveTemp(Pair.Value));
+			}
+		}
+		else if (Options.AreaMode == EGeometryScriptPolyOperationArea::PerTriangle)
+		{
+			for (int32 tid : Triangles)
+			{
+				TArray<int32> Temp;
+				Temp.Add(tid);
+				ApplyOffsetArea(MoveTemp(Temp));
+			}
+		}
+		else
+		{
+			ApplyOffsetArea(MoveTemp(Triangles));
+		}
+
+	}, EDynamicMeshChangeType::GeneralEdit, EDynamicMeshAttributeChangeFlags::Unknown, false);
+
+	return TargetMesh;
+}
+
+
+
+
+UDynamicMesh* UGeometryScriptLibrary_MeshModelingFunctions::ApplyMeshInsetOutsetFaces(
+	UDynamicMesh* TargetMesh,
+	FGeometryScriptMeshInsetOutsetFacesOptions Options,
+	FGeometryScriptMeshSelection Selection,
+	UGeometryScriptDebug* Debug)
+{
+	if (TargetMesh == nullptr)
+	{
+		UE::Geometry::AppendError(Debug, EGeometryScriptErrorType::InvalidInputs, LOCTEXT("ApplyMeshInsetOutsetFaces_InvalidInput", "ApplyMeshInsetOutsetFaces: TargetMesh is Null"));
+		return TargetMesh;
+	}
+
+	TargetMesh->EditMesh([&](FDynamicMesh3& EditMesh) 
+	{
+		auto ApplyInsetOutsetArea = [&EditMesh, &Options](TArray<int32>&& Triangles)
+		{
+			FInsetMeshRegion InsetOutset(&EditMesh);
+			InsetOutset.Triangles = MoveTemp(Triangles);
+			if (InsetOutset.Triangles.Num() > 0)
+			{
+				InsetOutset.InsetDistance = Options.Distance;
+				InsetOutset.bReproject = (Options.Distance < 0) ? false : Options.bReproject;
+				InsetOutset.Softness = Options.Softness;
+				InsetOutset.bSolveRegionInteriors = !Options.bBoundaryOnly;
+				InsetOutset.AreaCorrection = Options.AreaScale;
+				InsetOutset.Apply();
+
+				for (FInsetMeshRegion::FInsetInfo& InsetInfo : InsetOutset.InsetRegions)
+				{
+					UELocal::ApplyPolygroupModeToNewTriangles(EditMesh, InsetInfo.InitialTriangles, Options.GroupOptions);
+				}
+			}
+		};
+
+		TArray<int32> Triangles;
+		if (Selection.GetNumSelected() == 0)
+		{
+			for (int32 tid : EditMesh.TriangleIndicesItr())
+			{
+				Triangles.Add(tid);
+			}
+		}
+		else
+		{
+			Selection.ConvertToMeshIndexArray(EditMesh, Triangles, EGeometryScriptIndexType::Triangle);
+		}
+
+		if (Options.AreaMode == EGeometryScriptPolyOperationArea::PerPolygroup)
+		{
+			TMap<int32, TArray<int32>> PolygroupTris;
+			for (int32 tid : Triangles)
+			{
+				int32 gid = EditMesh.GetTriangleGroup(tid);
+				TArray<int32>* Found = PolygroupTris.Find(gid);
+				if (Found != nullptr)
+				{
+					Found->Add(tid);
+				}
+				else
+				{
+					TArray<int32> Temp;
+					Temp.Add(tid);
+					PolygroupTris.Add(gid, Temp);
+				}
+			}
+
+			for (TPair<int32, TArray<int32>> Pair : PolygroupTris)
+			{
+				ApplyInsetOutsetArea(MoveTemp(Pair.Value));
+			}
+		}
+		else if (Options.AreaMode == EGeometryScriptPolyOperationArea::PerTriangle)
+		{
+			for (int32 tid : Triangles)
+			{
+				TArray<int32> Temp;
+				Temp.Add(tid);
+				ApplyInsetOutsetArea(MoveTemp(Temp));
+			}
+		}
+		else
+		{
+			ApplyInsetOutsetArea(MoveTemp(Triangles));
+		}
+
+	}, EDynamicMeshChangeType::GeneralEdit, EDynamicMeshAttributeChangeFlags::Unknown, false);
+
+	return TargetMesh;
+}
+
+
+
 
 UDynamicMesh* UGeometryScriptLibrary_MeshModelingFunctions::ApplyMeshBevelSelection(
 	UDynamicMesh* TargetMesh,
diff --git a/Engine/Plugins/Experimental/GeometryScripting/Source/GeometryScriptingCore/Public/GeometryScript/MeshModelingFunctions.h b/Engine/Plugins/Experimental/GeometryScripting/Source/GeometryScriptingCore/Public/GeometryScript/MeshModelingFunctions.h
--- a/Engine/Plugins/Experimental/GeometryScripting/Source/GeometryScriptingCore/Public/GeometryScript/MeshModelingFunctions.h
+++ b/Engine/Plugins/Experimental/GeometryScripting/Source/GeometryScriptingCore/Public/GeometryScript/MeshModelingFunctions.h
@@ -57,6 +57,14 @@ public:
 };
 
 
+UENUM(BlueprintType)
+enum class EGeometryScriptPolyOperationArea : uint8
+{
+	EntireSelection = 0,
+	PerPolygroup = 1,
+	PerTriangle = 2
+};
+
 
 USTRUCT(BlueprintType)
 struct GEOMETRYSCRIPTINGCORE_API FGeometryScriptMeshExtrudeOptions
@@ -84,13 +92,6 @@ enum class EGeometryScriptLinearExtrudeDirection : uint8
 	AverageFaceNormal = 1
 };
 
-UENUM(BlueprintType)
-enum class EGeometryScriptLinearExtrudeArea : uint8
-{
-	EntireSelection = 0,
-	PerPolygroup = 1,
-	PerTriangle = 2
-};
 
 USTRUCT(BlueprintType)
 struct GEOMETRYSCRIPTINGCORE_API FGeometryScriptMeshLinearExtrudeOptions
@@ -107,7 +108,40 @@ public:
 	FVector Direction = FVector(0,0,1);
 
 	UPROPERTY(EditAnywhere, BlueprintReadWrite, Category = Options)
-	EGeometryScriptLinearExtrudeArea AreaMode = EGeometryScriptLinearExtrudeArea::EntireSelection;
+	EGeometryScriptPolyOperationArea AreaMode = EGeometryScriptPolyOperationArea::EntireSelection;
+
+	UPROPERTY(EditAnywhere, BlueprintReadWrite, Category = Options)
+	FGeometryScriptMeshEditPolygroupOptions GroupOptions;
+
+	UPROPERTY(EditAnywhere, BlueprintReadWrite, Category = Options)
+	float UVScale = 1.0f;
+
+	UPROPERTY(EditAnywhere, BlueprintReadWrite, Category = Options)
+	bool bSolidsToShells = true;
+};
+
+
+UENUM(BlueprintType)
+enum class EGeometryScriptOffsetFacesType : uint8
+{
+	VertexNormal = 0,
+	FaceNormal = 1,
+	ParallelFaceOffset = 2
+};
+
+USTRUCT(BlueprintType)
+struct GEOMETRYSCRIPTINGCORE_API FGeometryScriptMeshOffsetFacesOptions
+{
+	GENERATED_BODY()
+public:
+	UPROPERTY(EditAnywhere, BlueprintReadWrite, Category = Options)
+	float Distance = 1.0;
+
+	UPROPERTY(EditAnywhere, BlueprintReadWrite, Category = Options)
+	EGeometryScriptOffsetFacesType OffsetType = EGeometryScriptOffsetFacesType::ParallelFaceOffset;
+
+	UPROPERTY(EditAnywhere, BlueprintReadWrite, Category = Options)
+	EGeometryScriptPolyOperationArea AreaMode = EGeometryScriptPolyOperationArea::EntireSelection;
 
 	UPROPERTY(EditAnywhere, BlueprintReadWrite, Category = Options)
 	FGeometryScriptMeshEditPolygroupOptions GroupOptions;
@@ -122,6 +156,38 @@ public:
 
 
 
+USTRUCT(BlueprintType)
+struct GEOMETRYSCRIPTINGCORE_API FGeometryScriptMeshInsetOutsetFacesOptions
+{
+	GENERATED_BODY()
+public:
+	UPROPERTY(EditAnywhere, BlueprintReadWrite, Category = Options)
+	float Distance = 1.0;
+
+	UPROPERTY(EditAnywhere, BlueprintReadWrite, Category = Options)
+	bool bReproject = true;
+
+	UPROPERTY(EditAnywhere, BlueprintReadWrite, Category = Options)
+	bool bBoundaryOnly = false;
+
+	UPROPERTY(EditAnywhere, BlueprintReadWrite, Category = Options)
+	float Softness = 0.0;
+
+	UPROPERTY(EditAnywhere, BlueprintReadWrite, Category = Options)
+	float AreaScale = 1.0;
+
+	UPROPERTY(EditAnywhere, BlueprintReadWrite, Category = Options)
+	EGeometryScriptPolyOperationArea AreaMode = EGeometryScriptPolyOperationArea::EntireSelection;
+
+	UPROPERTY(EditAnywhere, BlueprintReadWrite, Category = Options)
+	FGeometryScriptMeshEditPolygroupOptions GroupOptions;
+
+	UPROPERTY(EditAnywhere, BlueprintReadWrite, Category = Options)
+	float UVScale = 1.0f;
+};
+
+
+
 USTRUCT(BlueprintType)
 struct GEOMETRYSCRIPTINGCORE_API FGeometryScriptMeshBevelOptions
 {
@@ -233,8 +299,10 @@ public:
 		UGeometryScriptDebug* Debug = nullptr );
 
 
-
-
+	/**
+	 * Offset the vertices of TargetMesh from their initial positions based on averaged vertex normals.
+	 * This function is intended for high-res meshes, for polymodeling-style offsets, ApplyMeshOffsetFaces will produce better results.
+	 */
 	UFUNCTION(BlueprintCallable, Category = "GeometryScript|Modeling", meta=(ScriptMethod))
 	static UPARAM(DisplayName = "Target Mesh") UDynamicMesh* 
 	ApplyMeshOffset(
@@ -242,6 +310,10 @@ public:
 		FGeometryScriptMeshOffsetOptions Options,
 		UGeometryScriptDebug* Debug = nullptr );
 
+	/**
+	 * Create a thickened shell from TargetMesh by offsetting the vertex positions along averaged vertex normals, inwards or outwards.
+	 * Similar to ApplyMeshOffset but also includes the initial mesh (possibly flipped, if the offset is positive)
+	 */
 	UFUNCTION(BlueprintCallable, Category = "GeometryScript|Modeling", meta=(ScriptMethod))
 	static UPARAM(DisplayName = "Target Mesh") UDynamicMesh* 
 	ApplyMeshShell(
@@ -262,7 +334,28 @@ public:
 		FGeometryScriptMeshSelection Selection,
 		UGeometryScriptDebug* Debug = nullptr );
 
+	/**
+	 * Apply an Offset to the faces of TargetMesh identified by the Selection, or all faces if the Selection is empty.
+	 * The Offset direction at each vertex can be derived from the averaged vertex normals or per-triangle normals.
+	 */
+	UFUNCTION(BlueprintCallable, Category = "GeometryScript|Modeling", meta=(ScriptMethod))
+	static UPARAM(DisplayName = "Target Mesh") UDynamicMesh* 
+	ApplyMeshOffsetFaces(
+		UDynamicMesh* TargetMesh,
+		FGeometryScriptMeshOffsetFacesOptions Options,
+		FGeometryScriptMeshSelection Selection,
+		UGeometryScriptDebug* Debug = nullptr );
 
+	/**
+	 * Apply an Inset or Outset to the faces of TargetMesh identified by the Selection, or all faces if the Selection is empty.
+	 */
+	UFUNCTION(BlueprintCallable, Category = "GeometryScript|Modeling", meta=(ScriptMethod))
+	static UPARAM(DisplayName = "Target Mesh") UDynamicMesh* 
+	ApplyMeshInsetOutsetFaces(
+		UDynamicMesh* TargetMesh,
+		FGeometryScriptMeshInsetOutsetFacesOptions Options,
+		FGeometryScriptMeshSelection Selection,
+		UGeometryScriptDebug* Debug = nullptr );
 
 	/**
 	 * Apply a Mesh Bevel operation to parts of TargetMesh using the BevelOptions settings.
diff --git a/Engine/Plugins/Runtime/GeometryProcessing/Source/DynamicMesh/Private/Operations/OffsetMeshRegion.cpp b/Engine/Plugins/Runtime/GeometryProcessing/Source/DynamicMesh/Private/Operations/OffsetMeshRegion.cpp
--- a/Engine/Plugins/Runtime/GeometryProcessing/Source/DynamicMesh/Private/Operations/OffsetMeshRegion.cpp
+++ b/Engine/Plugins/Runtime/GeometryProcessing/Source/DynamicMesh/Private/Operations/OffsetMeshRegion.cpp
@@ -37,14 +37,6 @@ FOffsetMeshRegion::FOffsetMeshRegion(FDynamicMesh3* mesh) : Mesh(mesh)
 
 bool FOffsetMeshRegion::Apply()
 {
-	FMeshNormals Normals;
-	bool bHaveVertexNormals = Mesh->HasVertexNormals();
-	if (!bHaveVertexNormals && ExtrusionVectorType == EVertexExtrusionVectorType::VertexNormal)
-	{
-		Normals = FMeshNormals(Mesh);
-		Normals.ComputeVertexNormals();
-	}
-
 	FMeshConnectedComponents RegionComponents(Mesh);
 	RegionComponents.FindConnectedTriangles(Triangles);
 
@@ -62,7 +54,7 @@ bool FOffsetMeshRegion::Apply()
 			Region.bIsSolid = AllTriangles.Num() == Region.OffsetTids.Num();
 		}
 
-		bool bRegionOK = ApplyOffset(Region, (bHaveVertexNormals) ? nullptr : &Normals);
+		bool bRegionOK = ApplyOffset(Region);
 
 		bAllOK = bAllOK && bRegionOK;
 	}
@@ -72,7 +64,7 @@ bool FOffsetMeshRegion::Apply()
 
 }
 
-bool FOffsetMeshRegion::ApplyOffset(FOffsetInfo& Region, FMeshNormals* UseNormals)
+bool FOffsetMeshRegion::ApplyOffset(FOffsetInfo& Region)
 {
 	// Store offset groups
 	if (Mesh->HasTriangleGroups())
@@ -293,24 +285,22 @@ bool FOffsetMeshRegion::ApplyOffset(FOffsetInfo& Region, FMeshNormals* UseNormal
 			VertexExtrudeVectors[i] = ExtrusionVector;
 		}
 	}
+	else if (ExtrusionVectorType == EVertexExtrusionVectorType::VertexNormal)
+	{
+		VertexExtrudeVectors.SetNumUninitialized(SelectedVids.Num());
+		for (int32 i = 0; i < SelectedVids.Num(); ++i)
+		{
+			int32 Vid = SelectedVids[i];
+			VertexExtrudeVectors[i] = Mesh->HasVertexNormals() ? (FVector3d)Mesh->GetVertexNormal(Vid) : FMeshNormals::ComputeVertexNormal(*Mesh, Vid);
+		}
+	}
 
 	// Perform the actual vertex displacement.
 	for (int32 i = 0; i < SelectedVids.Num(); ++i)
 	{
 		int32 Vid = SelectedVids[i];
-		FVector OldPosition = Mesh->GetVertex(Vid);
-		FVector ExtrusionVector = FVector::Zero();
-
-		switch (ExtrusionVectorType)
-		{
-		case EVertexExtrusionVectorType::VertexNormal:
-			ExtrusionVector = FVector(Mesh->GetVertexNormal(Vid));
-			break;
-		case EVertexExtrusionVectorType::SelectionTriNormalsAngleWeightedAverage:
-		case EVertexExtrusionVectorType::SelectionTriNormalsAngleWeightedAdjusted:
-			ExtrusionVector = VertexExtrudeVectors[i];
-			break;
-		}
+		FVector3d OldPosition = Mesh->GetVertex(Vid);
+		FVector3d ExtrusionVector = (ExtrusionVectorType == EVertexExtrusionVectorType::Zero) ? FVector3d::Zero() : VertexExtrudeVectors[i];
 
 		FVector3d NewPosition = OffsetPositionFunc(OldPosition, ExtrusionVector, Vid);
 		Mesh->SetVertex(Vid, NewPosition);
diff --git a/Engine/Plugins/Runtime/GeometryProcessing/Source/DynamicMesh/Public/Operations/OffsetMeshRegion.h b/Engine/Plugins/Runtime/GeometryProcessing/Source/DynamicMesh/Public/Operations/OffsetMeshRegion.h
--- a/Engine/Plugins/Runtime/GeometryProcessing/Source/DynamicMesh/Public/Operations/OffsetMeshRegion.h
+++ b/Engine/Plugins/Runtime/GeometryProcessing/Source/DynamicMesh/Public/Operations/OffsetMeshRegion.h
@@ -173,7 +173,7 @@ public:
 		int32 Eid1, int32 Eid2, bool bCheckColinearityAtBorder);
 protected:
 
-	virtual bool ApplyOffset(FOffsetInfo& Region, FMeshNormals* UseNormals = nullptr);
+	virtual bool ApplyOffset(FOffsetInfo& Region);
 };
 
 } // end namespace UE::Geometry
