GeometryScript: add InvertMeshSelection
GeometrySelectionUtil: add FindByTopologyID function, to search a Selection for selected element with matching TopologyID but ignoring GeometryID. Update FGeometryScriptMeshSelection::CombineSelectionInPlace() to use this, as it was not properly combining Polygroup selections.

diff --git a/Engine/Plugins/Experimental/GeometryScripting/Source/GeometryScriptingCore/Private/GeometryScriptSelectionTypes.cpp b/Engine/Plugins/Experimental/GeometryScripting/Source/GeometryScriptingCore/Private/GeometryScriptSelectionTypes.cpp
--- a/Engine/Plugins/Experimental/GeometryScripting/Source/GeometryScriptingCore/Private/GeometryScriptSelectionTypes.cpp
+++ b/Engine/Plugins/Experimental/GeometryScripting/Source/GeometryScriptingCore/Private/GeometryScriptSelectionTypes.cpp
@@ -103,33 +103,95 @@ void FGeometryScriptMeshSelection::CombineSelectionInPlace(const FGeometryScript
 		return;
 	}
 
-	if (CombineMode == EGeometryScriptCombineSelectionMode::Add)
+	if (GetSelectionType() != EGeometryScriptMeshSelectionType::Polygroups)
 	{
-		for (uint64 Item : SelectionB.GeoSelection->Selection)
+		if (CombineMode == EGeometryScriptCombineSelectionMode::Add)
 		{
-			GeoSelection->Selection.Add(Item);
+			for (uint64 ItemB : SelectionB.GeoSelection->Selection)
+			{
+				GeoSelection->Selection.Add(ItemB);
+			}
 		}
-	}
-	else if (CombineMode == EGeometryScriptCombineSelectionMode::Subtract)
-	{
-		for (uint64 Item : SelectionB.GeoSelection->Selection)
+		else if (CombineMode == EGeometryScriptCombineSelectionMode::Subtract)
+		{
+			if (SelectionB.IsEmpty() == false)
+			{
+				for (uint64 ItemB : SelectionB.GeoSelection->Selection)
+				{
+					GeoSelection->Selection.Remove(ItemB);
+				}
+				GeoSelection->Selection.Compact();
+			}
+		}
+		else if (CombineMode == EGeometryScriptCombineSelectionMode::Intersection)
 		{
-			GeoSelection->Selection.Remove(Item);
+			TArray<uint64, TInlineAllocator<32>> ToRemove;
+			for (uint64 ItemA : GeoSelection->Selection)
+			{
+				if (!SelectionB.GeoSelection->Selection.Contains(ItemA))
+				{
+					ToRemove.Add(ItemA);
+				}
+			}
+			if (ToRemove.Num() > 0)
+			{
+				for (uint64 ItemA : ToRemove)
+				{
+					GeoSelection->Selection.Remove(ItemA);
+				}
+				GeoSelection->Selection.Compact();
+			}
 		}
 	}
-	else if (CombineMode == EGeometryScriptCombineSelectionMode::Intersection)
+	else
 	{
-		TArray<uint64, TInlineAllocator<32>> ToRemove;
-		for (uint64 Item : GeoSelection->Selection)
+		// for Polygroup selections, we cannot rely on TSet operations because we have set an arbitrary Triangle ID 
+		// as the 'geometry' key.
+		if (CombineMode == EGeometryScriptCombineSelectionMode::Add)
 		{
-			if ( ! SelectionB.GeoSelection->Selection.Contains(Item) )
+			for (uint64 ItemB : SelectionB.GeoSelection->Selection)
 			{
-				ToRemove.Add(Item);
+				uint64 FoundItemA;
+				if ( UE::Geometry::FindInSelectionByTopologyID(*GeoSelection, FGeoSelectionID(ItemB).TopologyID, FoundItemA) == false)
+				{
+					GeoSelection->Selection.Add(ItemB);
+				}
 			}
 		}
-		for (uint64 Item : ToRemove)
+		else if (CombineMode == EGeometryScriptCombineSelectionMode::Subtract)
 		{
-			GeoSelection->Selection.Remove(Item);
+			if (SelectionB.IsEmpty() == false)
+			{
+				for (uint64 ItemB : SelectionB.GeoSelection->Selection)
+				{
+					uint64 FoundItemA;
+					if (UE::Geometry::FindInSelectionByTopologyID(*GeoSelection, FGeoSelectionID(ItemB).TopologyID, FoundItemA))
+					{
+						GeoSelection->Selection.Remove(FoundItemA);
+					}
+				}
+				GeoSelection->Selection.Compact();
+			}
+		}
+		else if (CombineMode == EGeometryScriptCombineSelectionMode::Intersection)
+		{
+			TArray<uint64, TInlineAllocator<32>> ToRemove;
+			for (uint64 ItemA : GeoSelection->Selection)
+			{
+				uint64 FoundItemB;
+				if (UE::Geometry::FindInSelectionByTopologyID(*GeoSelection, FGeoSelectionID(ItemA).TopologyID, FoundItemB) == false)
+				{
+					ToRemove.Add(ItemA);
+				}
+			}
+			if (ToRemove.Num() > 0)
+			{
+				for (uint64 ItemA : ToRemove)
+				{
+					GeoSelection->Selection.Remove(ItemA);
+				}
+				GeoSelection->Selection.Compact();
+			}
 		}
 	}
 
@@ -137,7 +199,6 @@ void FGeometryScriptMeshSelection::CombineSelectionInPlace(const FGeometryScript
 
 
 
-
 EGeometryScriptIndexType FGeometryScriptMeshSelection::ConvertToMeshIndexArray(
 	const FDynamicMesh3& Mesh, TArray<int32>& IndexListOut, EGeometryScriptIndexType ConvertToType) const
 {
diff --git a/Engine/Plugins/Experimental/GeometryScripting/Source/GeometryScriptingCore/Private/MeshSelectionFunctions.cpp b/Engine/Plugins/Experimental/GeometryScripting/Source/GeometryScriptingCore/Private/MeshSelectionFunctions.cpp
--- a/Engine/Plugins/Experimental/GeometryScripting/Source/GeometryScriptingCore/Private/MeshSelectionFunctions.cpp
+++ b/Engine/Plugins/Experimental/GeometryScripting/Source/GeometryScriptingCore/Private/MeshSelectionFunctions.cpp
@@ -81,6 +81,7 @@ UDynamicMesh* UGeometryScriptLibrary_MeshSelectionFunctions::CreateSelectAllMesh
 				if (UniqueGroupIDs.Contains(GroupID) == false)
 				{
 					NewSelection.Selection.Add(FGeoSelectionID::GroupFace(tid, GroupID).Encoded());
+					UniqueGroupIDs.Add(GroupID);
 				}
 			}
 		}
@@ -717,6 +718,34 @@ UDynamicMesh* UGeometryScriptLibrary_MeshSelectionFunctions::SelectMeshElementsI
 
 
 
+UDynamicMesh* UGeometryScriptLibrary_MeshSelectionFunctions::InvertMeshSelection(
+	UDynamicMesh* TargetMesh,
+	FGeometryScriptMeshSelection Selection,
+	FGeometryScriptMeshSelection& NewSelection,
+	bool bOnlyToConnected)
+{
+	if (TargetMesh == nullptr)
+	{
+		UE_LOG(LogGeometry, Warning, TEXT("InvertMeshSelection: TargetMesh is Null"));
+		return TargetMesh;
+	}
+
+	if (bOnlyToConnected)
+	{
+		ExpandMeshSelectionToConnected(TargetMesh, Selection, NewSelection, EGeometryScriptTopologyConnectionType::Geometric);
+	}
+	else
+	{
+		CreateSelectAllMeshSelection(TargetMesh, NewSelection, Selection.GetSelectionType());
+	}
+	NewSelection.CombineSelectionInPlace(Selection, EGeometryScriptCombineSelectionMode::Subtract);
+
+	return TargetMesh;
+}
+
+
+
+
 UDynamicMesh* UGeometryScriptLibrary_MeshSelectionFunctions::ExpandMeshSelectionToConnected(
 	UDynamicMesh* TargetMesh,
 	FGeometryScriptMeshSelection Selection,
diff --git a/Engine/Plugins/Experimental/GeometryScripting/Source/GeometryScriptingCore/Public/GeometryScript/MeshSelectionFunctions.h b/Engine/Plugins/Experimental/GeometryScripting/Source/GeometryScriptingCore/Public/GeometryScript/MeshSelectionFunctions.h
--- a/Engine/Plugins/Experimental/GeometryScripting/Source/GeometryScriptingCore/Public/GeometryScript/MeshSelectionFunctions.h
+++ b/Engine/Plugins/Experimental/GeometryScripting/Source/GeometryScriptingCore/Public/GeometryScript/MeshSelectionFunctions.h
@@ -249,6 +249,20 @@ public:
 		double WindingThreshold = 0.5,
 		int MinNumTrianglePoints = 3);
 
+
+	/**
+	 * Invert the Selection on the TargetMesh, ie select what is not currently selected
+	 * @param bOnlyToConnected if true, the inverse is limited to mesh areas geometrically connected to the Selection, instead of the entire mesh
+	 */
+	UFUNCTION(BlueprintCallable, Category = "GeometryScript|MeshSelection", meta=(ScriptMethod))
+	static UPARAM(DisplayName = "Target Mesh") UDynamicMesh* 
+	InvertMeshSelection(
+		UDynamicMesh* TargetMesh,
+		FGeometryScriptMeshSelection Selection, 
+		FGeometryScriptMeshSelection& NewSelection,
+		bool bOnlyToConnected = false );
+
+
 	/**
 	 * Expand the Selection on the TargetMesh to connected regions and return the NewSelection
 	 * @param ConnectionType defines what "connected" means, ie purely geometric connection, or some additional constraint like same MaterialIDs/etc 
@@ -259,8 +273,7 @@ public:
 		UDynamicMesh* TargetMesh,
 		FGeometryScriptMeshSelection Selection, 
 		FGeometryScriptMeshSelection& NewSelection,
-		EGeometryScriptTopologyConnectionType ConnectionType = EGeometryScriptTopologyConnectionType::Geometric
-	);
+		EGeometryScriptTopologyConnectionType ConnectionType = EGeometryScriptTopologyConnectionType::Geometric );
 
 
 	/**
@@ -280,8 +293,7 @@ public:
 		FGeometryScriptMeshSelection& NewSelection,
 		int32 Iterations = 1,
 		bool bContract = false,
-		bool bOnlyExpandToFaceNeighbours = false
-	);
+		bool bOnlyExpandToFaceNeighbours = false );
 
 
 };
\ No newline at end of file
diff --git a/Engine/Plugins/Runtime/GeometryProcessing/Source/DynamicMesh/Private/Selections/GeometrySelectionUtil.cpp b/Engine/Plugins/Runtime/GeometryProcessing/Source/DynamicMesh/Private/Selections/GeometrySelectionUtil.cpp
--- a/Engine/Plugins/Runtime/GeometryProcessing/Source/DynamicMesh/Private/Selections/GeometrySelectionUtil.cpp
+++ b/Engine/Plugins/Runtime/GeometryProcessing/Source/DynamicMesh/Private/Selections/GeometrySelectionUtil.cpp
@@ -9,6 +9,26 @@
 #include "SegmentTypes.h"
 #include "GroupTopology.h"
 #include "Selections/MeshConnectedComponents.h"
+#include "Algo/Find.h"
+
+
+bool UE::Geometry::FindInSelectionByTopologyID(
+	const FGeometrySelection& GeometrySelection,
+	uint32 TopologyID,
+	uint64& FoundValue)
+{
+	const uint64* Found = Algo::FindByPredicate(GeometrySelection.Selection, [&](uint64 Item)
+	{
+		return FGeoSelectionID(Item).TopologyID == TopologyID;
+	});
+	if (Found != nullptr)
+	{
+		FoundValue = *Found;
+		return true;
+	}
+	FoundValue = FGeoSelectionID().Encoded();
+	return false;
+}
 
 
 void UE::Geometry::UpdateTriangleSelectionViaRaycast(
diff --git a/Engine/Plugins/Runtime/GeometryProcessing/Source/DynamicMesh/Public/Selections/GeometrySelectionUtil.h b/Engine/Plugins/Runtime/GeometryProcessing/Source/DynamicMesh/Public/Selections/GeometrySelectionUtil.h
--- a/Engine/Plugins/Runtime/GeometryProcessing/Source/DynamicMesh/Public/Selections/GeometrySelectionUtil.h
+++ b/Engine/Plugins/Runtime/GeometryProcessing/Source/DynamicMesh/Public/Selections/GeometrySelectionUtil.h
@@ -19,6 +19,18 @@ class FColliderMesh;
 class FGroupTopology;
 struct FGroupTopologySelection;
 
+/**
+ * Assuming that the uint64 values in the GeometrySelection are encoded FGeoSelectionID's, 
+ * find the item that has a matching TopologyID, ignoring the GeometryID.
+ * @param FoundValue the hash value found 
+ * @return true if an item was found
+ */
+DYNAMICMESH_API bool FindInSelectionByTopologyID(
+	const FGeometrySelection& GeometrySelection,
+	uint32 TopologyID,
+	uint64& FoundValue);
+
+
 /**
  * Update a FGeometrySelection (via the Editor) of mesh vertices/edges/triangles
  * based on a raycast
