expose 2D voronoi diagram generation to geometryscript

diff --git a/Engine/Plugins/Experimental/GeometryScripting/Source/GeometryScriptingCore/Private/MeshPrimitiveFunctions.cpp b/Engine/Plugins/Experimental/GeometryScripting/Source/GeometryScriptingCore/Private/MeshPrimitiveFunctions.cpp
--- a/Engine/Plugins/Experimental/GeometryScripting/Source/GeometryScriptingCore/Private/MeshPrimitiveFunctions.cpp
+++ b/Engine/Plugins/Experimental/GeometryScripting/Source/GeometryScriptingCore/Private/MeshPrimitiveFunctions.cpp
@@ -20,6 +20,7 @@
 #include "Generators/StairGenerator.h"
 #include "ConstrainedDelaunay2.h"
 #include "Arrangement2d.h"
+#include "CompGeom/Delaunay2.h"
 
 using namespace UE::Geometry;
 
@@ -996,4 +997,93 @@ UDynamicMesh* UGeometryScriptLibrary_MeshPrimitiveFunctions::AppendCurvedStairs(
 
 
 
+UDynamicMesh* UGeometryScriptLibrary_MeshPrimitiveFunctions::AppendVoronoiDiagram2D(
+	UDynamicMesh* TargetMesh,
+	FGeometryScriptPrimitiveOptions PrimitiveOptions,
+	FTransform Transform,
+	const TArray<FVector2D>& VoronoiSites,
+	FGeometryScriptVoronoiOptions VoronoiOptions,
+	UGeometryScriptDebug* Debug)
+{
+	if (TargetMesh == nullptr)
+	{
+		UE::Geometry::AppendError(Debug, EGeometryScriptErrorType::InvalidInputs, LOCTEXT("PrimitiveFunctions_AppendVoronoiDiagram2D_NullTarget", "AppendVoronoiDiagram3d: TargetMesh is Null"));
+		return TargetMesh;
+	}
+
+	if (VoronoiSites.Num() < 3)
+	{
+		UE::Geometry::AppendError(Debug, EGeometryScriptErrorType::InvalidInputs, LOCTEXT("PrimitiveFunctions_AppendVoronoiDiagram2D_TooFewSites", "AppendVoronoiDiagram2D: VoronoiSites array requires at least 3 positions"));
+		return TargetMesh;
+	}
+
+	UE::Geometry::FDelaunay2 Delaunay;
+	bool bTriSuccess = Delaunay.Triangulate(VoronoiSites);
+
+	if (!bTriSuccess)
+	{
+		UE::Geometry::AppendError(Debug, EGeometryScriptErrorType::InvalidInputs, LOCTEXT("PrimitiveFunctions_AppendVoronoiDiagram2D_GenFailed", "AppendVoronoiDiagram2D: Voronoi diagram generation failed"));
+		return TargetMesh;
+	}
+
+	TAxisAlignedBox2<double> AABB;
+	if (VoronoiOptions.Bounds.IsValid)
+	{
+		AABB.Max = FVector2d(VoronoiOptions.Bounds.Max.X, VoronoiOptions.Bounds.Max.Y);
+		AABB.Min = FVector2d(VoronoiOptions.Bounds.Min.X, VoronoiOptions.Bounds.Min.Y);
+	}
+	else
+	{
+		AABB.Contain(VoronoiSites);
+	}
+	TArray<TArray<FVector2d>> Polygons = Delaunay.ComputeVoronoiCells<double>(VoronoiSites, VoronoiOptions.bIncludeBoundary, AABB, (double)VoronoiOptions.BoundsExpand);
+
+	if (Polygons.Num() == 0)
+	{
+		UE::Geometry::AppendError(Debug, EGeometryScriptErrorType::OperationFailed, LOCTEXT("PrimitiveFunctions_AppendVoronoiDiagram2D_NoPolygons", "AppendVoronoiDiagram2D: No Voronoi cells constructed"));
+		return TargetMesh;
+	}
+
+	FFlatTriangulationMeshGenerator TriangulationMeshGen;
+	auto AddCell = [&TriangulationMeshGen, &Polygons, &Debug](int32 PolyIdx)
+	{
+		if (PolyIdx < 0 || PolyIdx >= Polygons.Num())
+		{
+			UE::Geometry::AppendWarning(Debug, EGeometryScriptErrorType::InvalidInputs, LOCTEXT("PrimitiveFunctions_AppendVoronoiDiagram2D_BadID", "AppendVoronoiDiagram2D: Requested invalid cell ID"));
+			return;
+		}
+		int32 Start = TriangulationMeshGen.Vertices2D.Num();
+		TriangulationMeshGen.Vertices2D.Append(Polygons[PolyIdx]);
+		for (int32 Off = 1; Off + 1 < Polygons[PolyIdx].Num(); ++Off)
+		{
+			TriangulationMeshGen.Triangles2D.Emplace(Start, Start + Off + 1, Start + Off);
+			TriangulationMeshGen.Triangles2DPolygroups.Emplace(PolyIdx);
+		}
+	};
+
+	if (VoronoiOptions.CreateCells.IsEmpty())
+	{
+		for (int32 PolyIdx = 0; PolyIdx < Polygons.Num(); ++PolyIdx)
+		{
+			AddCell(PolyIdx);
+		}
+	}
+	else
+	{
+		for (int32 PolyIdx : VoronoiOptions.CreateCells)
+		{
+			AddCell(PolyIdx);
+		}
+	}
+
+	if (TriangulationMeshGen.Vertices2D.Num() > 2 && TriangulationMeshGen.Triangles2D.Num() > 0)
+	{
+		AppendPrimitive(TargetMesh, &TriangulationMeshGen.Generate(), Transform, PrimitiveOptions);
+	}
+
+	return TargetMesh;
+}
+
+
+
 #undef LOCTEXT_NAMESPACE
\ No newline at end of file
diff --git a/Engine/Plugins/Experimental/GeometryScripting/Source/GeometryScriptingCore/Public/GeometryScript/MeshPrimitiveFunctions.h b/Engine/Plugins/Experimental/GeometryScripting/Source/GeometryScriptingCore/Public/GeometryScript/MeshPrimitiveFunctions.h
--- a/Engine/Plugins/Experimental/GeometryScripting/Source/GeometryScriptingCore/Public/GeometryScript/MeshPrimitiveFunctions.h
+++ b/Engine/Plugins/Experimental/GeometryScripting/Source/GeometryScriptingCore/Public/GeometryScript/MeshPrimitiveFunctions.h
@@ -74,6 +74,26 @@ public:
 	bool bFillPartialRevolveEndcaps = true;
 };
 
+USTRUCT(BlueprintType)
+struct GEOMETRYSCRIPTINGCORE_API FGeometryScriptVoronoiOptions
+{
+	GENERATED_BODY()
+public:
+	UPROPERTY(EditAnywhere, BlueprintReadWrite, Category = Options)
+	float BoundsExpand = 0.0f;
+
+	UPROPERTY(EditAnywhere, BlueprintReadWrite, Category = Options)
+	FBox Bounds = FBox(EForceInit::ForceInit);
+
+	/** Optional list of cells to create meshes for.  If empty, create all cells. */
+	UPROPERTY(EditAnywhere, BlueprintReadWrite, Category = Options)
+	TArray<int32> CreateCells;
+
+	/** Whether to include the bordering Voronoi cells (which extend 'infinitely' to any boundary) */
+	UPROPERTY(EditAnywhere, BlueprintReadWrite, Category = Options)
+	bool bIncludeBoundary = true;
+};
+
 
 UCLASS(meta = (ScriptName = "GeometryScript_Primitives"))
 class GEOMETRYSCRIPTINGCORE_API UGeometryScriptLibrary_MeshPrimitiveFunctions : public UBlueprintFunctionLibrary
@@ -373,6 +393,16 @@ public:
 		UGeometryScriptDebug* Debug = nullptr);
 
 
+	UFUNCTION(BlueprintCallable, Category = "GeometryScript|Primitives", meta = (ScriptMethod))
+	static UPARAM(DisplayName = "Target Mesh") UDynamicMesh*
+	AppendVoronoiDiagram2D(
+		UDynamicMesh* TargetMesh,
+		FGeometryScriptPrimitiveOptions PrimitiveOptions,
+		FTransform Transform,
+		const TArray<FVector2D>& VoronoiSites,
+		FGeometryScriptVoronoiOptions VoronoiOptions,
+		UGeometryScriptDebug* Debug = nullptr);
+
 
 
 	//---------------------------------------------
diff --git a/Engine/Source/Runtime/GeometryCore/Private/Generators/FlatTriangulationMeshGenerator.cpp b/Engine/Source/Runtime/GeometryCore/Private/Generators/FlatTriangulationMeshGenerator.cpp
--- a/Engine/Source/Runtime/GeometryCore/Private/Generators/FlatTriangulationMeshGenerator.cpp
+++ b/Engine/Source/Runtime/GeometryCore/Private/Generators/FlatTriangulationMeshGenerator.cpp
@@ -51,9 +51,10 @@ FMeshShapeGenerator& FFlatTriangulationMeshGenerator::Generate()
 		SetTriangle(TriIdx, Tri);
 		SetTriangleUVs(TriIdx, Tri);
 		SetTriangleNormals(TriIdx, Tri);
-		SetTrianglePolygon(TriIdx, 0);
+		int PolygroupID = Triangles2DPolygroups.IsEmpty() ? 0 : Triangles2DPolygroups[TriIdx];
+		SetTrianglePolygon(TriIdx, PolygroupID);
 	}
 
 
 	return *this;
-}
\ No newline at end of file
+}
diff --git a/Engine/Source/Runtime/GeometryCore/Public/Generators/FlatTriangulationMeshGenerator.h b/Engine/Source/Runtime/GeometryCore/Public/Generators/FlatTriangulationMeshGenerator.h
--- a/Engine/Source/Runtime/GeometryCore/Public/Generators/FlatTriangulationMeshGenerator.h
+++ b/Engine/Source/Runtime/GeometryCore/Public/Generators/FlatTriangulationMeshGenerator.h
@@ -26,6 +26,9 @@ public:
 	/** Source triangle indices */
 	TArray<FIndex3i> Triangles2D;
 
+	/** Optional polygroup IDs per triangle -- ignored if empty */
+	TArray<int> Triangles2DPolygroups;
+
 	/** Normal vector of all vertices will be set to this value. Default is +Z axis. */
 	FVector3f Normal;
 
