Remove options for the tessellation patterns that haven't been implemented yet.

#rnx

diff --git a/Engine/Plugins/Experimental/GeometryScripting/Source/GeometryScriptingCore/Public/GeometryScript/MeshSubdivideFunctions.h b/Engine/Plugins/Experimental/GeometryScripting/Source/GeometryScriptingCore/Public/GeometryScript/MeshSubdivideFunctions.h
--- a/Engine/Plugins/Experimental/GeometryScripting/Source/GeometryScriptingCore/Public/GeometryScript/MeshSubdivideFunctions.h
+++ b/Engine/Plugins/Experimental/GeometryScripting/Source/GeometryScriptingCore/Public/GeometryScript/MeshSubdivideFunctions.h
@@ -29,9 +29,7 @@ public:
 UENUM(BlueprintType)
 enum class EAdaptiveTessellatePatternType : uint8
 {
-	Uniform = 0,
-	InnerUniform = 1,
-	ConcentricRings = 2
+	ConcentricRings = 0
 };
 
 
