GeometryScript: add ThickenShells option to ApplyMeshSolidify. Explicitly initialize new mesh generated in ApplyMeshSolidify and ApplyMeshMorphology to have per-vertex normals. Add handling for OffsetDistance=0 in ApplyMeshMorphology, because TImplicitMorphology does not allow this.

diff --git a/Engine/Plugins/Experimental/GeometryScripting/Source/GeometryScriptingCore/Private/MeshVoxelFunctions.cpp b/Engine/Plugins/Experimental/GeometryScripting/Source/GeometryScriptingCore/Private/MeshVoxelFunctions.cpp
--- a/Engine/Plugins/Experimental/GeometryScripting/Source/GeometryScriptingCore/Private/MeshVoxelFunctions.cpp
+++ b/Engine/Plugins/Experimental/GeometryScripting/Source/GeometryScriptingCore/Private/MeshVoxelFunctions.cpp
@@ -8,6 +8,7 @@
 #include "Spatial/FastWinding.h"
 #include "Implicit/Solidify.h"
 #include "Implicit/Morphology.h"
+#include "Operations/ExtrudeMesh.h"
 #include "UDynamicMesh.h"
 
 using namespace UE::Geometry;
@@ -64,21 +65,42 @@ UDynamicMesh* UGeometryScriptLibrary_MeshVoxelFunctions::ApplyMeshSolidify(
 		return TargetMesh;
 	}
 
-	TargetMesh->EditMesh([&](FDynamicMesh3& EditMesh) 
+	TargetMesh->EditMesh([&](FDynamicMesh3& EditMesh)
 	{
-		FDynamicMeshAABBTree3 Spatial(&EditMesh);
+		FDynamicMesh3& UseSourceMesh = EditMesh;
+
+		FAxisAlignedBox3d SourceMeshBounds = EditMesh.GetBounds(true);
+
+		FDynamicMesh3 ThickenedMesh;
+		if (Options.bThickenShells)
+		{
+			// thickness should be at least a cell wide so we don't end up deleting a bunch of the input surface
+			double CellSize = (Options.GridParameters.SizeMethod == EGeometryScriptGridSizingMethod::GridResolution) ?
+				(SourceMeshBounds.MaxDim() / Options.GridParameters.GridResolution) : Options.GridParameters.GridCellSize;
+			double SafeThickness = FMathd::Max(CellSize * 2, Options.ShellThickness);
+			ThickenedMesh.Copy(UseSourceMesh, false, false, false, false);
+			FMeshNormals::QuickComputeVertexNormals(ThickenedMesh);
+			FExtrudeMesh Extrude(&ThickenedMesh);
+			Extrude.bSkipClosedComponents = true;
+			Extrude.DefaultExtrudeDistance = -SafeThickness;
+			Extrude.IsPositiveOffset = false;
+			Extrude.Apply();
+			UseSourceMesh = ThickenedMesh;
+		}
+
+		FDynamicMeshAABBTree3 Spatial(&UseSourceMesh);
 		TFastWindingTree<FDynamicMesh3> FastWinding(&Spatial);
 
-		TImplicitSolidify<FDynamicMesh3> Solidify(&EditMesh, &Spatial, &FastWinding);
+		TImplicitSolidify<FDynamicMesh3> Solidify(&UseSourceMesh, &Spatial, &FastWinding);
 
 		if (Options.GridParameters.SizeMethod == EGeometryScriptGridSizingMethod::GridResolution)
 		{
 			int ClampResolution = GetClampedGridResolution(Options.GridParameters.GridResolution, TEXT("ApplyMeshSolidify"), Debug);
-			Solidify.SetCellSizeAndExtendBounds(Spatial.GetBoundingBox(), Options.ExtendBounds, ClampResolution);
+			Solidify.SetCellSizeAndExtendBounds(SourceMeshBounds, Options.ExtendBounds, ClampResolution);
 		}
 		else
 		{
-			double UseCellSize = GetClampedCellSize(Options.GridParameters.GridCellSize, Spatial.GetBoundingBox(), TEXT("ApplyMeshSolidify"), Debug);
+			double UseCellSize = GetClampedCellSize(Options.GridParameters.GridCellSize, SourceMeshBounds, TEXT("ApplyMeshSolidify"), Debug);
 			Solidify.ExtendBounds = Options.ExtendBounds;
 			Solidify.MeshCellSize = UseCellSize;
 		}
@@ -93,6 +115,11 @@ UDynamicMesh* UGeometryScriptLibrary_MeshVoxelFunctions::ApplyMeshSolidify(
 		{
 			UE::Geometry::AppendError(Debug, EGeometryScriptErrorType::OperationFailed, LOCTEXT("ApplyMeshSolidify_EmptyMesh", "ApplyMeshSolidify: Solidify operation output an empty mesh"));
 		}
+		else
+		{
+			EditMesh.EnableAttributes();
+			FMeshNormals::InitializeOverlayToPerVertexNormals(EditMesh.Attributes()->PrimaryNormals(), false);
+		}
 
 	}, EDynamicMeshChangeType::GeneralEdit, EDynamicMeshAttributeChangeFlags::Unknown, false);
 
@@ -124,6 +151,12 @@ UDynamicMesh* UGeometryScriptLibrary_MeshVoxelFunctions::ApplyMeshMorphology(
 		Morphology.SourceSpatial = &Spatial;
 		Morphology.MorphologyOp = static_cast<TImplicitMorphology<FDynamicMesh3>::EMorphologyOp>(Options.Operation);
 		Morphology.Distance = Options.Distance;
+		if (Morphology.Distance == 0)
+		{
+			// TImplicitMorphology does not like Distance=0, but in this function we want to fallback to basically just computing SDF
+			Morphology.MorphologyOp = TImplicitMorphology<FDynamicMesh3>::EMorphologyOp::Dilate;
+			Morphology.Distance = 0.00001;   
+		}
 
 		FGeometryScript3DGridParameters SDFGridParameters = Options.SDFGridParameters;
 		if (SDFGridParameters.SizeMethod == EGeometryScriptGridSizingMethod::GridResolution)
@@ -156,6 +189,11 @@ UDynamicMesh* UGeometryScriptLibrary_MeshVoxelFunctions::ApplyMeshMorphology(
 		{
 			UE::Geometry::AppendError(Debug, EGeometryScriptErrorType::OperationFailed, LOCTEXT("ApplyMeshMorphology_EmptyMesh", "ApplyMeshMorphology: Morphology operation output an empty mesh"));
 		}
+		else
+		{
+			EditMesh.EnableAttributes();
+			FMeshNormals::InitializeOverlayToPerVertexNormals(EditMesh.Attributes()->PrimaryNormals(), false);
+		}
 
 	}, EDynamicMeshChangeType::GeneralEdit, EDynamicMeshAttributeChangeFlags::Unknown, false);
 
diff --git a/Engine/Plugins/Experimental/GeometryScripting/Source/GeometryScriptingCore/Public/GeometryScript/MeshVoxelFunctions.h b/Engine/Plugins/Experimental/GeometryScripting/Source/GeometryScriptingCore/Public/GeometryScript/MeshVoxelFunctions.h
--- a/Engine/Plugins/Experimental/GeometryScripting/Source/GeometryScriptingCore/Public/GeometryScript/MeshVoxelFunctions.h
+++ b/Engine/Plugins/Experimental/GeometryScripting/Source/GeometryScriptingCore/Public/GeometryScript/MeshVoxelFunctions.h
@@ -58,6 +58,15 @@ public:
 
 	UPROPERTY(EditAnywhere, BlueprintReadWrite, Category = Options)
 	int SurfaceSearchSteps = 3;
+
+	/** When enabled, regions of the input mesh that have open boundaries (ie "shells") are thickened by extruding them into closed solids. This may be expensive on large meshes. */
+	UPROPERTY(BlueprintReadWrite, Category = Options)
+	bool bThickenShells = false;
+
+	/** Open Shells are Thickened by offsetting vertices along their averaged vertex normals by this amount. Dimension is but clamped to twice the grid cell size. */
+	UPROPERTY(BlueprintReadWrite, Category = Options)
+	double ShellThickness = 1.0;
+
 };
 
 
