GeometryScript: detect/prevent non-compact mesh input in AutoGeneratePatchBuilderMeshUVs

diff --git a/Engine/Plugins/Experimental/GeometryScripting/Source/GeometryScriptingCore/Private/MeshUVFunctions.cpp b/Engine/Plugins/Experimental/GeometryScripting/Source/GeometryScriptingCore/Private/MeshUVFunctions.cpp
--- a/Engine/Plugins/Experimental/GeometryScripting/Source/GeometryScriptingCore/Private/MeshUVFunctions.cpp
+++ b/Engine/Plugins/Experimental/GeometryScripting/Source/GeometryScriptingCore/Private/MeshUVFunctions.cpp
@@ -581,6 +581,12 @@ UDynamicMesh* UGeometryScriptLibrary_MeshUVFunctions::AutoGeneratePatchBuilderMe
 	ApplyMeshUVEditorOperation(TargetMesh, UVSetIndex, bHasUVSet, Debug,
 		[&](FDynamicMesh3& EditMesh, FDynamicMeshUVOverlay* UVOverlay, FDynamicMeshUVEditor& UVEditor)
 	{
+		if (EditMesh.IsCompact() == false)
+		{
+			UE::Geometry::AppendError(Debug, EGeometryScriptErrorType::InvalidInputs, LOCTEXT("AutoGeneratePatchBuilderMeshUVs_NonCompact", "AutoGeneratePatchBuilderMeshUVs: TargetMesh is non-Compact, PatchBuilder cannot be run. Try calling CompactMesh to update TargetMesh"));
+			return;
+		}
+
 		FPatchBasedMeshUVGenerator UVGenerator;
 
 		TUniquePtr<FPolygroupSet> PolygroupConstraint;
@@ -654,11 +660,10 @@ UDynamicMesh* UGeometryScriptLibrary_MeshUVFunctions::AutoGenerateXAtlasMeshUVs(
 	{
 		if (EditMesh.IsCompact() == false)
 		{
-			UE::Geometry::AppendError(Debug, EGeometryScriptErrorType::InvalidInputs, LOCTEXT("AutoGenerateXAtlasMeshUVs_NonCompact", "AutoGenerateXAtlasMeshUVs: TargetMesh must be Compacted before running XAtlas"));
+			UE::Geometry::AppendError(Debug, EGeometryScriptErrorType::InvalidInputs, LOCTEXT("AutoGenerateXAtlasMeshUVs_NonCompact", "AutoGenerateXAtlasMeshUVs: TargetMesh is non-Compact, XAtlas cannot be run. Try calling CompactMesh to update TargetMesh."));
 			return;
 		}
 
-
 		const bool bFixOrientation = false;
 		//const bool bFixOrientation = true;
 		//FDynamicMesh3 FlippedMesh(EMeshComponents::FaceGroups);
@@ -668,7 +673,6 @@ UDynamicMesh* UGeometryScriptLibrary_MeshUVFunctions::AutoGenerateXAtlasMeshUVs(
 		//	FlippedMesh.ReverseOrientation(false);
 		//}
 
-
 		int32 NumVertices = EditMesh.VertexCount();
 		TArray<FVector3f> VertexBuffer;
 		VertexBuffer.SetNum(NumVertices);
