GeometryScript: add auto-compacting option, enabled by default, to Simplify and Remesh functions. Some nodes require compacted meshes and most geoscript users do not seem to understand compact vs non-compact meshes, this will avoid issues in common cases (and in the context of Simplify or Remesh operations, compacting is a small overhead).

diff --git a/Engine/Plugins/Experimental/GeometryScripting/Source/GeometryScriptingCore/Private/MeshRemeshFunctions.cpp b/Engine/Plugins/Experimental/GeometryScripting/Source/GeometryScriptingCore/Private/MeshRemeshFunctions.cpp
--- a/Engine/Plugins/Experimental/GeometryScripting/Source/GeometryScriptingCore/Private/MeshRemeshFunctions.cpp
+++ b/Engine/Plugins/Experimental/GeometryScripting/Source/GeometryScriptingCore/Private/MeshRemeshFunctions.cpp
@@ -130,11 +130,18 @@ UDynamicMesh* UGeometryScriptLibrary_RemeshingFunctions::ApplyUniformRemesh(
 			UE::Geometry::AppendError(Debug, EGeometryScriptErrorType::InvalidInputs, LOCTEXT("ApplyUniformRemesh_ComputeError", "ApplyUniformRemesh: Error computing result, returning input mesh"));
 		}
 
-		// if we discarded attributes, re-enable the standard attributes
+		// compact the input mesh if enabled
+		if (RemeshOptions.bAutoCompact)
+		{
+			EditMesh.CompactInPlace();
+		}
+
+		// if we discarded attributes, re-enable the standard attributes, and initialize with per-vertex normals
 		if (RemeshOptions.bDiscardAttributes)
 		{
 			EditMesh.EnableTriangleGroups();
 			EditMesh.EnableAttributes();
+			FMeshNormals::InitializeOverlayToPerVertexNormals(EditMesh.Attributes()->PrimaryNormals(), false);
 			EditMesh.Attributes()->EnableMaterialID();
 		}
 
diff --git a/Engine/Plugins/Experimental/GeometryScripting/Source/GeometryScriptingCore/Private/MeshSimplifyFunctions.cpp b/Engine/Plugins/Experimental/GeometryScripting/Source/GeometryScriptingCore/Private/MeshSimplifyFunctions.cpp
--- a/Engine/Plugins/Experimental/GeometryScripting/Source/GeometryScriptingCore/Private/MeshSimplifyFunctions.cpp
+++ b/Engine/Plugins/Experimental/GeometryScripting/Source/GeometryScriptingCore/Private/MeshSimplifyFunctions.cpp
@@ -40,6 +40,11 @@ UDynamicMesh* UGeometryScriptLibrary_MeshSimplifyFunctions::ApplySimplifyToPlana
 		Simplifier.CollapseMode = FQEMSimplification::ESimplificationCollapseModes::AverageVertexPosition;
 		Simplifier.SimplifyToMinimalPlanar( FMath::Max(0.00001, Options.AngleThreshold) );
 
+		if (Options.bAutoCompact)
+		{
+			EditMesh.CompactInPlace();
+		}
+
 	}, EDynamicMeshChangeType::GeneralEdit, EDynamicMeshAttributeChangeFlags::Unknown, false);
 
 	return TargetMesh;
@@ -82,6 +87,11 @@ UDynamicMesh* UGeometryScriptLibrary_MeshSimplifyFunctions::ApplySimplifyToPolyg
 		Simplifier.SimplificationAngleTolerance = Options.AngleThreshold;
 		Simplifier.Compute();
 
+		if (Options.bAutoCompact)
+		{
+			EditMesh.CompactInPlace();
+		}
+
 	}, EDynamicMeshChangeType::GeneralEdit, EDynamicMeshAttributeChangeFlags::Unknown, false);
 
 	return TargetMesh;
@@ -144,6 +154,11 @@ void DoSimplifyMesh(
 	{
 		Simplifier.SimplifyToVertexCount( FMath::Max(1,TargetCount) );
 	}
+
+	if (Options.bAutoCompact)
+	{
+		EditMesh.CompactInPlace();
+	}
 }
 
 
diff --git a/Engine/Plugins/Experimental/GeometryScripting/Source/GeometryScriptingCore/Public/GeometryScript/MeshRemeshFunctions.h b/Engine/Plugins/Experimental/GeometryScripting/Source/GeometryScriptingCore/Public/GeometryScript/MeshRemeshFunctions.h
--- a/Engine/Plugins/Experimental/GeometryScripting/Source/GeometryScriptingCore/Public/GeometryScript/MeshRemeshFunctions.h
+++ b/Engine/Plugins/Experimental/GeometryScripting/Source/GeometryScriptingCore/Public/GeometryScript/MeshRemeshFunctions.h
@@ -111,6 +111,10 @@ public:
 	/** Maximum Number of iterations of the Remeshing Strategy to apply to the Mesh. More iterations are generally more expensive (much moreso with bUseFullRemeshPasses = true) */
 	UPROPERTY(BlueprintReadWrite, Category = Options, meta = (UIMin = 0, ClampMin = 0))
 	int32 RemeshIterations = 20;
+
+	/** If enabled, the output mesh is automatically compacted to remove gaps in the index space. This is expensive and can be disabled by advanced users. */
+	UPROPERTY(BlueprintReadWrite, Category = Options)
+	bool bAutoCompact = true;
 };
 
 
diff --git a/Engine/Plugins/Experimental/GeometryScripting/Source/GeometryScriptingCore/Public/GeometryScript/MeshSimplifyFunctions.h b/Engine/Plugins/Experimental/GeometryScripting/Source/GeometryScriptingCore/Public/GeometryScript/MeshSimplifyFunctions.h
--- a/Engine/Plugins/Experimental/GeometryScripting/Source/GeometryScriptingCore/Public/GeometryScript/MeshSimplifyFunctions.h
+++ b/Engine/Plugins/Experimental/GeometryScripting/Source/GeometryScriptingCore/Public/GeometryScript/MeshSimplifyFunctions.h
@@ -17,6 +17,10 @@ public:
 	/** */
 	UPROPERTY(EditAnywhere, BlueprintReadWrite, Category = Options)
 	float AngleThreshold = 0.001;
+
+	/** If enabled, the simplified mesh is automatically compacted to remove gaps in the index space. This is expensive and can be disabled by advanced users. */
+	UPROPERTY(BlueprintReadWrite, Category = Options)
+	bool bAutoCompact = true;
 };
 
 USTRUCT(BlueprintType)
@@ -27,6 +31,10 @@ public:
 	/** */
 	UPROPERTY(EditAnywhere, BlueprintReadWrite, Category = Options)
 	float AngleThreshold = 0.001;
+
+	/** If enabled, the simplified mesh is automatically compacted to remove gaps in the index space. This is expensive and can be disabled by advanced users. */
+	UPROPERTY(BlueprintReadWrite, Category = Options)
+	bool bAutoCompact = true;
 };
 
 
@@ -62,6 +70,10 @@ public:
 
 	UPROPERTY(EditAnywhere, BlueprintReadWrite, Category = Options)
 	bool bRetainQuadricMemory = false;
+
+	/** If enabled, the simplified mesh is automatically compacted to remove gaps in the index space. This is expensive and can be disabled by advanced users. */
+	UPROPERTY(BlueprintReadWrite, Category = Options)
+	bool bAutoCompact = true;
 };
 
 
