Various improvements to (Generated)DynamicMeshActor. Harden registration/unregistration of GeneratedDynamicMeshActor with GeometryGen Subsystem, to avoid recomputing procedural mesh blueprints in objects duplicated for PIE and other temporary object copies. Add support for built-in "frozen" flag on GeneratedDynamicMeshActor, when toggled on, OnRebuildGeneratedMesh event will not be fired (ie mesh is "frozen" in it's current state). Add ability for GeneratedDynamicMeshActor to start a SlowTask dialog during long mesh recomputes, and provide a few functions that a BP can use to update the progress.

diff --git a/Engine/Plugins/Experimental/GeometryScripting/Source/GeometryScriptingEditor/Private/GeometryActors/EditorGeometryGenerationSubsystem.cpp b/Engine/Plugins/Experimental/GeometryScripting/Source/GeometryScriptingEditor/Private/GeometryActors/EditorGeometryGenerationSubsystem.cpp
--- a/Engine/Plugins/Experimental/GeometryScripting/Source/GeometryScriptingEditor/Private/GeometryActors/EditorGeometryGenerationSubsystem.cpp
+++ b/Engine/Plugins/Experimental/GeometryScripting/Source/GeometryScriptingEditor/Private/GeometryActors/EditorGeometryGenerationSubsystem.cpp
@@ -35,13 +35,11 @@ void UEditorGeometryGenerationManager::Shutdown()
 
 void UEditorGeometryGenerationManager::RegisterGeneratedMeshActor(AGeneratedDynamicMeshActor* Actor)
 {
-	ensure(ActiveGeneratedActors.Contains(Actor) == false);
 	ActiveGeneratedActors.Add(Actor);
 }
 
 void UEditorGeometryGenerationManager::UnregisterGeneratedMeshActor(AGeneratedDynamicMeshActor* Actor)
 {
-	ensure(ActiveGeneratedActors.Contains(Actor));
 	ActiveGeneratedActors.Remove(Actor);
 }
 
@@ -80,18 +78,20 @@ void UEditorGeometryGenerationSubsystem::OnShutdown()
 
 
 
-void UEditorGeometryGenerationSubsystem::RegisterGeneratedMeshActor(AGeneratedDynamicMeshActor* Actor)
+bool UEditorGeometryGenerationSubsystem::RegisterGeneratedMeshActor(AGeneratedDynamicMeshActor* Actor)
 {
 	if (bIsShuttingDown || GEditor == nullptr)		// subsystem no longer exists
 	{
-		return;
+		return false;
 	}
 
 	UEditorGeometryGenerationSubsystem* Subsystem = GEditor->GetEditorSubsystem<UEditorGeometryGenerationSubsystem>();
 	if (ensure(Subsystem))
 	{
 		Subsystem->GenerationManager->RegisterGeneratedMeshActor(Actor);
+		return true;
 	}
+	return false;
 }
 
 void UEditorGeometryGenerationSubsystem::UnregisterGeneratedMeshActor(AGeneratedDynamicMeshActor* Actor)
diff --git a/Engine/Plugins/Experimental/GeometryScripting/Source/GeometryScriptingEditor/Private/GeometryActors/GeneratedDynamicMeshActor.cpp b/Engine/Plugins/Experimental/GeometryScripting/Source/GeometryScriptingEditor/Private/GeometryActors/GeneratedDynamicMeshActor.cpp
--- a/Engine/Plugins/Experimental/GeometryScripting/Source/GeometryScriptingEditor/Private/GeometryActors/GeneratedDynamicMeshActor.cpp
+++ b/Engine/Plugins/Experimental/GeometryScripting/Source/GeometryScriptingEditor/Private/GeometryActors/GeneratedDynamicMeshActor.cpp
@@ -6,17 +6,18 @@
 #include "Editor/EditorEngine.h" // for CopyPropertiesForUnrelatedObjects
 #include "Engine/StaticMeshActor.h"
 
-#define LOCTEXT_NAMESPACE "AGeneratedDynamicMeshActor"
+#include "Misc/ScopedSlowTask.h"
 
+#define LOCTEXT_NAMESPACE "AGeneratedDynamicMeshActor"
 
 AGeneratedDynamicMeshActor::AGeneratedDynamicMeshActor(const FObjectInitializer& ObjectInitializer)
 	: Super(ObjectInitializer)
 {
-	RegisterWithGenerationManager();
 }
 
 AGeneratedDynamicMeshActor::~AGeneratedDynamicMeshActor()
 {
+	// make sure we are unregistered on destruction
 	UnregisterWithGenerationManager();
 }
 
@@ -25,9 +26,24 @@ void AGeneratedDynamicMeshActor::OnConstruction(const FTransform& Transform)
 {
 	Super::OnConstruction(Transform);
 
+	// Currently we rely on a mix of various Actor functions to tell us when to register w/ generation manager.
+	// If that turns out to not work reliably, we can do it here at every Construction script invocation
+	//RegisterWithGenerationManager();
+
 	bGeneratedMeshRebuildPending = true;
 }
 
+void AGeneratedDynamicMeshActor::PostLoad()
+{
+	Super::PostLoad();
+	RegisterWithGenerationManager();
+}
+
+void AGeneratedDynamicMeshActor::PostActorCreated()
+{
+	Super::PostActorCreated();
+	RegisterWithGenerationManager();
+}
 
 void AGeneratedDynamicMeshActor::Destroyed()
 {
@@ -36,27 +52,87 @@ void AGeneratedDynamicMeshActor::Destroyed()
 }
 
 
-void AGeneratedDynamicMeshActor::RegisterWithGenerationManager()
+
+void AGeneratedDynamicMeshActor::PreRegisterAllComponents()
 {
-	if (HasAnyFlags(RF_ClassDefaultObject))		// ignore for CDO
+	Super::PreRegisterAllComponents();
+
+	// Handle UWorld::AddToWorld(), i.e. turning on level visibility
+	if (const ULevel* Level = GetLevel())
 	{
-		return;
+		// This function gets called in editor all the time, we're only interested the case where our level is being added to world.
+		if (Level->bIsAssociatingLevel)
+		{
+			RegisterWithGenerationManager();
+		}
 	}
+}
 
-	if (bIsRegisteredWithGenerationManager == false)
+void AGeneratedDynamicMeshActor::PostUnregisterAllComponents()
+{
+	// Handle UWorld::RemoveFromWorld(), i.e. turning off level visibility
+	if (const ULevel* Level = GetLevel())
 	{
-		UEditorGeometryGenerationSubsystem::RegisterGeneratedMeshActor(this);
-		bIsRegisteredWithGenerationManager = true;
+		// This function gets called in editor all the time, we're only interested the case where our level is being removed from world.
+		if (Level->bIsDisassociatingLevel)
+		{
+			UnregisterWithGenerationManager();
+		}
 	}
+
+	Super::PostUnregisterAllComponents();
 }
 
-void AGeneratedDynamicMeshActor::UnregisterWithGenerationManager()
+
+#if WITH_EDITOR
+
+void AGeneratedDynamicMeshActor::PostEditUndo()
+{
+	Super::PostEditUndo();
+
+	// There is no direct signal that an Actor is being created or destroyed due to undo/redo.
+	// Currently (5.1) the checks below will tell us if the undo/redo has destroyed the
+	// Actor, and we assume otherwise it was created
+
+	if ( IsActorBeingDestroyed() || !IsValid(this) )	// equivalent to AActor::IsPendingKillPending()
+	{
+		UnregisterWithGenerationManager();
+	}
+	else
+	{
+		RegisterWithGenerationManager();
+	}
+}
+
+#endif
+
+
+
+void AGeneratedDynamicMeshActor::RegisterWithGenerationManager()
 {
-	if (HasAnyFlags(RF_ClassDefaultObject))		// ignore for CDO
+	// do not run mesh generation for CDO
+	if (HasAnyFlags(RF_ClassDefaultObject))
+	{
+		return;
+	}
+
+	// Do not run mesh generation for actors spawned for PIE.
+	// (If the Actor existed in Editor the existing UDynamicMesh will still be duplicated to PIE)
+	if (GetPackage()->GetPIEInstanceID() != INDEX_NONE)
 	{
 		return;
 	}
 
+	if (bIsRegisteredWithGenerationManager == false)
+	{
+		// this could fail if the subsystem is not initialized yet, or if it is shutting down
+		bIsRegisteredWithGenerationManager = UEditorGeometryGenerationSubsystem::RegisterGeneratedMeshActor(this);
+	}
+}
+
+
+void AGeneratedDynamicMeshActor::UnregisterWithGenerationManager()
+{
 	if (bIsRegisteredWithGenerationManager)
 	{
 		UEditorGeometryGenerationSubsystem::UnregisterGeneratedMeshActor(this);
@@ -66,37 +142,64 @@ void AGeneratedDynamicMeshActor::UnregisterWithGenerationManager()
 }
 
 
-
 void AGeneratedDynamicMeshActor::ExecuteRebuildGeneratedMeshIfPending()
 {
-	if (bGeneratedMeshRebuildPending)
+	if (bFrozen)
 	{
-		// Automatically defer collision updates during generated mesh rebuild. If we do not do this, then
-		// each mesh change will result in collision being rebuilt, which is very expensive !!
-		bool bEnabledDeferredCollision = false;
-		if (DynamicMeshComponent->bDeferCollisionUpdates == false)
-		{
-			DynamicMeshComponent->SetDeferredCollisionUpdatesEnabled(true, false);
-			bEnabledDeferredCollision = true;
-		}
+		return;
+	}
+	if (bGeneratedMeshRebuildPending == false)
+	{
+		return;
+	}
 
-		if (bResetOnRebuild && DynamicMeshComponent && DynamicMeshComponent->GetDynamicMesh())
-		{
-			DynamicMeshComponent->GetDynamicMesh()->Reset();
-		}
+	// Automatically defer collision updates during generated mesh rebuild. If we do not do this, then
+	// each mesh change will result in collision being rebuilt, which is very expensive !!
+	bool bEnabledDeferredCollision = false;
+	if (DynamicMeshComponent->bDeferCollisionUpdates == false)
+	{
+		DynamicMeshComponent->SetDeferredCollisionUpdatesEnabled(true, false);
+		bEnabledDeferredCollision = true;
+	}
+
+	if (bResetOnRebuild && DynamicMeshComponent && DynamicMeshComponent->GetDynamicMesh())
+	{
+		DynamicMeshComponent->GetDynamicMesh()->Reset();
+	}
+
+	FEditorScriptExecutionGuard Guard;
 
-		FEditorScriptExecutionGuard Guard;
+	if (bEnableRebuildProgress)
+	{
+		FScopedSlowTask Progress(this->NumProgressSteps, FText::FromString(this->ProgressMessage));
+		Progress.MakeDialogDelayed(this->DialogDelay, true);
+		ActiveSlowTask = &Progress;
+		CurProgressAccumSteps = 0;
 		OnRebuildGeneratedMesh(DynamicMeshComponent->GetDynamicMesh());
-		bGeneratedMeshRebuildPending = false;
+		ActiveSlowTask = nullptr;
+	}
+	else
+	{
+		OnRebuildGeneratedMesh(DynamicMeshComponent->GetDynamicMesh());
+	}
 
-		if (bEnabledDeferredCollision)
-		{
-			DynamicMeshComponent->SetDeferredCollisionUpdatesEnabled(false, true);
-		}
+	bGeneratedMeshRebuildPending = false;
+
+	if (bEnabledDeferredCollision)
+	{
+		DynamicMeshComponent->SetDeferredCollisionUpdatesEnabled(false, true);
 	}
 }
 
-
+void AGeneratedDynamicMeshActor::IncrementProgress(int NumSteps, FString Message)
+{
+	if (ActiveSlowTask)
+	{
+		int NextProgressAccumSteps = FMath::Min(CurProgressAccumSteps + NumSteps, NumProgressSteps);
+		ActiveSlowTask->EnterProgressFrame((float)(NextProgressAccumSteps - CurProgressAccumSteps), FText::FromString(Message));
+		CurProgressAccumSteps = NextProgressAccumSteps;
+	}
+}
 
 
 void AGeneratedDynamicMeshActor::CopyPropertiesToStaticMesh(AStaticMeshActor* StaticMeshActor, bool bCopyComponentMaterials)
diff --git a/Engine/Plugins/Experimental/GeometryScripting/Source/GeometryScriptingEditor/Public/GeometryActors/EditorGeometryGenerationSubsystem.h b/Engine/Plugins/Experimental/GeometryScripting/Source/GeometryScriptingEditor/Public/GeometryActors/EditorGeometryGenerationSubsystem.h
--- a/Engine/Plugins/Experimental/GeometryScripting/Source/GeometryScriptingEditor/Public/GeometryActors/EditorGeometryGenerationSubsystem.h
+++ b/Engine/Plugins/Experimental/GeometryScripting/Source/GeometryScriptingEditor/Public/GeometryActors/EditorGeometryGenerationSubsystem.h
@@ -23,7 +23,7 @@ class AGeneratedDynamicMeshActor;
  * 
  */
 UCLASS()
-class UEditorGeometryGenerationSubsystem : public UEditorSubsystem
+class GEOMETRYSCRIPTINGEDITOR_API UEditorGeometryGenerationSubsystem : public UEditorSubsystem
 {
 	GENERATED_BODY()
 
@@ -40,7 +40,7 @@ public:
 	// Static functions that simplify registration of an Actor w/ the Subsystem
 	//
 public:
-	static void RegisterGeneratedMeshActor(AGeneratedDynamicMeshActor* Actor);
+	static bool RegisterGeneratedMeshActor(AGeneratedDynamicMeshActor* Actor);
 	static void UnregisterGeneratedMeshActor(AGeneratedDynamicMeshActor* Actor);
 
 
@@ -60,7 +60,7 @@ private:
  * store registrations and provide a Tick()
  */
 UCLASS()
-class UEditorGeometryGenerationManager : public UObject, public FTickableEditorObject
+class GEOMETRYSCRIPTINGEDITOR_API UEditorGeometryGenerationManager : public UObject, public FTickableEditorObject
 {
 	GENERATED_BODY()
 
diff --git a/Engine/Plugins/Experimental/GeometryScripting/Source/GeometryScriptingEditor/Public/GeometryActors/GeneratedDynamicMeshActor.h b/Engine/Plugins/Experimental/GeometryScripting/Source/GeometryScriptingEditor/Public/GeometryActors/GeneratedDynamicMeshActor.h
--- a/Engine/Plugins/Experimental/GeometryScripting/Source/GeometryScriptingEditor/Public/GeometryActors/GeneratedDynamicMeshActor.h
+++ b/Engine/Plugins/Experimental/GeometryScripting/Source/GeometryScriptingEditor/Public/GeometryActors/GeneratedDynamicMeshActor.h
@@ -7,6 +7,7 @@
 #include "GeneratedDynamicMeshActor.generated.h"
 
 class AStaticMeshActor;
+struct FScopedSlowTask;
 
 
 /**
@@ -30,10 +31,13 @@ public:
 
 
 public:
-	/** If true, the internal UDynamicMesh will be cleared before the OnRebuildGeneratedMesh event is fired. */
-	UPROPERTY(Category = DynamicMeshActor, EditAnywhere, BlueprintReadWrite)
-	bool bResetOnRebuild = true;
+	/** If true, the DynamicMeshComponent will be "Frozen" in its current state, and automatic rebuilding will be disabled. However the DynamicMesh can still be modified by explicitly-called functions/etc. */
+	UPROPERTY(Category = "DynamicMeshActor", EditAnywhere, BlueprintReadWrite)
+	bool bFrozen = false;
 
+	/** If true, the DynamicMeshComponent will be cleared before the OnRebuildGeneratedMesh event is executed. */
+	UPROPERTY(Category = "DynamicMeshActor|Advanced", EditAnywhere, BlueprintReadWrite)
+	bool bResetOnRebuild = true;
 
 	/**
 	 * This event will be fired to notify the BP that the generated Mesh should
@@ -73,10 +77,20 @@ public:
 
 public:
 
-	// AActor overrides
+	//~ Begin UObject/AActor Interface
+	virtual void PostLoad() override;
+	virtual void PostActorCreated() override;
 	virtual void OnConstruction(const FTransform& Transform) override;
 	virtual void Destroyed() override;
 
+	// these are called when Actor exists in a sublevel that is hidden/shown
+	virtual void PreRegisterAllComponents() override;
+	virtual void PostUnregisterAllComponents() override;
+
+#if WITH_EDITOR
+	virtual void PostEditUndo() override;
+#endif
+
 
 protected:
 	// this internal flag is set in OnConstruction, and will cause ExecuteRebuildGeneratedMesh to
@@ -89,4 +103,39 @@ protected:
 
 	virtual void RegisterWithGenerationManager();
 	virtual void UnregisterWithGenerationManager();
+
+
+
+	//
+	// Support for progress dialog display during OnRebuildGeneratedMesh. This is intended to be
+	// used for long-running procedural generators that otherwise block the Editor. 
+	// 
+	// This functionality may be deprecated/removed in the future.
+	//
+
+public:
+	/** If enabled, a long-running OnRebuildGeneratedMesh event will show a progress dialog (The Script being executed must call IncrementProgress regularly) */
+	UPROPERTY(Category = "DynamicMeshActor|Progress", EditAnywhere, BlueprintReadWrite)
+	bool bEnableRebuildProgress = false;
+
+	/** Delay in seconds before the progress dialog is shown, if enabled */
+	UPROPERTY(Category = "DynamicMeshActor|Progress", EditAnywhere, BlueprintReadWrite)
+	float DialogDelay = 1;
+
+	/** Number of progress steps/ticks that the progress bar will be subdivided into */
+	UPROPERTY(Category = "DynamicMeshActor|Progress", EditAnywhere, BlueprintReadWrite)
+	int NumProgressSteps = 4;
+
+	/** The default progress message */
+	UPROPERTY(Category = "DynamicMeshActor|Progress", EditAnywhere, BlueprintReadWrite)
+	FString ProgressMessage = TEXT("Rebuilding Mesh...");
+
+	/** Call this function from within OnRebuildGeneratedMesh to update progress tracking. */
+	UFUNCTION(BlueprintCallable, Category = "DynamicMeshActor|Progress")
+	void IncrementProgress(int NumSteps, FString Message);
+
+protected:
+	FScopedSlowTask* ActiveSlowTask = nullptr;
+	int32 CurProgressAccumSteps = 0;
+
 };
diff --git a/Engine/Source/Runtime/GeometryFramework/Public/DynamicMeshActor.h b/Engine/Source/Runtime/GeometryFramework/Public/DynamicMeshActor.h
--- a/Engine/Source/Runtime/GeometryFramework/Public/DynamicMeshActor.h
+++ b/Engine/Source/Runtime/GeometryFramework/Public/DynamicMeshActor.h
@@ -36,7 +36,7 @@ public:
 
 public:
 	/** Control whether the DynamicMeshPool will be created when requested via GetComputeMeshPool() */
-	UPROPERTY(Category = DynamicMeshActor, EditAnywhere, BlueprintReadWrite)
+	UPROPERTY(Category = "DynamicMeshActor|Advanced", EditAnywhere, BlueprintReadWrite)
 	bool bEnableComputeMeshPool = true;
 protected:
 	/** The internal Mesh Pool, for use in DynamicMeshActor BPs. Use GetComputeMeshPool() to access this, as it will only be created on-demand if bEnableComputeMeshPool = true */
