GeometryScript: add AppendSweepPolyline function, mirrors UE5.1 AppendSweepPolygon but for open 2D path

diff --git a/Engine/Plugins/Experimental/GeometryScripting/Source/GeometryScriptingCore/Private/MeshPrimitiveFunctions.cpp b/Engine/Plugins/Experimental/GeometryScripting/Source/GeometryScriptingCore/Private/MeshPrimitiveFunctions.cpp
--- a/Engine/Plugins/Experimental/GeometryScripting/Source/GeometryScriptingCore/Private/MeshPrimitiveFunctions.cpp
+++ b/Engine/Plugins/Experimental/GeometryScripting/Source/GeometryScriptingCore/Private/MeshPrimitiveFunctions.cpp
@@ -548,6 +548,85 @@ UDynamicMesh* UGeometryScriptLibrary_MeshPrimitiveFunctions::AppendRevolvePath(
 }
 
 
+UDynamicMesh* UGeometryScriptLibrary_MeshPrimitiveFunctions::AppendSweepPolyline(
+	UDynamicMesh* TargetMesh,
+	FGeometryScriptPrimitiveOptions PrimitiveOptions,
+	FTransform Transform,
+	const TArray<FVector2D>& PolylineVertices,
+	const TArray<FTransform>& SweepPath,
+	const TArray<float>& PolylineTexParamU,
+	const TArray<float>& SweepPathTexParamV,
+	bool bLoop,
+	float StartScale,
+	float EndScale,
+	float RotationAngleDeg,
+	UGeometryScriptDebug* Debug)
+{
+	if (TargetMesh == nullptr)
+	{
+		UE::Geometry::AppendError(Debug, EGeometryScriptErrorType::InvalidInputs, LOCTEXT("AppendSweepPolyline_NullMesh", "AppendSweepPolyline: TargetMesh is Null"));
+		return TargetMesh;
+	}
+	if (PolylineVertices.Num() < 2)
+	{
+		UE::Geometry::AppendError(Debug, EGeometryScriptErrorType::InvalidInputs, LOCTEXT("AppendSweepPolyline_InvalidPolygon", "AppendSweepPolyline: Polyline array requires at least 2 positions"));
+		return TargetMesh;
+	}
+	if (SweepPath.Num() < 2)
+	{
+		UE::Geometry::AppendError(Debug, EGeometryScriptErrorType::InvalidInputs, LOCTEXT("AppendSweepPolyline_InvalidSweepPath", "AppendSweepPolyline: SweepPath array requires at least 2 positions"));
+		return TargetMesh;
+	}
+	if (PolylineTexParamU.Num() != 0 && PolylineTexParamU.Num() != PolylineVertices.Num())
+	{
+		UE::Geometry::AppendError(Debug, EGeometryScriptErrorType::InvalidInputs, LOCTEXT("AppendSweepPolyline_InvalidTexParamU", "AppendSweepPolyline: Polyline Texture Parameter U array must be same length as PolylineVertices"));
+		return TargetMesh;
+	}
+	if (SweepPathTexParamV.Num() != 0 && SweepPathTexParamV.Num() != (SweepPath.Num() + (bLoop?1:0)) )
+	{
+		UE::Geometry::AppendError(Debug, EGeometryScriptErrorType::InvalidInputs, LOCTEXT("AppendSweepPolyline_InvalidTexParamU", "AppendSweepPolyline: SweepPath Texture Parameter V array must be same length as SweepPath, or (length+1) for a closed loop"));
+		return TargetMesh;
+	}
+
+	FMatrix2d Rotation2D = FMatrix2d::RotationDeg(-RotationAngleDeg);
+	FGeneralizedCylinderGenerator SweepGen;
+	for (const FVector2D& Point : PolylineVertices)
+	{
+		SweepGen.CrossSection.AppendVertex(Rotation2D * FVector2d(Point.X, Point.Y));
+	}
+	for (const FTransform& SweepXForm : SweepPath)
+	{
+		SweepGen.Path.Add(SweepXForm.GetLocation());
+		FQuaterniond Rotation(SweepXForm.GetRotation());
+		SweepGen.PathFrames.Add(
+			FFrame3d(SweepXForm.GetLocation(), Rotation.AxisY(), Rotation.AxisZ(), Rotation.AxisX())
+		);
+		FVector3d Scale = SweepXForm.GetScale3D();
+		SweepGen.PathScales.Add(FVector2d(Scale.Y, Scale.Z));
+	}
+
+	SweepGen.bProfileCurveIsClosed = false;
+	SweepGen.bLoop = bLoop;
+	if (PolylineTexParamU.Num() > 0)
+	{
+		SweepGen.CrossSectionTexCoord = PolylineTexParamU;
+	}
+	if (SweepPathTexParamV.Num() > 0)
+	{
+		SweepGen.PathTexCoord = SweepPathTexParamV;
+	}
+	SweepGen.bPolygroupPerQuad = (PrimitiveOptions.PolygroupMode == EGeometryScriptPrimitivePolygroupMode::PerQuad);
+	SweepGen.InitialFrame = FFrame3d(SweepGen.Path[0]);
+	SweepGen.StartScale = StartScale;
+	SweepGen.EndScale = EndScale;
+
+	SweepGen.Generate();
+
+	AppendPrimitive(TargetMesh, &SweepGen, Transform, PrimitiveOptions);
+	return TargetMesh;
+}
+
+
 
 UDynamicMesh* UGeometryScriptLibrary_MeshPrimitiveFunctions::AppendSimpleExtrudePolygon(
 	UDynamicMesh* TargetMesh,
diff --git a/Engine/Plugins/Experimental/GeometryScripting/Source/GeometryScriptingCore/Public/GeometryScript/MeshPrimitiveFunctions.h b/Engine/Plugins/Experimental/GeometryScripting/Source/GeometryScriptingCore/Public/GeometryScript/MeshPrimitiveFunctions.h
--- a/Engine/Plugins/Experimental/GeometryScripting/Source/GeometryScriptingCore/Public/GeometryScript/MeshPrimitiveFunctions.h
+++ b/Engine/Plugins/Experimental/GeometryScripting/Source/GeometryScriptingCore/Public/GeometryScript/MeshPrimitiveFunctions.h
@@ -248,6 +248,35 @@ public:
 		bool bCapped = true,
 		UGeometryScriptDebug* Debug = nullptr);
 
+	/**
+	 * Sweep the given 2D PolylineVertices along the SweepPath specified as a set of FTransforms
+	 * If the 2D vertices are (U,V), then in the coordinate space of the FTransform, X points "along" the path,
+	 * Y points "right" (U) and Z points "up" (V).
+	 * @param PolylineVertices vertices of the open 2D path that will be swept along the SweepPath
+	 * @param SweepPath defines the 3D sweep path curve as a 3D poly-path, with rotation and scaling at each polypath vertex taken from the Transform
+	 * @param PolylineTexParamU defines U coordinate value for each element in PolylineVertices. Must be same length as PolylineVertices (ignored if length=0).
+	 * @param SweepPathTexParamV defines V coordinate value for each element in SweepPath. Must be same length as PolylineVertices if bLoop=false, length+1 if bLoop=true, and ignored if length=0.
+	 * @param bLoop if true, SweepPath is considered to be a Loop and a section connecting the end and start of the path is added (bCapped is ignored)
+	 * @param StartScale uniform scaling applied to the 2D polygon at the start of the path. Interpolated via arc length to EndScale at the end of the path.
+	 * @param EndScale uniform scaling applied to the 2D polygon at the end of the path
+	 * @param RotationAngleDeg Rotation applied to the 2D Polygon. Positive rotation rotates clockwise, ie Up/+Z/+V towards Right/+Y/+U. This Rotation is applied before any rotation in the SweepPath Transforms.
+	 */
+	UFUNCTION(BlueprintCallable, Category = "GeometryScript|Primitives", meta=(ScriptMethod, AutoCreateRefTerm="PolylineTexParamU, SweepPathTexParamV"))
+	static UPARAM(DisplayName = "Target Mesh") UDynamicMesh* 
+	AppendSweepPolyline( 
+		UDynamicMesh* TargetMesh, 
+		FGeometryScriptPrimitiveOptions PrimitiveOptions,
+		FTransform Transform,
+		const TArray<FVector2D>& PolylineVertices,
+		const TArray<FTransform>& SweepPath,
+		const TArray<float>& PolylineTexParamU,
+		const TArray<float>& SweepPathTexParamV,
+		bool bLoop = false,
+		float StartScale = 1.0f,
+		float EndScale = 1.0f,
+		float RotationAngleDeg = 0.0f,
+		UGeometryScriptDebug* Debug = nullptr);
+
 	/**
 	 * Polygon should be oriented counter-clockwise to produce a correctly-oriented shape, otherwise it will be inside-out
 	 * Polygon endpoint is not repeated.
