GeometryScript: fix loctext issue

diff --git a/Engine/Plugins/Experimental/GeometryScripting/Source/GeometryScriptingCore/Private/MeshPrimitiveFunctions.cpp b/Engine/Plugins/Experimental/GeometryScripting/Source/GeometryScriptingCore/Private/MeshPrimitiveFunctions.cpp
--- a/Engine/Plugins/Experimental/GeometryScripting/Source/GeometryScriptingCore/Private/MeshPrimitiveFunctions.cpp
+++ b/Engine/Plugins/Experimental/GeometryScripting/Source/GeometryScriptingCore/Private/MeshPrimitiveFunctions.cpp
@@ -584,7 +584,7 @@ UDynamicMesh* UGeometryScriptLibrary_MeshPrimitiveFunctions::AppendSweepPolyline
 	}
 	if (SweepPathTexParamV.Num() != 0 && SweepPathTexParamV.Num() != (SweepPath.Num() + (bLoop?1:0)) )
 	{
-		UE::Geometry::AppendError(Debug, EGeometryScriptErrorType::InvalidInputs, LOCTEXT("AppendSweepPolyline_InvalidTexParamU", "AppendSweepPolyline: SweepPath Texture Parameter V array must be same length as SweepPath, or (length+1) for a closed loop"));
+		UE::Geometry::AppendError(Debug, EGeometryScriptErrorType::InvalidInputs, LOCTEXT("AppendSweepPolyline_InvalidTexParamV", "AppendSweepPolyline: SweepPath Texture Parameter V array must be same length as SweepPath, or (length+1) for a closed loop"));
 		return TargetMesh;
 	}
 
