Modify various interfaces in ModelingComponents and GeometryScripting to use the engine FMeshNaniteSettings instead of just a "nanite proxy percentage" parameter or custom nanite settings struct. Update affected callers and Lyra assets.

This change does break existing Blueprints that have used the now-deprecated pins/properties. There is no way to gracefully handle this as deprecating the uproperty requires that it no longer be editable in the Blueprint.

diff --git a/Engine/Plugins/Experimental/GeometryScripting/Source/GeometryScriptingCore/Private/MeshAssetFunctions.cpp b/Engine/Plugins/Experimental/GeometryScripting/Source/GeometryScriptingCore/Private/MeshAssetFunctions.cpp
--- a/Engine/Plugins/Experimental/GeometryScripting/Source/GeometryScriptingCore/Private/MeshAssetFunctions.cpp
+++ b/Engine/Plugins/Experimental/GeometryScripting/Source/GeometryScriptingCore/Private/MeshAssetFunctions.cpp
@@ -306,11 +306,7 @@ UDynamicMesh*  UGeometryScriptLibrary_StaticMeshFunctions::CopyMeshToStaticMesh(
 
 	if (Options.bApplyNaniteSettings)
 	{
-		FMeshNaniteSettings NewNaniteSettings = ToStaticMeshAsset->NaniteSettings;
-		NewNaniteSettings.bEnabled = Options.NaniteSettings.bEnabled;
-		NewNaniteSettings.FallbackPercentTriangles = FMath::Clamp(Options.NaniteSettings.FallbackPercentTriangles / 100.0f, 0.0f, 1.0f);
-		NewNaniteSettings.FallbackRelativeError = FMath::Max(Options.NaniteSettings.FallbackRelativeError, 0.0f);
-		ToStaticMeshAsset->NaniteSettings = NewNaniteSettings;
+		ToStaticMeshAsset->NaniteSettings = Options.NewNaniteSettings;
 	}
 
 	if (Options.bReplaceMaterials)
diff --git a/Engine/Plugins/Experimental/GeometryScripting/Source/GeometryScriptingCore/Public/GeometryScript/MeshAssetFunctions.h b/Engine/Plugins/Experimental/GeometryScripting/Source/GeometryScriptingCore/Public/GeometryScript/MeshAssetFunctions.h
--- a/Engine/Plugins/Experimental/GeometryScripting/Source/GeometryScriptingCore/Public/GeometryScript/MeshAssetFunctions.h
+++ b/Engine/Plugins/Experimental/GeometryScripting/Source/GeometryScriptingCore/Public/GeometryScript/MeshAssetFunctions.h
@@ -5,6 +5,7 @@
 #include "CoreMinimal.h"
 #include "Kismet/BlueprintFunctionLibrary.h"
 #include "GeometryScript/GeometryScriptTypes.h"
+#include "Engine/EngineTypes.h"   // FMeshNaniteSettings
 #include "MeshAssetFunctions.generated.h"
 
 class UStaticMesh;
@@ -78,9 +79,13 @@ public:
 	UPROPERTY(EditAnywhere, BlueprintReadWrite, Category = Options)
 	bool bApplyNaniteSettings = false;
 
+	/** FGeometryScriptNaniteOptions is no longer in use */
+	UPROPERTY(meta = (DeprecatedProperty, DeprecationMessage = "Replaced FGeometryScriptNaniteOptions with usage of Engine FMeshNaniteSettings"))
+	FGeometryScriptNaniteOptions NaniteSettings_DEPRECATED = FGeometryScriptNaniteOptions();
+
 	/** Nanite Settings applied to the target Asset, if bApplyNaniteSettings = true */
-	UPROPERTY(EditAnywhere, BlueprintReadWrite, Category = Options)
-	FGeometryScriptNaniteOptions NaniteSettings = FGeometryScriptNaniteOptions();
+	UPROPERTY(EditAnywhere, BlueprintReadWrite, Category = Options, meta = (DisplayName = "Nanite Settings"))
+	FMeshNaniteSettings NewNaniteSettings;
 
 	UPROPERTY(EditAnywhere, BlueprintReadWrite, Category = Options)
 	bool bEmitTransaction = true;
diff --git a/Engine/Plugins/Experimental/GeometryScripting/Source/GeometryScriptingEditor/Private/CreateNewAssetUtilityFunctions.cpp b/Engine/Plugins/Experimental/GeometryScripting/Source/GeometryScriptingEditor/Private/CreateNewAssetUtilityFunctions.cpp
--- a/Engine/Plugins/Experimental/GeometryScripting/Source/GeometryScriptingEditor/Private/CreateNewAssetUtilityFunctions.cpp
+++ b/Engine/Plugins/Experimental/GeometryScripting/Source/GeometryScriptingEditor/Private/CreateNewAssetUtilityFunctions.cpp
@@ -177,7 +177,7 @@ UStaticMesh* UGeometryScriptLibrary_CreateNewAssetFunctions::CreateNewStaticMesh
 	AssetOptions.bEnableRecomputeNormals = Options.bEnableRecomputeNormals;
 	AssetOptions.bEnableRecomputeTangents = Options.bEnableRecomputeTangents;
 	AssetOptions.bGenerateNaniteEnabledMesh = Options.bEnableNanite;
-	AssetOptions.NaniteProxyTrianglePercent = Options.NaniteProxyTrianglePercent;
+	AssetOptions.NaniteSettings = Options.NaniteSettings;
 
 	AssetOptions.bCreatePhysicsBody = Options.bEnableCollision;
 	AssetOptions.CollisionType = Options.CollisionMode;
diff --git a/Engine/Plugins/Experimental/GeometryScripting/Source/GeometryScriptingEditor/Public/GeometryScript/CreateNewAssetUtilityFunctions.h b/Engine/Plugins/Experimental/GeometryScripting/Source/GeometryScriptingEditor/Public/GeometryScript/CreateNewAssetUtilityFunctions.h
--- a/Engine/Plugins/Experimental/GeometryScripting/Source/GeometryScriptingEditor/Public/GeometryScript/CreateNewAssetUtilityFunctions.h
+++ b/Engine/Plugins/Experimental/GeometryScripting/Source/GeometryScriptingEditor/Public/GeometryScript/CreateNewAssetUtilityFunctions.h
@@ -54,8 +54,13 @@ public:
 	UPROPERTY(BlueprintReadWrite, Category = Options)
 	bool bEnableNanite = false;
 
+	/** Nanite settings to set on new StaticMesh Asset */
 	UPROPERTY(BlueprintReadWrite, Category = Options)
-	float NaniteProxyTrianglePercent = 0.0f;
+	FMeshNaniteSettings NaniteSettings;
+
+	/** NaniteProxyTrianglePercent property has been replaced by the more general NaniteSettings property */
+	UPROPERTY(meta = (DeprecatedProperty, DeprecationMessage = "Replaced NaniteProxyTrianglePercent with usage of Engine FMeshNaniteSettings, use NaniteSettings property instead"))
+	float NaniteProxyTrianglePercent_DEPRECATED = 0.0f;
 
 	UPROPERTY(BlueprintReadWrite, Category = Options)
 	bool bEnableCollision = true;
diff --git a/Engine/Plugins/Experimental/MeshModelingToolsetExp/Source/GeometryProcessingAdapters/Private/ApproximateActorsImpl.cpp b/Engine/Plugins/Experimental/MeshModelingToolsetExp/Source/GeometryProcessingAdapters/Private/ApproximateActorsImpl.cpp
--- a/Engine/Plugins/Experimental/MeshModelingToolsetExp/Source/GeometryProcessingAdapters/Private/ApproximateActorsImpl.cpp
+++ b/Engine/Plugins/Experimental/MeshModelingToolsetExp/Source/GeometryProcessingAdapters/Private/ApproximateActorsImpl.cpp
@@ -1268,7 +1268,8 @@ UStaticMesh* FApproximateActorsImpl::EmitGeneratedMeshAsset(
 	MeshAssetOptions.SourceMeshes.DynamicMeshes.Add(FinalMesh);
 
 	MeshAssetOptions.bGenerateNaniteEnabledMesh = Options.bGenerateNaniteEnabledMesh;
-	MeshAssetOptions.NaniteProxyTrianglePercent = Options.NaniteProxyTrianglePercent;
+	MeshAssetOptions.NaniteSettings.bEnabled = Options.bGenerateNaniteEnabledMesh;
+	MeshAssetOptions.NaniteSettings.FallbackPercentTriangles = Options.NaniteProxyTrianglePercent / 100.0;	// NaniteSettings wants value in range 0-1
 
 	MeshAssetOptions.bSupportRayTracing = Options.bSupportRayTracing;
 	MeshAssetOptions.bAllowDistanceField = Options.bAllowDistanceField;
diff --git a/Engine/Plugins/Media/MediaPlate/Source/MediaPlateEditor/Private/MediaPlateCustomizationMesh.cpp b/Engine/Plugins/Media/MediaPlate/Source/MediaPlateEditor/Private/MediaPlateCustomizationMesh.cpp
--- a/Engine/Plugins/Media/MediaPlate/Source/MediaPlateEditor/Private/MediaPlateCustomizationMesh.cpp
+++ b/Engine/Plugins/Media/MediaPlate/Source/MediaPlateEditor/Private/MediaPlateCustomizationMesh.cpp
@@ -168,7 +168,7 @@ UStaticMesh* FMediaPlateCustomizationMesh::CreateStaticMeshAsset(FDynamicMesh3*
 	AssetOptions.bEnableRecomputeNormals = false;
 	AssetOptions.bEnableRecomputeTangents = false;
 	AssetOptions.bGenerateNaniteEnabledMesh = false;
-	AssetOptions.NaniteProxyTrianglePercent = 100.0f;
+	AssetOptions.NaniteSettings.FallbackPercentTriangles = 1.0;
 
 	AssetOptions.bCreatePhysicsBody = true;;
 	AssetOptions.CollisionType = ECollisionTraceFlag::CTF_UseComplexAsSimple;
diff --git a/Engine/Plugins/Runtime/MeshModelingToolset/Source/ModelingComponents/Public/ModelingObjectsCreationAPI.h b/Engine/Plugins/Runtime/MeshModelingToolset/Source/ModelingComponents/Public/ModelingObjectsCreationAPI.h
--- a/Engine/Plugins/Runtime/MeshModelingToolset/Source/ModelingComponents/Public/ModelingObjectsCreationAPI.h
+++ b/Engine/Plugins/Runtime/MeshModelingToolset/Source/ModelingComponents/Public/ModelingObjectsCreationAPI.h
@@ -5,6 +5,7 @@
 
 #include "CoreMinimal.h"
 #include "BodySetupEnums.h"
+#include "Engine/EngineTypes.h"   // FMeshNaniteSettings
 #include "MeshDescription.h"
 #include "DynamicMesh/DynamicMesh3.h"
 #include "Misc/Optional.h"
@@ -160,9 +161,12 @@ struct MODELINGCOMPONENTS_API FCreateMeshObjectParams
 	bool bEnableNanite = false;
 
 	/** Specify the Nanite proxy triangle percentage for this new mesh object */
-	UPROPERTY(Category = "CreateMeshObjectParams", EditAnywhere)
-	float NaniteProxyTrianglePercent = 100.0f;
+	UPROPERTY(meta = (DeprecatedProperty, DeprecationMessage = "Replaced NaniteProxyTrianglePercent with usage of Engine FMeshNaniteSettings"))
+	float NaniteProxyTrianglePercent_DEPRECATED = 100.0f;
 
+	/** Specify the Nanite Settings for this new mesh object, only used if bEnableNanite=true */
+	UPROPERTY(Category = "CreateMeshObjectParams", EditAnywhere)
+	FMeshNaniteSettings NaniteSettings = FMeshNaniteSettings();
 
 	//
 	// The Mesh Object should be created based on the mesh data structures below.
diff --git a/Engine/Plugins/Runtime/MeshModelingToolset/Source/ModelingComponentsEditorOnly/Private/AssetUtils/CreateStaticMeshUtil.cpp b/Engine/Plugins/Runtime/MeshModelingToolset/Source/ModelingComponentsEditorOnly/Private/AssetUtils/CreateStaticMeshUtil.cpp
--- a/Engine/Plugins/Runtime/MeshModelingToolset/Source/ModelingComponentsEditorOnly/Private/AssetUtils/CreateStaticMeshUtil.cpp
+++ b/Engine/Plugins/Runtime/MeshModelingToolset/Source/ModelingComponentsEditorOnly/Private/AssetUtils/CreateStaticMeshUtil.cpp
@@ -148,11 +148,10 @@ UE::AssetUtils::ECreateStaticMeshResult UE::AssetUtils::CreateStaticMeshAsset(
 	}
 
 	// Nanite options
-	NewStaticMesh->NaniteSettings.bEnabled = Options.bGenerateNaniteEnabledMesh;
+	NewStaticMesh->NaniteSettings.bEnabled = false;
 	if (Options.bGenerateNaniteEnabledMesh)
 	{
-		NewStaticMesh->NaniteSettings.FallbackPercentTriangles = Options.NaniteProxyTrianglePercent * 0.01f;
-		NewStaticMesh->NaniteSettings.PositionPrecision = MIN_int32;
+		NewStaticMesh->NaniteSettings = Options.NaniteSettings;
 	}
 
 	// Ray tracing
diff --git a/Engine/Plugins/Runtime/MeshModelingToolset/Source/ModelingComponentsEditorOnly/Private/EditorModelingObjectsCreationAPI.cpp b/Engine/Plugins/Runtime/MeshModelingToolset/Source/ModelingComponentsEditorOnly/Private/EditorModelingObjectsCreationAPI.cpp
--- a/Engine/Plugins/Runtime/MeshModelingToolset/Source/ModelingComponentsEditorOnly/Private/EditorModelingObjectsCreationAPI.cpp
+++ b/Engine/Plugins/Runtime/MeshModelingToolset/Source/ModelingComponentsEditorOnly/Private/EditorModelingObjectsCreationAPI.cpp
@@ -289,7 +289,7 @@ FCreateMeshObjectResult UEditorModelingObjectsCreationAPI::CreateStaticMeshAsset
 	AssetOptions.bEnableRecomputeNormals = CreateMeshParams.bEnableRecomputeNormals;
 	AssetOptions.bEnableRecomputeTangents = CreateMeshParams.bEnableRecomputeTangents;
 	AssetOptions.bGenerateNaniteEnabledMesh = CreateMeshParams.bEnableNanite;
-	AssetOptions.NaniteProxyTrianglePercent = CreateMeshParams.NaniteProxyTrianglePercent;
+	AssetOptions.NaniteSettings = CreateMeshParams.NaniteSettings;
 
 	AssetOptions.bCreatePhysicsBody = CreateMeshParams.bEnableCollision;
 	AssetOptions.CollisionType = CreateMeshParams.CollisionMode;
diff --git a/Engine/Plugins/Runtime/MeshModelingToolset/Source/ModelingComponentsEditorOnly/Public/AssetUtils/CreateStaticMeshUtil.h b/Engine/Plugins/Runtime/MeshModelingToolset/Source/ModelingComponentsEditorOnly/Public/AssetUtils/CreateStaticMeshUtil.h
--- a/Engine/Plugins/Runtime/MeshModelingToolset/Source/ModelingComponentsEditorOnly/Public/AssetUtils/CreateStaticMeshUtil.h
+++ b/Engine/Plugins/Runtime/MeshModelingToolset/Source/ModelingComponentsEditorOnly/Public/AssetUtils/CreateStaticMeshUtil.h
@@ -4,6 +4,7 @@
 
 #include "CoreMinimal.h"
 #include "DynamicMesh/DynamicMesh3.h"
+#include "Engine/EngineTypes.h"   // FMeshNaniteSettings
 #include "BodySetupEnums.h"
 
 class UPackage;
@@ -64,8 +65,10 @@ namespace AssetUtils
 
 		// Whether to generate a nanite-enabled mesh
 		bool bGenerateNaniteEnabledMesh = false;
-		// Percentage of triangles to reduce down to for generating a coarse proxy mesh from the Nanite mesh
-		float NaniteProxyTrianglePercent = 0;
+		// Percentage of triangles to reduce down to for generating a coarse proxy mesh from the Nanite mesh. DEPRECATED, use NaniteSettings instead.
+		float NaniteProxyTrianglePercent_DEPRECATED = 0;
+		// Nanite settings to set on static mesh asset
+		FMeshNaniteSettings NaniteSettings = FMeshNaniteSettings();
 
 		// Whether ray tracing will be supported on this mesh. Disable this to save memory if the generated mesh will only be rendered in the distance
 		bool bSupportRayTracing = true;
