Fix attributes getting lost by AppendMesh in geometry script

diff --git a/Engine/Plugins/Experimental/GeometryScripting/Source/GeometryScriptingCore/Private/MeshBasicEditFunctions.cpp b/Engine/Plugins/Experimental/GeometryScripting/Source/GeometryScriptingCore/Private/MeshBasicEditFunctions.cpp
--- a/Engine/Plugins/Experimental/GeometryScripting/Source/GeometryScriptingCore/Private/MeshBasicEditFunctions.cpp
+++ b/Engine/Plugins/Experimental/GeometryScripting/Source/GeometryScriptingCore/Private/MeshBasicEditFunctions.cpp
@@ -327,6 +327,18 @@ UDynamicMesh* UGeometryScriptLibrary_MeshBasicEditFunctions::DeleteSelectedTrian
 }
 
 
+void FGeometryScriptAppendMeshOptions::UpdateAttributesForCombineMode(FDynamicMesh3& Target, const FDynamicMesh3& Source)
+{
+	if (CombineMode == EGeometryScriptCombineAttributesMode::EnableAllMatching)
+	{
+		Target.EnableMatchingAttributes(Source, false);
+	}
+	else if (CombineMode == EGeometryScriptCombineAttributesMode::UseSource)
+	{
+		Target.EnableMatchingAttributes(Source, true);
+	}
+	// else the mode is UseTarget, which already corresponds to the default behavior for AppendMesh
+}
 
 
 UDynamicMesh* UGeometryScriptLibrary_MeshBasicEditFunctions::AppendMesh(
@@ -334,6 +346,7 @@ UDynamicMesh* UGeometryScriptLibrary_MeshBasicEditFunctions::AppendMesh(
 	UDynamicMesh* AppendMesh,
 	FTransform AppendTransform,
 	bool bDeferChangeNotifications,
+	FGeometryScriptAppendMeshOptions AppendOptions,
 	UGeometryScriptDebug* Debug)
 {
 	if (TargetMesh == nullptr)
@@ -351,6 +364,7 @@ UDynamicMesh* UGeometryScriptLibrary_MeshBasicEditFunctions::AppendMesh(
 	{
 		AppendMesh->ProcessMesh([&](const FDynamicMesh3& OtherMesh)
 		{
+			AppendOptions.UpdateAttributesForCombineMode(AppendToMesh, OtherMesh);
 			FTransformSRT3d XForm(AppendTransform);
 			FMeshIndexMappings TmpMappings;
 			FDynamicMeshEditor Editor(&AppendToMesh);
@@ -379,6 +393,7 @@ UDynamicMesh* UGeometryScriptLibrary_MeshBasicEditFunctions::AppendMeshTransform
 	FTransform ConstantTransform,
 	bool bConstantTransformIsRelative,
 	bool bDeferChangeNotifications,
+	FGeometryScriptAppendMeshOptions AppendOptions,
 	UGeometryScriptDebug* Debug)
 {
 	if (TargetMesh == nullptr)
@@ -401,6 +416,7 @@ UDynamicMesh* UGeometryScriptLibrary_MeshBasicEditFunctions::AppendMeshTransform
 	{
 		AppendMesh->ProcessMesh([&](const FDynamicMesh3& OtherMesh)
 		{
+			AppendOptions.UpdateAttributesForCombineMode(AppendToMesh, OtherMesh);
 			FMeshIndexMappings TmpMappings;
 			FDynamicMeshEditor Editor(&AppendToMesh);
 			const FDynamicMesh3* UseOtherMesh = &OtherMesh;
@@ -452,6 +468,7 @@ UDynamicMesh* UGeometryScriptLibrary_MeshBasicEditFunctions::AppendMeshRepeated(
 	int32 RepeatCount,
 	bool bApplyTransformToFirstInstance,
 	bool bDeferChangeNotifications,
+	FGeometryScriptAppendMeshOptions AppendOptions,
 	UGeometryScriptDebug* Debug)
 {
 	if (TargetMesh == nullptr)
@@ -475,6 +492,7 @@ UDynamicMesh* UGeometryScriptLibrary_MeshBasicEditFunctions::AppendMeshRepeated(
 		}
 		TargetMesh->EditMesh([&](FDynamicMesh3& AppendToMesh)
 		{
+			AppendOptions.UpdateAttributesForCombineMode(AppendToMesh, TmpMesh);
 			FMeshIndexMappings TmpMappings;
 			FDynamicMeshEditor Editor(&AppendToMesh);
 			for (int32 k = 0; k < RepeatCount; ++k)
diff --git a/Engine/Plugins/Experimental/GeometryScripting/Source/GeometryScriptingCore/Public/GeometryScript/MeshBasicEditFunctions.h b/Engine/Plugins/Experimental/GeometryScripting/Source/GeometryScriptingCore/Public/GeometryScript/MeshBasicEditFunctions.h
--- a/Engine/Plugins/Experimental/GeometryScripting/Source/GeometryScriptingCore/Public/GeometryScript/MeshBasicEditFunctions.h
+++ b/Engine/Plugins/Experimental/GeometryScripting/Source/GeometryScriptingCore/Public/GeometryScript/MeshBasicEditFunctions.h
@@ -9,6 +9,10 @@
 #include "MeshBasicEditFunctions.generated.h"
 
 class UDynamicMesh;
+namespace UE::Geometry
+{
+	class FDynamicMesh3;
+}
 
 
 USTRUCT(BlueprintType)
@@ -51,6 +55,37 @@ public:
 };
 
 
+// Options for how attributes from a source and target mesh are combined into the target mesh
+UENUM(BlueprintType)
+enum class EGeometryScriptCombineAttributesMode : uint8
+{
+	// Include attributes enabled on either the source or target mesh
+	EnableAllMatching,
+	// Only include attributes that are already enabled on the target mesh
+	UseTarget,
+	// Make the target mesh have only the attributes that are enabled on the source mesh
+	UseSource
+};
+
+
+/**
+ * Control how details like mesh attributes are handled when one mesh is appended to another
+ */
+USTRUCT(BlueprintType)
+struct GEOMETRYSCRIPTINGCORE_API FGeometryScriptAppendMeshOptions
+{
+	GENERATED_BODY()
+public:
+
+	// How attributes from each mesh are combined into the result
+	UPROPERTY(EditAnywhere, BlueprintReadWrite, Category = Options)
+	EGeometryScriptCombineAttributesMode CombineMode =
+		EGeometryScriptCombineAttributesMode::EnableAllMatching;
+
+	void UpdateAttributesForCombineMode(UE::Geometry::FDynamicMesh3& Target, const UE::Geometry::FDynamicMesh3& Source);
+};
+
+
 UCLASS(meta = (ScriptName = "GeometryScript_MeshEdits"))
 class GEOMETRYSCRIPTINGCORE_API UGeometryScriptLibrary_MeshBasicEditFunctions : public UBlueprintFunctionLibrary
 {
@@ -154,7 +189,8 @@ public:
 		bool bDeferChangeNotifications = false );
 
 	/**
-	 * Apply AppendTransform to AppendMesh and then add it's geometry to the TargetMesh
+	 * Apply AppendTransform to AppendMesh and then add its geometry to the TargetMesh
+	 * @param AppendOptions Control how details like mesh attributes are handled when one mesh is appended to another
 	 */
 	UFUNCTION(BlueprintCallable, Category = "GeometryScript|MeshEdits", meta=(ScriptMethod))
 	static UPARAM(DisplayName = "Target Mesh") UDynamicMesh* 
@@ -163,12 +199,14 @@ public:
 		UDynamicMesh* AppendMesh, 
 		FTransform AppendTransform, 
 		bool bDeferChangeNotifications = false,
+		FGeometryScriptAppendMeshOptions AppendOptions = FGeometryScriptAppendMeshOptions(),
 		UGeometryScriptDebug* Debug = nullptr);
 
 	/**
-	 * For each transform in AppendTransforms, apply the transform to AppendMesh and then add it's geometry to the TargetMesh.
+	 * For each transform in AppendTransforms, apply the transform to AppendMesh and then add its geometry to the TargetMesh.
 	 * @param ConstantTransform the Constant transform will be applied after each Append transform
 	 * @param bConstantTransformIsRelative if true, the Constant transform is applied "in the frame" of the Append Transform, otherwise it is applied as a second transform in local coordinates (ie rotate around the AppendTransform X axis, vs around the local X axis)
+	 * @param AppendOptions Control how details like mesh attributes are handled when one mesh is appended to another
 	 */
 	UFUNCTION(BlueprintCallable, Category = "GeometryScript|MeshEdits", meta=(ScriptMethod))
 	static UPARAM(DisplayName = "Target Mesh") UDynamicMesh* 
@@ -179,12 +217,14 @@ public:
 		FTransform ConstantTransform,
 		bool bConstantTransformIsRelative = true,
 		bool bDeferChangeNotifications = false,
+		FGeometryScriptAppendMeshOptions AppendOptions = FGeometryScriptAppendMeshOptions(),
 		UGeometryScriptDebug* Debug = nullptr);
 
 	/**
 	 * Repeatedly apply AppendTransform to the AppendMesh, each time adding the geometry to TargetMesh.
 	 * @param RepeatCount number of times to repeat the transform-append cycle
 	 * @param bApplyTransformToFirstInstance if true, the AppendTransform is applied before the first mesh append, otherwise it is applied after
+	 * @param AppendOptions Control how details like mesh attributes are handled when one mesh is appended to another
 	 */
 	UFUNCTION(BlueprintCallable, Category = "GeometryScript|MeshEdits", meta=(ScriptMethod))
 	static UPARAM(DisplayName = "Target Mesh") UDynamicMesh* 
@@ -195,6 +235,7 @@ public:
 		int RepeatCount = 1,
 		bool bApplyTransformToFirstInstance = true,
 		bool bDeferChangeNotifications = false,
+		FGeometryScriptAppendMeshOptions AppendOptions = FGeometryScriptAppendMeshOptions(),
 		UGeometryScriptDebug* Debug = nullptr);
 
 
