Geometry Script: Adding the ability to set hi-res mesh.

diff --git a/Engine/Plugins/Experimental/GeometryScripting/Source/GeometryScriptingCore/Private/MeshAssetFunctions.cpp b/Engine/Plugins/Experimental/GeometryScripting/Source/GeometryScriptingCore/Private/MeshAssetFunctions.cpp
--- a/Engine/Plugins/Experimental/GeometryScripting/Source/GeometryScriptingCore/Private/MeshAssetFunctions.cpp
+++ b/Engine/Plugins/Experimental/GeometryScripting/Source/GeometryScriptingCore/Private/MeshAssetFunctions.cpp
@@ -228,11 +228,6 @@ UDynamicMesh*  UGeometryScriptLibrary_StaticMeshFunctions::CopyMeshToStaticMesh(
 		UE::Geometry::AppendError(Debug, EGeometryScriptErrorType::InvalidInputs, LOCTEXT("CopyMeshToAsset_InvalidInput2", "CopyMeshToStaticMesh: ToStaticMeshAsset is Null"));
 		return FromDynamicMesh;
 	}
-	if (TargetLOD.bWriteHiResSource)
-	{
-		UE::Geometry::AppendError(Debug, EGeometryScriptErrorType::InvalidInputs, LOCTEXT("CopyMeshToAsset_Unsupported", "CopyMeshToStaticMesh: Writing HiResSource LOD is not yet supported"));
-		return FromDynamicMesh;
-	}
 
 #if WITH_EDITOR
 
@@ -266,77 +261,108 @@ UDynamicMesh*  UGeometryScriptLibrary_StaticMeshFunctions::CopyMeshToStaticMesh(
 	// mark as modified
 	ToStaticMeshAsset->Modify();
 
-	if (ToStaticMeshAsset->GetNumSourceModels() < UseLODIndex+1)
+	auto ConfigureBuildSettingsFromOptions = [](FStaticMeshSourceModel& SourceModel, FGeometryScriptCopyMeshToAssetOptions& Options)
+												{
+													FMeshBuildSettings& BuildSettings = SourceModel.BuildSettings;
+													BuildSettings.bRecomputeNormals  = Options.bEnableRecomputeNormals;
+													BuildSettings.bRecomputeTangents = Options.bEnableRecomputeTangents;
+													BuildSettings.bRemoveDegenerates = Options.bEnableRemoveDegenerates;
+												};
+
+	if (TargetLOD.bWriteHiResSource)
 	{
-		ToStaticMeshAsset->SetNumSourceModels(UseLODIndex+1);
-	}
+		// update model build settings
+		ConfigureBuildSettingsFromOptions(ToStaticMeshAsset->GetHiResSourceModel(), Options);
 
-	// configure build settings from options
-	FStaticMeshSourceModel& LODSourceModel = ToStaticMeshAsset->GetSourceModel(UseLODIndex);
-	FMeshBuildSettings& BuildSettings = LODSourceModel.BuildSettings;
-	BuildSettings.bRecomputeNormals = Options.bEnableRecomputeNormals;
-	BuildSettings.bRecomputeTangents = Options.bEnableRecomputeTangents;
-	BuildSettings.bRemoveDegenerates = Options.bEnableRemoveDegenerates;
+		ToStaticMeshAsset->ModifyHiResMeshDescription();
+		FMeshDescription* NewHiResMD = ToStaticMeshAsset->CreateHiResMeshDescription();
 
-	FMeshDescription* MeshDescription = ToStaticMeshAsset->GetMeshDescription(UseLODIndex);
-	if (MeshDescription == nullptr)
-	{
-		MeshDescription = ToStaticMeshAsset->CreateMeshDescription(UseLODIndex);
-	}
+		if (!ensure(NewHiResMD != nullptr))
+		{
+			UE::Geometry::AppendError(Debug, EGeometryScriptErrorType::InvalidInputs, LOCTEXT("CopyMeshToAsset_NullMeshDescription", "CopyMeshToAsset: MeshDescription for HiRes is null?"));
+			return FromDynamicMesh;
+		}
 
-	// mark mesh description for modify
-	ToStaticMeshAsset->ModifyMeshDescription(UseLODIndex);
+		FConversionToMeshDescriptionOptions ConversionOptions;
+		FDynamicMeshToMeshDescription Converter(ConversionOptions);
+		FromDynamicMesh->ProcessMesh([&](const FDynamicMesh3& ReadMesh)
+			{
+				Converter.Convert(&ReadMesh, *NewHiResMD, !Options.bEnableRecomputeTangents);
+			});
 
-	if (!ensure(MeshDescription != nullptr))
-	{
-		UE::Geometry::AppendError(Debug, EGeometryScriptErrorType::InvalidInputs,
-			FText::Format(LOCTEXT("CopyMeshToAsset_NullMeshDescription", "CopyMeshToAsset: MeshDescription for LOD {0} is null?"), FText::AsNumber(UseLODIndex)));
-		return FromDynamicMesh;
+
+		ToStaticMeshAsset->CommitHiResMeshDescription();
 	}
+	else
+	{ 
 
-	FConversionToMeshDescriptionOptions ConversionOptions;
-	FDynamicMeshToMeshDescription Converter(ConversionOptions);
-	FromDynamicMesh->ProcessMesh([&](const FDynamicMesh3& ReadMesh)
-	{
-		Converter.Convert(&ReadMesh, *MeshDescription, !BuildSettings.bRecomputeTangents);
-	});
+		if (ToStaticMeshAsset->GetNumSourceModels() < UseLODIndex+1)
+		{
+			ToStaticMeshAsset->SetNumSourceModels(UseLODIndex+1);
+		}
 
-	// Setting to prevent the standard static mesh reduction from running and replacing the render LOD.
-	FStaticMeshSourceModel& ThisSourceModel = ToStaticMeshAsset->GetSourceModel(UseLODIndex);
-	ThisSourceModel.ReductionSettings.PercentTriangles = 1.f;
-	ThisSourceModel.ReductionSettings.PercentVertices = 1.f;
+		// update model build settings
+		ConfigureBuildSettingsFromOptions(ToStaticMeshAsset->GetSourceModel(UseLODIndex), Options);
 
-	if (Options.bApplyNaniteSettings)
-	{
-		ToStaticMeshAsset->NaniteSettings = Options.NewNaniteSettings;
-	}
+		FMeshDescription* MeshDescription = ToStaticMeshAsset->GetMeshDescription(UseLODIndex);
+		if (MeshDescription == nullptr)
+		{
+			MeshDescription = ToStaticMeshAsset->CreateMeshDescription(UseLODIndex);
+		}
 
-	if (Options.bReplaceMaterials)
-	{
-		bool bHaveSlotNames = (Options.NewMaterialSlotNames.Num() == Options.NewMaterials.Num());
+		// mark mesh description for modify
+		ToStaticMeshAsset->ModifyMeshDescription(UseLODIndex);
 
-		TArray<FStaticMaterial> NewMaterials;
-		for (int32 k = 0; k < Options.NewMaterials.Num(); ++k)
+		if (!ensure(MeshDescription != nullptr))
 		{
-			FStaticMaterial NewMaterial;
-			NewMaterial.MaterialInterface = Options.NewMaterials[k];
-			FName UseSlotName = (bHaveSlotNames && Options.NewMaterialSlotNames[k] != NAME_None) ? Options.NewMaterialSlotNames[k] :
-				UE::AssetUtils::GenerateNewMaterialSlotName(NewMaterials, NewMaterial.MaterialInterface, k);
-
-			NewMaterial.MaterialSlotName = UseSlotName;
-			NewMaterial.ImportedMaterialSlotName = UseSlotName;
-			NewMaterial.UVChannelData = FMeshUVChannelInfo(1.f);		// this avoids an ensure in  UStaticMesh::GetUVChannelData
-			NewMaterials.Add(NewMaterial);
+			UE::Geometry::AppendError(Debug, EGeometryScriptErrorType::InvalidInputs,
+				FText::Format(LOCTEXT("CopyMeshToAsset_NullMeshDescription", "CopyMeshToAsset: MeshDescription for LOD {0} is null?"), FText::AsNumber(UseLODIndex)));
+			return FromDynamicMesh;
 		}
 
-		ToStaticMeshAsset->SetStaticMaterials(NewMaterials);
+		FConversionToMeshDescriptionOptions ConversionOptions;
+		FDynamicMeshToMeshDescription Converter(ConversionOptions);
+		FromDynamicMesh->ProcessMesh([&](const FDynamicMesh3& ReadMesh)
+		{
+			Converter.Convert(&ReadMesh, *MeshDescription, !Options.bEnableRecomputeTangents);
+		});
 
-		// Reset the section info map
-		ToStaticMeshAsset->GetSectionInfoMap().Clear();
-	}
+		// Setting to prevent the standard static mesh reduction from running and replacing the render LOD.
+		FStaticMeshSourceModel& ThisSourceModel = ToStaticMeshAsset->GetSourceModel(UseLODIndex);
+		ThisSourceModel.ReductionSettings.PercentTriangles = 1.f;
+		ThisSourceModel.ReductionSettings.PercentVertices = 1.f;
+
+		if (Options.bApplyNaniteSettings)
+		{
+			ToStaticMeshAsset->NaniteSettings = Options.NewNaniteSettings;
+		}
 
-	ToStaticMeshAsset->CommitMeshDescription(UseLODIndex);
+		if (Options.bReplaceMaterials)
+		{
+			bool bHaveSlotNames = (Options.NewMaterialSlotNames.Num() == Options.NewMaterials.Num());
+
+			TArray<FStaticMaterial> NewMaterials;
+			for (int32 k = 0; k < Options.NewMaterials.Num(); ++k)
+			{
+				FStaticMaterial NewMaterial;
+				NewMaterial.MaterialInterface = Options.NewMaterials[k];
+				FName UseSlotName = (bHaveSlotNames && Options.NewMaterialSlotNames[k] != NAME_None) ? Options.NewMaterialSlotNames[k] :
+					UE::AssetUtils::GenerateNewMaterialSlotName(NewMaterials, NewMaterial.MaterialInterface, k);
+
+				NewMaterial.MaterialSlotName = UseSlotName;
+				NewMaterial.ImportedMaterialSlotName = UseSlotName;
+				NewMaterial.UVChannelData = FMeshUVChannelInfo(1.f);		// this avoids an ensure in  UStaticMesh::GetUVChannelData
+				NewMaterials.Add(NewMaterial);
+			}
+
+			ToStaticMeshAsset->SetStaticMaterials(NewMaterials);
+
+			// Reset the section info map
+			ToStaticMeshAsset->GetSectionInfoMap().Clear();
+		}
 
+		ToStaticMeshAsset->CommitMeshDescription(UseLODIndex);
+	}
 
 	if (Options.bDeferMeshPostEditChange == false)
 	{
