GeometryScript : MeshAssetFunctions - Fix  Localization error

diff --git a/Engine/Plugins/Experimental/GeometryScripting/Source/GeometryScriptingCore/Private/MeshAssetFunctions.cpp b/Engine/Plugins/Experimental/GeometryScripting/Source/GeometryScriptingCore/Private/MeshAssetFunctions.cpp
--- a/Engine/Plugins/Experimental/GeometryScripting/Source/GeometryScriptingCore/Private/MeshAssetFunctions.cpp
+++ b/Engine/Plugins/Experimental/GeometryScripting/Source/GeometryScriptingCore/Private/MeshAssetFunctions.cpp
@@ -279,7 +279,7 @@ UDynamicMesh*  UGeometryScriptLibrary_StaticMeshFunctions::CopyMeshToStaticMesh(
 
 		if (!ensure(NewHiResMD != nullptr))
 		{
-			UE::Geometry::AppendError(Debug, EGeometryScriptErrorType::InvalidInputs, LOCTEXT("CopyMeshToAsset_NullMeshDescription", "CopyMeshToAsset: MeshDescription for HiRes is null?"));
+			UE::Geometry::AppendError(Debug, EGeometryScriptErrorType::InvalidInputs, LOCTEXT("CopyMeshToAsset_NullHiResMeshDescription", "CopyMeshToAsset: MeshDescription for HiRes is null?"));
 			return FromDynamicMesh;
 		}
 
