fix Geometry Script function SetMeshSelectionVertexColor so that it works even if mesh does not have vertex colors defined. Fix DynamicMeshComponent ConstantColor override, the override was ignored if the mesh had vertex colors defined.

diff --git a/Engine/Plugins/Experimental/GeometryScripting/Source/GeometryScriptingCore/Private/MeshVertexColorFunctions.cpp b/Engine/Plugins/Experimental/GeometryScripting/Source/GeometryScriptingCore/Private/MeshVertexColorFunctions.cpp
--- a/Engine/Plugins/Experimental/GeometryScripting/Source/GeometryScriptingCore/Private/MeshVertexColorFunctions.cpp
+++ b/Engine/Plugins/Experimental/GeometryScripting/Source/GeometryScriptingCore/Private/MeshVertexColorFunctions.cpp
@@ -114,6 +114,22 @@ UDynamicMesh* UGeometryScriptLibrary_MeshVertexColorFunctions::SetMeshSelectionV
 		}
 		FDynamicMeshColorOverlay* Colors = EditMesh.Attributes()->PrimaryColors();
 
+		// If we created a new color overlay, we are going to initialize the vertex colors to white.
+		if (bCreated)
+		{
+			TArray<int32> ElemIDs;
+			ElemIDs.SetNum(EditMesh.MaxVertexID());
+			for (int32 VertexID : EditMesh.VertexIndicesItr())
+			{
+				ElemIDs[VertexID] = Colors->AppendElement(FVector4f::One());
+			}
+			for (int32 TriangleID : EditMesh.TriangleIndicesItr())
+			{
+				FIndex3i Triangle = EditMesh.GetTriangle(TriangleID);
+				Colors->SetTriangle(TriangleID, FIndex3i(ElemIDs[Triangle.A], ElemIDs[Triangle.B], ElemIDs[Triangle.C]));
+			}
+		}
+
 		if (bCreateColorSeam)
 		{
 			FLinearColor NewColor = CombineColors(FLinearColor::Black, Color, Flags);
diff --git a/Engine/Source/Runtime/GeometryFramework/Private/Components/BaseDynamicMeshSceneProxy.cpp b/Engine/Source/Runtime/GeometryFramework/Private/Components/BaseDynamicMeshSceneProxy.cpp
--- a/Engine/Source/Runtime/GeometryFramework/Private/Components/BaseDynamicMeshSceneProxy.cpp
+++ b/Engine/Source/Runtime/GeometryFramework/Private/Components/BaseDynamicMeshSceneProxy.cpp
@@ -14,6 +14,7 @@ FBaseDynamicMeshSceneProxy::FBaseDynamicMeshSceneProxy(UBaseDynamicMeshComponent
 	if (Component->GetColorOverrideMode() == EDynamicMeshComponentColorOverrideMode::Constant)
 	{
 		ConstantVertexColor = Component->GetConstantOverrideColor();
+		bIgnoreVertexColors = true;
 	}
 
 	bUsePerTriangleNormals = Component->GetFlatShadingEnabled();
