GeometryScript: undo deprecation of nanite settings fields in create and update static mesh functions, as it breaks Lyra blueprints

diff --git a/Engine/Plugins/Experimental/GeometryScripting/Source/GeometryScriptingCore/Public/GeometryScript/MeshAssetFunctions.h b/Engine/Plugins/Experimental/GeometryScripting/Source/GeometryScriptingCore/Public/GeometryScript/MeshAssetFunctions.h
--- a/Engine/Plugins/Experimental/GeometryScripting/Source/GeometryScriptingCore/Public/GeometryScript/MeshAssetFunctions.h
+++ b/Engine/Plugins/Experimental/GeometryScripting/Source/GeometryScriptingCore/Public/GeometryScript/MeshAssetFunctions.h
@@ -79,9 +79,9 @@ public:
 	UPROPERTY(EditAnywhere, BlueprintReadWrite, Category = Options)
 	bool bApplyNaniteSettings = false;
 
-	/** FGeometryScriptNaniteOptions is no longer in use */
-	UPROPERTY(meta = (DeprecatedProperty, DeprecationMessage = "Replaced FGeometryScriptNaniteOptions with usage of Engine FMeshNaniteSettings"))
-	FGeometryScriptNaniteOptions NaniteSettings_DEPRECATED = FGeometryScriptNaniteOptions();
+	/** Replaced FGeometryScriptNaniteOptions with usage of Engine FMeshNaniteSettings */
+	UPROPERTY(EditAnywhere, BlueprintReadWrite, Category = Deprecated, AdvancedDisplay, meta=(DisplayName = "DEPRECATED NANITE SETTING") )
+	FGeometryScriptNaniteOptions NaniteSettings = FGeometryScriptNaniteOptions();
 
 	/** Nanite Settings applied to the target Asset, if bApplyNaniteSettings = true */
 	UPROPERTY(EditAnywhere, BlueprintReadWrite, Category = Options, meta = (DisplayName = "Nanite Settings"))
diff --git a/Engine/Plugins/Experimental/GeometryScripting/Source/GeometryScriptingEditor/Public/GeometryScript/CreateNewAssetUtilityFunctions.h b/Engine/Plugins/Experimental/GeometryScripting/Source/GeometryScriptingEditor/Public/GeometryScript/CreateNewAssetUtilityFunctions.h
--- a/Engine/Plugins/Experimental/GeometryScripting/Source/GeometryScriptingEditor/Public/GeometryScript/CreateNewAssetUtilityFunctions.h
+++ b/Engine/Plugins/Experimental/GeometryScripting/Source/GeometryScriptingEditor/Public/GeometryScript/CreateNewAssetUtilityFunctions.h
@@ -58,9 +58,9 @@ public:
 	UPROPERTY(BlueprintReadWrite, Category = Options)
 	FMeshNaniteSettings NaniteSettings;
 
-	/** NaniteProxyTrianglePercent property has been replaced by the more general NaniteSettings property */
-	UPROPERTY(meta = (DeprecatedProperty, DeprecationMessage = "Replaced NaniteProxyTrianglePercent with usage of Engine FMeshNaniteSettings, use NaniteSettings property instead"))
-	float NaniteProxyTrianglePercent_DEPRECATED = 0.0f;
+	/** Replaced NaniteProxyTrianglePercent with usage of Engine FMeshNaniteSettings, use NaniteSettings property instead */
+	UPROPERTY(EditAnywhere, BlueprintReadWrite, Category = Deprecated, AdvancedDisplay, meta = (DisplayName = "DEPRECATED NANITE SETTING"))
+	float NaniteProxyTrianglePercent = 0.0f;
 
 	UPROPERTY(BlueprintReadWrite, Category = Options)
 	bool bEnableCollision = true;
