GeometryScript: functions for working with 3D AABBs

diff --git a/Engine/Plugins/Experimental/GeometryScripting/Source/GeometryScriptingCore/Private/ShapeFunctions.cpp b/Engine/Plugins/Experimental/GeometryScripting/Source/GeometryScriptingCore/Private/ShapeFunctions.cpp
new file mode 100644
--- /dev/null
+++ b/Engine/Plugins/Experimental/GeometryScripting/Source/GeometryScriptingCore/Private/ShapeFunctions.cpp
@@ -0,0 +1,128 @@
+// Copyright Epic Games, Inc. All Rights Reserved.
+
+#include "GeometryScript/ShapeFunctions.h"
+#include "VectorTypes.h"
+
+using namespace UE::Geometry;
+
+#define LOCTEXT_NAMESPACE "UGeometryScriptLibrary_BoxFunctions"
+
+FBox UGeometryScriptLibrary_BoxFunctions::MakeBoxFromCenterSize(const FVector& Center, const FVector& Dimensions)
+{
+	FVector Extents( FMathd::Max(0, Dimensions.X * 0.5), FMathd::Max(0, Dimensions.Y * 0.5), FMathd::Max(0, Dimensions.Z * 0.5) );
+	return FBox(Center - Extents, Center + Extents);
+}
+
+void UGeometryScriptLibrary_BoxFunctions::GetBoxCenterSize(const FBox& Box, FVector& Center, FVector& Dimensions)
+{
+	FVector Extents;
+	Box.GetCenterAndExtents(Center, Extents);
+	Dimensions = 2.0 * Extents;
+}
+
+FVector UGeometryScriptLibrary_BoxFunctions::GetBoxCorner(const FBox& Box, int CornerIndex)
+{
+	CornerIndex = FMath::Clamp(CornerIndex, 0, 7);
+	switch (CornerIndex)		// matches Box.GetVertices()
+	{
+	default:
+	case 0: return Box.Min;
+	case 1: return FVector(Box.Min.X, Box.Min.Y, Box.Max.Z);
+	case 2: return FVector(Box.Min.X, Box.Max.Y, Box.Min.Z);
+	case 3: return FVector(Box.Max.X, Box.Min.Y, Box.Min.Z);
+	case 4: return FVector(Box.Max.X, Box.Max.Y, Box.Min.Z);
+	case 5: return FVector(Box.Max.X, Box.Min.Y, Box.Max.Z);
+	case 6: return FVector(Box.Min.X, Box.Max.Y, Box.Max.Z);
+	case 7: return Box.Max;
+	}
+}
+
+FVector UGeometryScriptLibrary_BoxFunctions::GetBoxFaceCenter(const FBox& Box, int FaceIndex, FVector& FaceNormal)
+{
+	FaceIndex = FMath::Clamp(FaceIndex, 0, 5);
+	FVector Center = 0.5 * (Box.Min + Box.Max);
+	switch (FaceIndex)
+	{
+	default:
+	case 0: FaceNormal = FVector(0, 0, -1); return FVector(Center.X, Center.Y, Box.Min.Z);
+	case 1: FaceNormal = FVector(0, 0, 1); return FVector(Center.X, Center.Y, Box.Max.Z);
+	case 2: FaceNormal = FVector(0, -1, 0); return FVector(Center.X, Box.Min.Y, Center.Z);
+	case 3: FaceNormal = FVector(0, 1, 0); return FVector(Center.X, Box.Max.Y, Center.Z);
+	case 4: FaceNormal = FVector(-1, 0, 0); return FVector(Box.Min.X, Center.Y, Center.Z);
+	case 5: FaceNormal = FVector(1, 0, 0); return FVector(Box.Max.X, Center.Y, Center.Z);
+	}
+}
+
+void UGeometryScriptLibrary_BoxFunctions::GetBoxVolumeArea(const FBox& Box, double& Volume, double& SurfaceArea)
+{
+	FVector Dimensions = Box.GetSize();
+	double AreaXY = Dimensions.X * Dimensions.Y;
+	double AreaXZ = Dimensions.X * Dimensions.Z;
+	double AreaYZ = Dimensions.Y * Dimensions.Z;
+	SurfaceArea = 2.0*AreaXY + 2.0*AreaXZ + 2.0*AreaYZ;
+	Volume = Dimensions.X * Dimensions.Y * Dimensions.Z;
+}
+
+FBox UGeometryScriptLibrary_BoxFunctions::GetExpandedBox(const FBox& Box, const FVector& ExpandBy)
+{
+	FBox Result = Box.ExpandBy(ExpandBy);
+	// ExpandBy with negative expansion factor does not clamp to original box center
+	for (int32 j = 0; j < 3; ++j)
+	{
+		if (Result.Min[j] > Result.Max[j])
+		{
+			Result.Min[j] = Result.Max[j] = 0.5*(Box.Min[j] + Box.Max[j]);
+		}
+	}
+	return Result;
+}
+
+FBox UGeometryScriptLibrary_BoxFunctions::GetTransformedBox(const FBox& Box, const FTransform& Transform)
+{
+	return Box.TransformBy(Transform);
+}
+
+
+bool UGeometryScriptLibrary_BoxFunctions::TestBoxBoxIntersection(const FBox& Box1, const FBox& Box2)
+{
+	return Box1.Intersect(Box2);
+}
+
+
+FBox UGeometryScriptLibrary_BoxFunctions::FindBoxBoxIntersection(const FBox& Box1, const FBox& Box2, bool& bIsIntersecting)
+{
+	bIsIntersecting = Box1.Intersect(Box2);
+	return Box1.Overlap(Box2);
+}
+
+double UGeometryScriptLibrary_BoxFunctions::GetBoxBoxDistance(const FBox& Box1, const FBox& Box2)
+{
+	double DistSqr = FMathd::Max(0.0, Box1.ComputeSquaredDistanceToBox(Box2));
+	return FMathd::Sqrt(DistSqr);
+}
+
+bool UGeometryScriptLibrary_BoxFunctions::TestPointInsideBox(const FBox& Box, const FVector& Point, bool bConsiderOnBoxAsInside)
+{
+	return (bConsiderOnBoxAsInside) ? Box.IsInsideOrOn(Point) : Box.IsInside(Point);
+}
+
+FVector UGeometryScriptLibrary_BoxFunctions::FindClosestPointOnBox(const FBox& Box, const FVector& Point, bool& bIsInside)
+{
+	bIsInside = Box.IsInside(Point);
+	return Box.GetClosestPointTo(Point);
+}
+
+double UGeometryScriptLibrary_BoxFunctions::GetBoxPointDistance(const FBox& Box, const FVector& Point)
+{
+	double DistSqr = FMathd::Max(0.0, Box.ComputeSquaredDistanceToPoint(Point));
+	return FMathd::Sqrt(DistSqr);
+}
+
+bool UGeometryScriptLibrary_BoxFunctions::TestBoxSphereIntersection(const FBox& Box, const FVector& SphereCenter, double SphereRadius)
+{
+	return FMath::SphereAABBIntersection(SphereCenter, SphereRadius*SphereRadius, Box);
+}
+
+
+
+#undef LOCTEXT_NAMESPACE
\ No newline at end of file
diff --git a/Engine/Plugins/Experimental/GeometryScripting/Source/GeometryScriptingCore/Public/GeometryScript/ShapeFunctions.h b/Engine/Plugins/Experimental/GeometryScripting/Source/GeometryScriptingCore/Public/GeometryScript/ShapeFunctions.h
new file mode 100644
--- /dev/null
+++ b/Engine/Plugins/Experimental/GeometryScripting/Source/GeometryScriptingCore/Public/GeometryScript/ShapeFunctions.h
@@ -0,0 +1,113 @@
+// Copyright Epic Games, Inc. All Rights Reserved.
+
+#pragma once
+
+#include "CoreMinimal.h"
+#include "Kismet/BlueprintFunctionLibrary.h"
+#include "Math/Box.h"
+#include "ShapeFunctions.generated.h"
+
+
+
+UCLASS(meta = (ScriptName = "GeometryScript_Box"))
+class GEOMETRYSCRIPTINGCORE_API UGeometryScriptLibrary_BoxFunctions : public UBlueprintFunctionLibrary
+{
+	GENERATED_BODY()
+public:
+
+	/**
+	 * Create a Box from a Center point and X/Y/Z Dimensions (*not* Extents, which are half-dimensions)
+	 */
+	UFUNCTION(BlueprintPure, Category = "GeometryScript|Shapes|Box")
+	static FBox MakeBoxFromCenterSize(const FVector& Center, const FVector& Dimensions);
+
+	/**
+	 * Get the Center point and X/Y/Z Dimensions of a Box (full dimensions, not Extents)
+	 */
+	UFUNCTION(BlueprintPure, Category = "GeometryScript|Shapes|Box", meta=(ScriptMethod))
+	static void GetBoxCenterSize(const FBox& Box, FVector& Center, FVector& Dimensions);
+
+
+	/**
+	 * Get the position of a corner of the Box. Corners are indexed from 0 to 7, using
+	 * an ordering where 0 is the Min corner, 1/2/3 are +Z/+Y/+X from the Min corner, 
+	 * 7 is the Max corner, and 4/5/6 are -Z/-Y/-X from the Max corner.
+	 */
+	UFUNCTION(BlueprintPure, Category = "GeometryScript|Shapes|Box", meta=(ScriptMethod))
+	static FVector GetBoxCorner(const FBox& Box, int CornerIndex);
+
+	/**
+	 * Get the position of the center of a face of the Box. Faces are indexed from 0 to 5,
+	 * using an ordering where 0/1 are the MinZ/MaxZ faces, 2/3 are MinY/MaxY, and 4/5 are MinX/MaxX
+	 * @param FaceNormal returned Normal vector of the identified face
+	 */
+	UFUNCTION(BlueprintPure, Category = "GeometryScript|Shapes|Box", meta=(ScriptMethod))
+	static FVector GetBoxFaceCenter(const FBox& Box, int FaceIndex, FVector& FaceNormal);
+
+	/**
+	 * Get the Volume and Surface Area of a Box
+	 */
+	UFUNCTION(BlueprintPure, Category = "GeometryScript|Shapes|Box", meta=(ScriptMethod))
+	static void GetBoxVolumeArea(const FBox& Box, double& Volume, double& SurfaceArea);
+
+	/**
+	 * Get the input Box expanded by adding the ExpandBy parameter to both the Min and Max.
+	 * Dimensions will be clamped to the center point if any of ExpandBy are larger than half the box size
+	 */
+	UFUNCTION(BlueprintPure, Category = "GeometryScript|Shapes|Box", meta=(ScriptMethod))
+	static FBox GetExpandedBox(const FBox& Box, const FVector& ExpandBy);
+
+	/**
+	 * Apply the input Transform to the corners of the input Box, and return the new Box containing those points
+	 */
+	UFUNCTION(BlueprintPure, Category = "GeometryScript|Shapes|Box", meta=(ScriptMethod))
+	static FBox GetTransformedBox(const FBox& Box, const FTransform& Transform);
+
+	/**
+	 * Test if Box1 and Box2 intersect
+	 */
+	UFUNCTION(BlueprintPure, Category = "GeometryScript|Shapes|Box", meta=(ScriptMethod))
+	static bool TestBoxBoxIntersection(const FBox& Box1, const FBox& Box2);
+
+	/**
+	 * Find the Box formed by the intersection of Box1 and Box2
+	 * @param bIsIntersecting if the boxes do not intersect, this will be returned as false, otherwise true
+	 */
+	UFUNCTION(BlueprintPure, Category = "GeometryScript|Shapes|Box", meta=(ScriptMethod))
+	static FBox FindBoxBoxIntersection(const FBox& Box1, const FBox& Box2, bool& bIsIntersecting);
+
+	/**
+	 * Calculate the minimum distance between Box1 and Box2
+	 */
+	UFUNCTION(BlueprintPure, Category = "GeometryScript|Shapes|Box", meta=(ScriptMethod))
+	static double GetBoxBoxDistance(const FBox& Box1, const FBox& Box2);
+
+	/**
+	 * Test if a Point is inside the Box, returning true if so, otherwise false
+	 * @param bConsiderOnBoxAsInside if true, a point lying on the box face is considered "inside", otherwise it is considered "outside"
+	 */
+	UFUNCTION(BlueprintPure, Category = "GeometryScript|Shapes|Box", meta=(ScriptMethod))
+	static bool TestPointInsideBox(const FBox& Box, const FVector& Point, bool bConsiderOnBoxAsInside = true);
+
+	/**
+	 * Find the point on the faces of the Box that is closest to the input Point.
+	 * If the Point is inside the Box, it is returned, ie points Inside do not project to the Box Faces
+	 * @param bIsInside if the Point is inside the Box, this will return as true, otherwise false
+	 */
+	UFUNCTION(BlueprintPure, Category = "GeometryScript|Shapes|Box", meta=(ScriptMethod))
+	static FVector FindClosestPointOnBox(const FBox& Box, const FVector& Point, bool& bIsInside);
+
+	/**
+	 * Calculate the minimum distance between the Box and the Point
+	 */
+	UFUNCTION(BlueprintPure, Category = "GeometryScript|Shapes|Box", meta=(ScriptMethod))
+	static double GetBoxPointDistance(const FBox& Box, const FVector& Point);
+
+	/**
+	 * Check if the Box intersects a Sphere defined by the SphereCenter and SphereRadius
+	 */
+	UFUNCTION(BlueprintPure, Category = "GeometryScript|Shapes|Box", meta=(ScriptMethod))
+	static bool TestBoxSphereIntersection(const FBox& Box, const FVector& SphereCenter, double SphereRadius);
+
+};
+
