GeometryScript: add new MeshBoneWeightFunctions library with functions MeshHasBoneWeights, GetMaxBoneWeightIndex, GetVertexBoneWeights, GetLargestVertexBoneWeight, SetVertexBoneWeights

diff --git a/Engine/Plugins/Experimental/GeometryScripting/Source/GeometryScriptingCore/Private/MeshBoneWeightFunctions.cpp b/Engine/Plugins/Experimental/GeometryScripting/Source/GeometryScriptingCore/Private/MeshBoneWeightFunctions.cpp
new file mode 100644
--- /dev/null
+++ b/Engine/Plugins/Experimental/GeometryScripting/Source/GeometryScriptingCore/Private/MeshBoneWeightFunctions.cpp
@@ -0,0 +1,230 @@
+// Copyright Epic Games, Inc. All Rights Reserved.
+
+#include "GeometryScript/MeshBoneWeightFunctions.h"
+
+#include "DynamicMesh/DynamicMesh3.h"
+#include "DynamicMesh/DynamicMeshAttributeSet.h"
+#include "DynamicMesh/DynamicVertexSkinWeightsAttribute.h"
+#include "UDynamicMesh.h"
+#include "BoneWeights.h"
+
+using namespace UE::Geometry;
+using namespace UE::AnimationCore;
+
+#define LOCTEXT_NAMESPACE "UGeometryScriptLibrary_MeshBoneWeightFunctions"
+
+
+template<typename ReturnType> 
+ReturnType SimpleMeshBoneWeightQuery(
+	UDynamicMesh* Mesh, FGeometryScriptBoneWeightProfile Profile, 
+	bool& bIsValidBoneWeights, ReturnType DefaultValue, 
+	TFunctionRef<ReturnType(const FDynamicMesh3& Mesh, const FDynamicMeshVertexSkinWeightsAttribute& SkinWeights)> QueryFunc)
+{
+	bIsValidBoneWeights = false;
+	ReturnType RetVal = DefaultValue;
+	if (Mesh)
+	{
+		Mesh->ProcessMesh([&](const FDynamicMesh3& ReadMesh)
+		{
+			if (ReadMesh.HasAttributes())
+			{
+				if ( const FDynamicMeshVertexSkinWeightsAttribute* BoneWeights = ReadMesh.Attributes()->GetSkinWeightsAttribute(Profile.GetProfileName()) )
+				{
+					bIsValidBoneWeights = true;
+					RetVal = QueryFunc(ReadMesh, *BoneWeights);
+				}
+			}
+		});
+	}
+	return RetVal;
+}
+
+
+template<typename ReturnType> 
+ReturnType SimpleMeshBoneWeightEdit(
+	UDynamicMesh* Mesh, FGeometryScriptBoneWeightProfile Profile, 
+	bool& bIsValidBoneWeights, ReturnType DefaultValue, 
+	TFunctionRef<ReturnType(FDynamicMesh3& Mesh, FDynamicMeshVertexSkinWeightsAttribute& SkinWeights)> EditFunc)
+{
+	bIsValidBoneWeights = false;
+	ReturnType RetVal = DefaultValue;
+	if (Mesh)
+	{
+		Mesh->EditMesh([&](FDynamicMesh3& ReadMesh)
+		{
+			if (ReadMesh.HasAttributes())
+			{
+				if ( FDynamicMeshVertexSkinWeightsAttribute* BoneWeights = ReadMesh.Attributes()->GetSkinWeightsAttribute(Profile.GetProfileName()) )
+				{
+					bIsValidBoneWeights = true;
+					RetVal = EditFunc(ReadMesh, *BoneWeights);
+				}
+			}
+		}, EDynamicMeshChangeType::GeneralEdit, EDynamicMeshAttributeChangeFlags::Unknown, false);
+	}
+	return RetVal;
+}
+
+
+
+
+UDynamicMesh* UGeometryScriptLibrary_MeshBoneWeightFunctions::MeshHasBoneWeights(
+	UDynamicMesh* TargetMesh,
+	bool& bHasBoneWeights,
+	FGeometryScriptBoneWeightProfile Profile)
+{
+	bool bHasBoneWeightProfile = false;
+	bool bOK = SimpleMeshBoneWeightQuery<bool>(TargetMesh, Profile, bHasBoneWeightProfile, false,
+		[&](const FDynamicMesh3& Mesh, const FDynamicMeshVertexSkinWeightsAttribute& SkinWeights) { return true; });
+	bHasBoneWeights = bHasBoneWeightProfile;
+	return TargetMesh;
+}
+
+
+UDynamicMesh* UGeometryScriptLibrary_MeshBoneWeightFunctions::GetMaxBoneWeightIndex(
+	UDynamicMesh* TargetMesh,
+	bool& bHasBoneWeights,
+	int& MaxBoneIndex,
+	FGeometryScriptBoneWeightProfile Profile)
+{
+	MaxBoneIndex = -1;
+	bool bHasBoneWeightProfile = false;
+	bool bOK = SimpleMeshBoneWeightQuery<bool>(TargetMesh, Profile, bHasBoneWeightProfile, false,
+		[&](const FDynamicMesh3& Mesh, const FDynamicMeshVertexSkinWeightsAttribute& SkinWeights) 
+		{ 
+			for (int32 VertexID : Mesh.VertexIndicesItr())
+			{
+				FBoneWeights BoneWeights;
+				SkinWeights.GetValue(VertexID, BoneWeights);
+				int32 Num = BoneWeights.Num();
+				for (int32 k = 0; k < Num; ++k)
+				{
+					MaxBoneIndex = FMathd::Max(MaxBoneIndex, BoneWeights[k].GetBoneIndex());
+				}
+			}
+			return true;
+		});
+	bHasBoneWeights = bHasBoneWeightProfile;
+	return TargetMesh;
+}
+
+
+
+UDynamicMesh* UGeometryScriptLibrary_MeshBoneWeightFunctions::GetVertexBoneWeights(
+	UDynamicMesh* TargetMesh,
+	int VertexID,
+	TArray<FGeometryScriptBoneWeight>& BoneWeightsOut,
+	bool& bHasValidBoneWeights,
+	FGeometryScriptBoneWeightProfile Profile)
+{
+	bool bHasBoneWeightProfile = false;
+	bHasValidBoneWeights = SimpleMeshBoneWeightQuery<bool>(TargetMesh, Profile, bHasBoneWeightProfile, false,
+		[&](const FDynamicMesh3& Mesh, const FDynamicMeshVertexSkinWeightsAttribute& SkinWeights)
+	{
+		if (Mesh.IsVertex(VertexID))
+		{
+			FBoneWeights BoneWeights;
+			SkinWeights.GetValue(VertexID, BoneWeights);
+			int32 Num = BoneWeights.Num();
+			BoneWeightsOut.SetNum(Num);
+			for (int32 k = 0; k < Num; ++k)
+			{
+				FGeometryScriptBoneWeight NewBoneWeight;
+				NewBoneWeight.BoneIndex = BoneWeights[k].GetBoneIndex();
+				NewBoneWeight.Weight = BoneWeights[k].GetWeight();
+				BoneWeightsOut.Add(NewBoneWeight);
+			}
+		}
+		return BoneWeightsOut.Num() > 0;
+	});
+
+	return TargetMesh;
+}
+
+
+UDynamicMesh* UGeometryScriptLibrary_MeshBoneWeightFunctions::GetLargestVertexBoneWeight(
+	UDynamicMesh* TargetMesh,
+	int VertexID,
+	FGeometryScriptBoneWeight& BoneWeight,
+	bool& bHasValidBoneWeights,
+	FGeometryScriptBoneWeightProfile Profile)
+{
+	bHasValidBoneWeights = false;
+	bool bHasBoneWeightProfile = false;
+	FBoneWeight FoundMax = SimpleMeshBoneWeightQuery<FBoneWeight>(TargetMesh, Profile, bHasBoneWeightProfile, FBoneWeight(),
+	[&](const FDynamicMesh3& Mesh, const FDynamicMeshVertexSkinWeightsAttribute& SkinWeights) 
+	{ 
+		FBoneWeight MaxBoneWeight = FBoneWeight();
+		if (Mesh.IsVertex(VertexID))
+		{
+			bHasValidBoneWeights = true;
+			float MaxWeight = 0;
+			FBoneWeights BoneWeights;
+			SkinWeights.GetValue(VertexID, BoneWeights);
+			int32 Num = BoneWeights.Num();
+			for (int32 k = 0; k < Num; ++k)
+			{
+				const FBoneWeight& BoneWeight = BoneWeights[k];
+				if (BoneWeight.GetWeight() > MaxWeight)
+				{
+					MaxWeight = BoneWeight.GetWeight();
+					MaxBoneWeight = BoneWeight;
+				}
+			}
+		}
+		else
+		{
+			UE_LOG(LogGeometry, Warning, TEXT("GetLargestMeshBoneWeight: VertexID %d does not exist"), VertexID);
+		}
+		return MaxBoneWeight;
+	});
+	
+	if (bHasValidBoneWeights)
+	{
+		BoneWeight.BoneIndex = FoundMax.GetBoneIndex();
+		BoneWeight.Weight = FoundMax.GetWeight();
+	}
+
+	return TargetMesh;
+}
+
+
+
+
+
+UDynamicMesh* UGeometryScriptLibrary_MeshBoneWeightFunctions::SetVertexBoneWeights(
+	UDynamicMesh* TargetMesh,
+	int VertexID,
+	const TArray<FGeometryScriptBoneWeight>& BoneWeights,
+	bool& bHasValidBoneWeights,
+	FGeometryScriptBoneWeightProfile Profile)
+{
+	bool bHasBoneWeightProfile = false;
+	bHasValidBoneWeights = SimpleMeshBoneWeightEdit<bool>(TargetMesh, Profile, bHasBoneWeightProfile, false,
+		[&](FDynamicMesh3& Mesh, FDynamicMeshVertexSkinWeightsAttribute& SkinWeights)
+	{
+		if (Mesh.IsVertex(VertexID))
+		{
+			int32 Num = BoneWeights.Num();
+			TArray<FBoneWeight> NewWeightsList;
+			for (int32 k = 0; k < Num; ++k)
+			{
+				FBoneWeight NewWeight;
+				NewWeight.SetBoneIndex(BoneWeights[k].BoneIndex);
+				NewWeight.SetWeight(BoneWeights[k].Weight);
+				NewWeightsList.Add(NewWeight);
+			}
+			FBoneWeights NewBoneWeights = FBoneWeights::Create(NewWeightsList);
+			SkinWeights.SetValue(VertexID, NewBoneWeights);
+			return true;
+		}
+		return false;
+	});
+
+	return TargetMesh;
+}
+
+
+
+
+#undef LOCTEXT_NAMESPACE
\ No newline at end of file
diff --git a/Engine/Plugins/Experimental/GeometryScripting/Source/GeometryScriptingCore/Public/GeometryScript/MeshBoneWeightFunctions.h b/Engine/Plugins/Experimental/GeometryScripting/Source/GeometryScriptingCore/Public/GeometryScript/MeshBoneWeightFunctions.h
new file mode 100644
--- /dev/null
+++ b/Engine/Plugins/Experimental/GeometryScripting/Source/GeometryScriptingCore/Public/GeometryScript/MeshBoneWeightFunctions.h
@@ -0,0 +1,121 @@
+// Copyright Epic Games, Inc. All Rights Reserved.
+
+#pragma once
+
+#include "CoreMinimal.h"
+#include "Kismet/BlueprintFunctionLibrary.h"
+#include "GeometryScript/GeometryScriptTypes.h"
+#include "SkeletalMeshAttributes.h"
+#include "MeshBoneWeightFunctions.generated.h"
+
+class UDynamicMesh;
+
+
+USTRUCT(BlueprintType, meta = (DisplayName = "Bone Weights"))
+struct GEOMETRYSCRIPTINGCORE_API FGeometryScriptBoneWeight
+{
+	GENERATED_BODY()
+
+	UPROPERTY(BlueprintReadWrite, Category = BoneWeights)
+	int32 BoneIndex = 0;
+
+	UPROPERTY(BlueprintReadWrite, Category = BoneWeights)
+	float Weight = 0;
+};
+
+
+USTRUCT(BlueprintType, meta = (DisplayName = "Bone Weights Profile"))
+struct GEOMETRYSCRIPTINGCORE_API FGeometryScriptBoneWeightProfile
+{
+	GENERATED_BODY()
+
+	UPROPERTY(BlueprintReadWrite, Category = BoneWeights)
+	FName ProfileName = FSkeletalMeshAttributes::DefaultSkinWeightProfileName;
+
+	FName GetProfileName() const { return ProfileName; }
+};
+
+
+
+UCLASS(meta = (ScriptName = "GeometryScript_BoneWeights"))
+class GEOMETRYSCRIPTINGCORE_API UGeometryScriptLibrary_MeshBoneWeightFunctions : public UBlueprintFunctionLibrary
+{
+	GENERATED_BODY()
+public:
+
+	/**
+	 * Check whether the TargetMesh has a per-vertex Bone/Skin Weight Attribute set
+	 * @param bHasBoneWeights will be returned true if the requested bone weight profile exists
+	 * @param Profile identifier for the bone/skin weight profile
+	 */
+	UFUNCTION(BlueprintPure, Category = "GeometryScript|MeshQueries|BoneWeights", meta=(ScriptMethod))
+	static UPARAM(DisplayName = "Target Mesh") UDynamicMesh* 
+	MeshHasBoneWeights( 
+		UDynamicMesh* TargetMesh,
+		bool& bHasBoneWeights,
+		FGeometryScriptBoneWeightProfile Profile = FGeometryScriptBoneWeightProfile() );
+	
+	/**
+	 * Determine the largest bone weight index that exists on the Mesh
+	 * @param bHasBoneWeights will be returned true if the requested bone weight profile exists
+	 * @param MaxBoneIndex maximum bone index will be returned here, or -1 if no bone indices exist
+	 * @param Profile identifier for the bone/skin weight profile
+	 */
+	UFUNCTION(BlueprintCallable, Category = "GeometryScript|MeshQueries|BoneWeights", meta=(ScriptMethod))
+	static UPARAM(DisplayName = "Target Mesh") UDynamicMesh* 
+	GetMaxBoneWeightIndex( 
+		UDynamicMesh* TargetMesh,
+		bool& bHasBoneWeights,
+		int& MaxBoneIndex,
+		FGeometryScriptBoneWeightProfile Profile = FGeometryScriptBoneWeightProfile() );
+
+	/**
+	 * Return an array of Bone/Skin Weights at a given vertex of TargetMesh
+	 * @param VertexID requested vertex
+	 * @param BoneWeights output array of bone index/weight pairs for the given Vertex
+	 * @param bHasValidBoneWeights will be returned as true if the vertex has bone weights in the given profile, ie BoneWeights is valid
+	 * @param Profile identifier for the bone/skin weight profile
+	 */
+	UFUNCTION(BlueprintCallable, Category = "GeometryScript|MeshQueries|BoneWeights", meta=(ScriptMethod))
+	static UPARAM(DisplayName = "Target Mesh") UDynamicMesh* 
+	GetVertexBoneWeights( 
+		UDynamicMesh* TargetMesh,
+		int VertexID,
+		TArray<FGeometryScriptBoneWeight>& BoneWeights,
+		bool& bHasValidBoneWeights,
+		FGeometryScriptBoneWeightProfile Profile = FGeometryScriptBoneWeightProfile() );
+
+	/**
+	 * Return the Bone/Skin Weight with the maximum weight at a given vertex of TargetMesh
+	 * @param VertexID requested vertex
+	 * @param BoneWeight the bone index and weight that was found to have the maximum weight
+	 * @param bHasValidBoneWeights will be returned as true if the vertex has bone weights in the given profile and a largest weight was found
+	 * @param Profile identifier for the bone/skin weight profile
+	 */
+	UFUNCTION(BlueprintCallable, Category = "GeometryScript|MeshQueries|BoneWeights", meta=(ScriptMethod))
+	static UPARAM(DisplayName = "Target Mesh") UDynamicMesh* 
+	GetLargestVertexBoneWeight( 
+		UDynamicMesh* TargetMesh,
+		int VertexID,
+		FGeometryScriptBoneWeight& BoneWeight,
+		bool& bHasValidBoneWeights,
+		FGeometryScriptBoneWeightProfile Profile = FGeometryScriptBoneWeightProfile() );
+
+	/**
+	 * Set the Bone/Skin Weights at a given vertex of TargetMesh
+	 * @param VertexID vertex to update
+	 * @param BoneWeights input array of bone index/weight pairs for the Vertex
+	 * @param bIsValidVertexID will be returned as true if the vertex ID is valid
+	 * @param Profile identifier for the bone/skin weight profile
+	 */
+	UFUNCTION(BlueprintCallable, Category = "GeometryScript|MeshQueries|BoneWeights", meta=(ScriptMethod))
+	static UPARAM(DisplayName = "Target Mesh") UDynamicMesh* 
+	SetVertexBoneWeights( 
+		UDynamicMesh* TargetMesh,
+		int VertexID,
+		const TArray<FGeometryScriptBoneWeight>& BoneWeights,
+		bool& bIsValidVertexID,
+		FGeometryScriptBoneWeightProfile Profile = FGeometryScriptBoneWeightProfile() );
+
+
+};
\ No newline at end of file
