ModelingTools: add option to visualize Tangent seams in Inspector Tool. Add action to AttributeEditorTool to reset/clear Tangents.
GeometryScripting: add new functions GetMeshHasTangents, DiscardTangents, GetMeshPerVertexTangents, SetMeshPerVertexTangents, SetMeshPerVertexNormals, UpdateVertexNormal

diff --git a/Engine/Plugins/Experimental/GeometryScripting/Source/GeometryScriptingCore/Private/MeshNormalsFunctions.cpp b/Engine/Plugins/Experimental/GeometryScripting/Source/GeometryScriptingCore/Private/MeshNormalsFunctions.cpp
--- a/Engine/Plugins/Experimental/GeometryScripting/Source/GeometryScriptingCore/Private/MeshNormalsFunctions.cpp
+++ b/Engine/Plugins/Experimental/GeometryScripting/Source/GeometryScriptingCore/Private/MeshNormalsFunctions.cpp
@@ -209,6 +209,48 @@ UDynamicMesh* UGeometryScriptLibrary_MeshNormalsFunctions::ComputeSplitNormals(
 
 
 
+UDynamicMesh* UGeometryScriptLibrary_MeshNormalsFunctions::GetMeshHasTangents(
+	UDynamicMesh* TargetMesh,
+	bool& bHasTangents,
+	UGeometryScriptDebug* Debug)
+{
+	if (TargetMesh == nullptr)
+	{
+		UE::Geometry::AppendError(Debug, EGeometryScriptErrorType::InvalidInputs, LOCTEXT("GetMeshHasTangents_InvalidInput", "GetMeshHasTangents: TargetMesh is Null"));
+		return TargetMesh;
+	}
+	bHasTangents = false;
+	TargetMesh->ProcessMesh([&](const FDynamicMesh3& ReadMesh)
+	{
+		if (ReadMesh.HasAttributes() && ReadMesh.Attributes()->HasTangentSpace())
+		{
+			bHasTangents = true;
+		}
+	});
+	return TargetMesh;
+}
+
+
+UDynamicMesh* UGeometryScriptLibrary_MeshNormalsFunctions::DiscardTangents(
+	UDynamicMesh* TargetMesh,
+	UGeometryScriptDebug* Debug)
+{
+	if (TargetMesh == nullptr)
+	{
+		UE::Geometry::AppendError(Debug, EGeometryScriptErrorType::InvalidInputs, LOCTEXT("DiscardTangents_InvalidInput", "DiscardTangents: TargetMesh is Null"));
+		return TargetMesh;
+	}
+	TargetMesh->EditMesh([&](FDynamicMesh3& EditMesh)
+	{
+		if (EditMesh.HasAttributes() && EditMesh.Attributes()->HasTangentSpace())
+		{
+			EditMesh.Attributes()->DisableTangents();
+		}
+	});
+	return TargetMesh;
+}
+
+
 UDynamicMesh* UGeometryScriptLibrary_MeshNormalsFunctions::ComputeTangents( 
 	UDynamicMesh* TargetMesh, 
 	FGeometryScriptTangentsOptions Options,
@@ -282,6 +324,57 @@ UDynamicMesh* UGeometryScriptLibrary_MeshNormalsFunctions::SetMeshTriangleNormal
 
 
 
+UDynamicMesh* UGeometryScriptLibrary_MeshNormalsFunctions::SetMeshPerVertexNormals(
+	UDynamicMesh* TargetMesh,
+	FGeometryScriptVectorList VertexNormalList,
+	UGeometryScriptDebug* Debug)
+{
+	if (TargetMesh == nullptr)
+	{
+		UE::Geometry::AppendError(Debug, EGeometryScriptErrorType::InvalidInputs, LOCTEXT("SetMeshPerVertexNormals_InvalidMesh", "SetMeshPerVertexNormals: TargetMesh is Null"));
+		return TargetMesh;
+	}
+	if (VertexNormalList.List.IsValid() == false || VertexNormalList.List->Num() == 0)
+	{
+		UE::Geometry::AppendError(Debug, EGeometryScriptErrorType::InvalidInputs, LOCTEXT("SetMeshPerVertexNormals_InvalidList", "SetMeshPerVertexNormals: List is empty"));
+		return TargetMesh;
+	}
+
+	TargetMesh->EditMesh([&](FDynamicMesh3& EditMesh) 
+	{
+		const TArray<FVector>& VertexNormals = *VertexNormalList.List;
+		if (VertexNormals.Num() < EditMesh.MaxVertexID())
+		{
+			UE::Geometry::AppendError(Debug, EGeometryScriptErrorType::InvalidInputs, LOCTEXT("SetMeshPerVertexNormals_IncorrectCount", "SetMeshPerVertexNormals: size of provided VertexNormalList is smaller than MaxVertexID of Mesh"));
+		}
+		else
+		{
+			if (EditMesh.HasAttributes() == false)
+			{
+				EditMesh.EnableAttributes();
+			}
+			FDynamicMeshNormalOverlay* Normals = EditMesh.Attributes()->PrimaryNormals();
+			Normals->ClearElements();
+			TArray<int32> ElemIDs;
+			ElemIDs.SetNum(EditMesh.MaxVertexID());
+			for (int32 VertexID : EditMesh.VertexIndicesItr())
+			{
+				const FVector& Normal = VertexNormals[VertexID];
+				ElemIDs[VertexID] = Normals->AppendElement((FVector3f)Normal);
+			}
+			for (int32 TriangleID : EditMesh.TriangleIndicesItr())
+			{
+				FIndex3i Triangle = EditMesh.GetTriangle(TriangleID);
+				Normals->SetTriangle(TriangleID, FIndex3i(ElemIDs[Triangle.A], ElemIDs[Triangle.B], ElemIDs[Triangle.C]) );
+			}
+		}
+
+	}, EDynamicMeshChangeType::GeneralEdit, EDynamicMeshAttributeChangeFlags::Unknown, false);
+
+	return TargetMesh;
+}
+
+
 
 
 UDynamicMesh* UGeometryScriptLibrary_MeshNormalsFunctions::GetMeshPerVertexNormals(
@@ -291,71 +384,357 @@ UDynamicMesh* UGeometryScriptLibrary_MeshNormalsFunctions::GetMeshPerVertexNorma
 	bool& bHasVertexIDGaps,
 	bool bAverageSplitVertexValues)
 {
+	if (TargetMesh == nullptr)
+	{
+		UE_LOG(LogGeometry, Warning, TEXT("GetMeshPerVertexNormals: TargetMesh is Null"));
+		return TargetMesh;
+	}
+
 	NormalList.Reset();
 	TArray<FVector>& Normals = *NormalList.List;
 	bHasVertexIDGaps = false;
 	bIsValidNormalSet = false;
-	if (TargetMesh)
+	TargetMesh->ProcessMesh([&](const FDynamicMesh3& ReadMesh)
 	{
-		TargetMesh->ProcessMesh([&](const FDynamicMesh3& ReadMesh)
+		bHasVertexIDGaps = ! ReadMesh.IsCompactV();
+
+		if (ReadMesh.HasAttributes() && ReadMesh.Attributes()->NumNormalLayers() > 0 )
 		{
-			bHasVertexIDGaps = ! ReadMesh.IsCompactV();
+			const FDynamicMeshNormalOverlay* NormalOverlay = ReadMesh.Attributes()->PrimaryNormals();
 
-			if (ReadMesh.HasAttributes() && ReadMesh.Attributes()->NumNormalLayers() > 0 )
+			if (bAverageSplitVertexValues)
 			{
-				const FDynamicMeshNormalOverlay* NormalOverlay = ReadMesh.Attributes()->PrimaryNormals();
-
-				if (bAverageSplitVertexValues)
+				Normals.Init(FVector::Zero(), ReadMesh.MaxVertexID());
+				for (int32 tid : ReadMesh.TriangleIndicesItr())
 				{
-					Normals.Init(FVector::Zero(), ReadMesh.MaxVertexID());
-					for (int32 tid : ReadMesh.TriangleIndicesItr())
+					if (NormalOverlay->IsSetTriangle(tid))
 					{
-						if (NormalOverlay->IsSetTriangle(tid))
-						{
-							FIndex3i TriV = ReadMesh.GetTriangle(tid);
-							FVector3f A, B, C;
-							NormalOverlay->GetTriElements(tid, A, B, C);
-							Normals[TriV.A] += (FVector)A;
-							Normals[TriV.B] += (FVector)B;
-							Normals[TriV.C] += (FVector)C;
-						}
+						FIndex3i TriV = ReadMesh.GetTriangle(tid);
+						FVector3f A, B, C;
+						NormalOverlay->GetTriElements(tid, A, B, C);
+						Normals[TriV.A] += (FVector)A;
+						Normals[TriV.B] += (FVector)B;
+						Normals[TriV.C] += (FVector)C;
 					}
+				}
 
-					for (int32 k = 0; k < Normals.Num(); ++k)
+				for (int32 k = 0; k < Normals.Num(); ++k)
+				{
+					if (Normals[k].SquaredLength() > 0)
 					{
-						if (Normals[k].SquaredLength() > 0)
-						{
-							Normalize(Normals[k]);
-						}
+						Normalize(Normals[k]);
 					}
 				}
-				else
+			}
+			else
+			{
+				Normals.Init(FVector::UnitZ(), ReadMesh.MaxVertexID());
+				for (int32 tid : ReadMesh.TriangleIndicesItr())
 				{
-					Normals.Init(FVector::UnitZ(), ReadMesh.MaxVertexID());
-					for (int32 tid : ReadMesh.TriangleIndicesItr())
+					if (NormalOverlay->IsSetTriangle(tid))
 					{
-						if (NormalOverlay->IsSetTriangle(tid))
-						{
-							FIndex3i TriV = ReadMesh.GetTriangle(tid);
-							FVector3f A, B, C;
-							NormalOverlay->GetTriElements(tid, A, B, C);
-							Normals[TriV.A] = (FVector)A;
-							Normals[TriV.B] = (FVector)B;
-							Normals[TriV.C] = (FVector)C;
-						}
+						FIndex3i TriV = ReadMesh.GetTriangle(tid);
+						FVector3f A, B, C;
+						NormalOverlay->GetTriElements(tid, A, B, C);
+						Normals[TriV.A] = (FVector)A;
+						Normals[TriV.B] = (FVector)B;
+						Normals[TriV.C] = (FVector)C;
 					}
 				}
+			}
+
+			bIsValidNormalSet = true;
+		}
+	});
+
+	return TargetMesh;
+}
+
+
 
-				bIsValidNormalSet = true;
+
+
+
+UDynamicMesh* UGeometryScriptLibrary_MeshNormalsFunctions::SetMeshPerVertexTangents(
+	UDynamicMesh* TargetMesh,
+	FGeometryScriptVectorList TangentXList,
+	FGeometryScriptVectorList TangentYList,
+	UGeometryScriptDebug* Debug)
+{
+	if (TargetMesh == nullptr)
+	{
+		UE::Geometry::AppendError(Debug, EGeometryScriptErrorType::InvalidInputs, LOCTEXT("SetMeshPerVertexTangents_InvalidMesh", "SetMeshPerVertexTangents: TargetMesh is Null"));
+		return TargetMesh;
+	}
+	if (TangentXList.List.IsValid() == false || TangentXList.List->Num() == 0)
+	{
+		UE::Geometry::AppendError(Debug, EGeometryScriptErrorType::InvalidInputs, LOCTEXT("SetMeshPerVertexTangents_InvalidList", "SetMeshPerVertexTangents: TangentXList is empty"));
+		return TargetMesh;
+	}
+	if (TangentYList.List.IsValid() == false || TangentYList.List->Num() != TangentXList.List->Num())
+	{
+		UE::Geometry::AppendError(Debug, EGeometryScriptErrorType::InvalidInputs, LOCTEXT("SetMeshPerVertexTangents_InvalidListY", "SetMeshPerVertexTangents: TangentYList must be the same size as TangentXList"));
+		return TargetMesh;
+	}
+
+	TargetMesh->EditMesh([&](FDynamicMesh3& EditMesh) 
+	{
+		const TArray<FVector>& TangentsX = *TangentXList.List;
+		const TArray<FVector>& TangentsY = *TangentYList.List;
+		if (TangentsX.Num() < EditMesh.MaxVertexID() )
+		{
+			UE::Geometry::AppendError(Debug, EGeometryScriptErrorType::InvalidInputs, LOCTEXT("SetMeshPerVertexTangents_IncorrectCount", "SetMeshPerVertexTangents: size of provided TangentXList is smaller than MaxVertexID of Mesh"));
+		}
+		else
+		{
+			if (EditMesh.HasAttributes() == false)
+			{
+				EditMesh.EnableAttributes();
 			}
-		});
+			if (EditMesh.Attributes()->HasTangentSpace() == false)
+			{
+				EditMesh.Attributes()->EnableTangents();
+			}
+			FDynamicMeshNormalOverlay* TangentsOverlayX = EditMesh.Attributes()->PrimaryTangents();
+			FDynamicMeshNormalOverlay* TangentsOverlayY = EditMesh.Attributes()->PrimaryBiTangents();
+			TangentsOverlayX->ClearElements();
+			TangentsOverlayY->ClearElements();
+			TArray<int32> ElemIDs;
+			ElemIDs.SetNum(EditMesh.MaxVertexID());
+			for (int32 k = 0; k < 2; ++k)
+			{
+				const TArray<FVector>& UseList = (k == 0) ? TangentsX : TangentsY;
+				FDynamicMeshNormalOverlay* UseOverlay = (k == 0) ? TangentsOverlayX : TangentsOverlayY;
+				for (int32 VertexID : EditMesh.VertexIndicesItr())
+				{
+					const FVector& Tangent = UseList[VertexID];
+					ElemIDs[VertexID] = UseOverlay->AppendElement((FVector3f)Tangent);
+				}
+				for (int32 TriangleID : EditMesh.TriangleIndicesItr())
+				{
+					FIndex3i Triangle = EditMesh.GetTriangle(TriangleID);
+					UseOverlay->SetTriangle(TriangleID, FIndex3i(ElemIDs[Triangle.A], ElemIDs[Triangle.B], ElemIDs[Triangle.C]));
+				}
+			}
+		}
+
+	}, EDynamicMeshChangeType::GeneralEdit, EDynamicMeshAttributeChangeFlags::Unknown, false);
+
+	return TargetMesh;
+}
+
+
+
+
+
+UDynamicMesh* UGeometryScriptLibrary_MeshNormalsFunctions::GetMeshPerVertexTangents(
+	UDynamicMesh* TargetMesh, 
+	FGeometryScriptVectorList& TangentXList,
+	FGeometryScriptVectorList& TangentYList,
+	bool& bIsValidTangentSet,
+	bool& bHasVertexIDGaps,
+	bool bAverageSplitVertexValues)
+{
+	if (TargetMesh == nullptr)
+	{
+		UE_LOG(LogGeometry, Warning, TEXT("GetMeshPerVertexTangents: TargetMesh is Null"));
+		return TargetMesh;
 	}
 
+	TangentXList.Reset();
+	TangentYList.Reset();
+	TArray<FVector>& TangentsX = *TangentXList.List;
+	TArray<FVector>& TangentsY = *TangentYList.List;
+	bHasVertexIDGaps = false;
+	bIsValidTangentSet = false;
+	TargetMesh->ProcessMesh([&](const FDynamicMesh3& ReadMesh)
+	{
+		bHasVertexIDGaps = ! ReadMesh.IsCompactV();
+
+		if (ReadMesh.HasAttributes() && ReadMesh.Attributes()->HasTangentSpace() )
+		{
+			const FDynamicMeshNormalOverlay* TangentXOverlay = ReadMesh.Attributes()->PrimaryTangents();
+			const FDynamicMeshNormalOverlay* TangentYOverlay = ReadMesh.Attributes()->PrimaryBiTangents();
+
+			if (bAverageSplitVertexValues)
+			{
+				TangentsX.Init(FVector::Zero(), ReadMesh.MaxVertexID());
+				TangentsY.Init(FVector::Zero(), ReadMesh.MaxVertexID());
+				for (int32 tid : ReadMesh.TriangleIndicesItr())
+				{
+					if (TangentXOverlay->IsSetTriangle(tid) && TangentYOverlay->IsSetTriangle(tid))
+					{
+						FIndex3i TriV = ReadMesh.GetTriangle(tid);
+						FVector3f A, B, C;
+						TangentXOverlay->GetTriElements(tid, A, B, C);
+						TangentsX[TriV.A] += (FVector)A;
+						TangentsX[TriV.B] += (FVector)B;
+						TangentsX[TriV.C] += (FVector)C;
+						TangentYOverlay->GetTriElements(tid, A, B, C);
+						TangentsY[TriV.A] += (FVector)A;
+						TangentsY[TriV.B] += (FVector)B;
+						TangentsY[TriV.C] += (FVector)C;
+					}
+				}
+
+				for (int32 k = 0; k < TangentsX.Num(); ++k)
+				{
+					if (TangentsX[k].SquaredLength() > 0)
+					{
+						Normalize(TangentsX[k]);
+					}
+					if (TangentsY[k].SquaredLength() > 0)
+					{
+						Normalize(TangentsY[k]);
+					}
+				}
+			}
+			else
+			{
+				TangentsX.Init(FVector::UnitZ(), ReadMesh.MaxVertexID());
+				TangentsY.Init(FVector::UnitZ(), ReadMesh.MaxVertexID());
+				for (int32 tid : ReadMesh.TriangleIndicesItr())
+				{
+					if (TangentXOverlay->IsSetTriangle(tid) && TangentYOverlay->IsSetTriangle(tid))
+					{
+						FIndex3i TriV = ReadMesh.GetTriangle(tid);
+						FVector3f A, B, C;
+						TangentXOverlay->GetTriElements(tid, A, B, C);
+						TangentsX[TriV.A] = (FVector)A;
+						TangentsX[TriV.B] = (FVector)B;
+						TangentsX[TriV.C] = (FVector)C;
+						TangentYOverlay->GetTriElements(tid, A, B, C);
+						TangentsY[TriV.A] = (FVector)A;
+						TangentsY[TriV.B] = (FVector)B;
+						TangentsY[TriV.C] = (FVector)C;
+					}
+				}
+			}
+
+			bIsValidTangentSet = true;
+		}
+	});
+
 	return TargetMesh;
 }
 
 
 
 
+UDynamicMesh* UGeometryScriptLibrary_MeshNormalsFunctions::UpdateVertexNormal(
+	UDynamicMesh* TargetMesh,
+	int VertexID,
+	bool bUpdateNormal,
+	FVector NewNormal,
+	bool bUpdateTangents,
+	FVector NewTangentX,
+	FVector NewTangentY,
+	bool& bIsValidVertex,
+	bool bMergeSplitNormals,
+	bool bDeferChangeNotifications)
+{
+	if (TargetMesh == nullptr)
+	{
+		UE_LOG(LogGeometry, Warning, TEXT("UpdateVertexNormal: TargetMesh is Null"));
+		return TargetMesh;
+	}
+
+	TArray<FIndex3i> TriVtxElements;
+	TArray<int32> UniqueElementIDs;
+	TriVtxElements.Reserve(16);
+	UniqueElementIDs.Reserve(16);
+
+	// Updates VertexID's associated elements in Overlay with NewValue. 
+	// Returns false if there were no elements to update (ie all "unset" triangles)
+	auto UpdateOverlay = [VertexID, &TriVtxElements, &UniqueElementIDs, bMergeSplitNormals](FDynamicMesh3& EditMesh, FDynamicMeshNormalOverlay* Overlay, FVector NewValue)
+	{
+		TriVtxElements.Reset(); UniqueElementIDs.Reset();
+		EditMesh.EnumerateVertexTriangles(VertexID, [&](int32 TriangleID)
+		{
+			if (Overlay->IsSetTriangle(TriangleID))
+			{
+				FIndex3i Tri = EditMesh.GetTriangle(TriangleID);
+				int32 Index = Tri.IndexOf(VertexID);
+				FIndex3i OverlayTri = Overlay->GetTriangle(TriangleID);
+				TriVtxElements.Add(FIndex3i(TriangleID, Index, OverlayTri[Index]));
+				UniqueElementIDs.AddUnique(OverlayTri[Index]);
+			}
+		});
+		if (TriVtxElements.IsEmpty())
+		{
+			return false;
+		}
+		if (UniqueElementIDs.Num() == 1 || bMergeSplitNormals == false)
+		{
+			// just update existing elements to new normal
+			for (int32 ElementID : UniqueElementIDs)
+			{
+				Overlay->SetElement(ElementID, (FVector3f)NewValue);
+			}
+		}
+		else   // bMergeSplitNormals == true && UniqueElementIDs.Num() > 1
+		{
+			Overlay->SetElement(UniqueElementIDs[0], (FVector3f)NewValue);
+			for (FIndex3i TriInfo : TriVtxElements)
+			{
+				FIndex3i OverlayTri = Overlay->GetTriangle(TriInfo.A);
+				OverlayTri[TriInfo.B] = UniqueElementIDs[0];
+				Overlay->SetTriangle(TriInfo.A, OverlayTri);
+			}
+		}
+		return true;
+	};
+
+	bIsValidVertex = false;
+	TargetMesh->EditMesh([&](FDynamicMesh3& EditMesh) 
+	{
+		if (EditMesh.HasAttributes() == false 
+			|| (EditMesh.HasAttributes() && EditMesh.Attributes()->PrimaryNormals() == nullptr) 
+			|| (EditMesh.HasAttributes() && EditMesh.Attributes()->PrimaryNormals()->ElementCount() == 0) )
+		{
+			UE_LOG(LogGeometry, Warning, TEXT("UpdateVertexNormal: TargetMesh does not have valid Normals attributes. Try computing Vertex Normals before using UpdateVertexNormal."));
+			return;
+		}
+		if (bUpdateTangents && EditMesh.Attributes()->HasTangentSpace() == false)
+		{
+			UE_LOG(LogGeometry, Warning, TEXT("UpdateVertexNormal: TargetMesh does not have valid Tangents attributes. Try computing Tangents before using UpdateVertexNormal."));
+			return;
+		}
+		if (EditMesh.IsVertex(VertexID) == false)
+		{
+			UE_LOG(LogGeometry, Warning, TEXT("UpdateVertexNormal: VertexID %d is not a valid vertex in TargetMesh"));
+			return;
+		}
+		bIsValidVertex = true;
+
+		if (bUpdateNormal)
+		{
+			FDynamicMeshNormalOverlay* Normals = EditMesh.Attributes()->PrimaryNormals();
+			if (UpdateOverlay(EditMesh, Normals, NewNormal) == false)
+			{
+				UE_LOG(LogGeometry, Warning, TEXT("UpdateVertexNormal: VertexID %d has no existing normals in Normal Overlay. Try computing Vertex Normals before using UpdateVertexNormal."));
+				bIsValidVertex = false;
+			}
+		}
+		if (bUpdateTangents)
+		{
+			FDynamicMeshNormalOverlay* TangentX = EditMesh.Attributes()->PrimaryTangents();
+			FDynamicMeshNormalOverlay* TangentY = EditMesh.Attributes()->PrimaryBiTangents();
+			bool bTangentXOK = UpdateOverlay(EditMesh, TangentX, NewTangentX);
+			bool bTangentYOK = UpdateOverlay(EditMesh, TangentY, NewTangentY);
+			if (bTangentXOK == false || bTangentYOK == false)
+			{
+				UE_LOG(LogGeometry, Warning, TEXT("UpdateVertexNormal: VertexID %d has no existing tangents in Tangent Overlay. Try computing Tangents before using UpdateVertexNormal."));
+				bIsValidVertex = false;
+			}
+		}
+
+	}, EDynamicMeshChangeType::GeneralEdit, EDynamicMeshAttributeChangeFlags::Unknown, bDeferChangeNotifications);
+
+	return TargetMesh;
+}
+
+
+
 #undef LOCTEXT_NAMESPACE
 
diff --git a/Engine/Plugins/Experimental/GeometryScripting/Source/GeometryScriptingCore/Public/GeometryScript/MeshNormalsFunctions.h b/Engine/Plugins/Experimental/GeometryScripting/Source/GeometryScriptingCore/Public/GeometryScript/MeshNormalsFunctions.h
--- a/Engine/Plugins/Experimental/GeometryScripting/Source/GeometryScriptingCore/Public/GeometryScript/MeshNormalsFunctions.h
+++ b/Engine/Plugins/Experimental/GeometryScripting/Source/GeometryScriptingCore/Public/GeometryScript/MeshNormalsFunctions.h
@@ -70,30 +70,53 @@ class GEOMETRYSCRIPTINGCORE_API UGeometryScriptLibrary_MeshNormalsFunctions : pu
 	GENERATED_BODY()
 public:
 
+	/**
+	 * Flip/Invert the normal vectors of TargetMesh by multiplying them by -1, as well as reversing the mesh triangle orientations, ie triangle (a,b,c) becomes (b,a,c)
+	 */
 	UFUNCTION(BlueprintCallable, Category = "GeometryScript|Normals", meta=(ScriptMethod))
 	static UPARAM(DisplayName = "Target Mesh") UDynamicMesh* 
 	FlipNormals( 
 		UDynamicMesh* TargetMesh, 
 		UGeometryScriptDebug* Debug = nullptr);
 
+	/**
+	 * Attempt to repair inconsistent normals in TargetMesh. Currently this is done in two passes. In the first pass, triangles with
+	 * reversed orientation from their neighours are incrementally flipped until each connected component has a consistent orientation,
+	 * if this is possible (note that this is not always globally possible, eg for a mobius-strip topology there is no consistent orientation).
+	 * In the second pass, the "global" orientation is detected by casting rays from outside the mesh. This may produce incorrect results for
+	 * meshes that are not closed. 
+	 */
 	UFUNCTION(BlueprintCallable, Category = "GeometryScript|Normals", meta=(ScriptMethod))
 	static UPARAM(DisplayName = "Target Mesh") UDynamicMesh* 
 	AutoRepairNormals( 
 		UDynamicMesh* TargetMesh, 
 		UGeometryScriptDebug* Debug = nullptr);
 
-	UFUNCTION(BlueprintCallable, Category = "GeometryScript|Normals", meta=(ScriptMethod))
+	/**
+	 * Recompute the normals of TargetMesh by averaging the triangle/face normals around each vertex, using combined area and angle weighting.
+	 * Each vertex will have a single normal, ie there will be no hard edges.
+	 */
+	UFUNCTION(BlueprintCallable, Category = "GeometryScript|Normals", meta=(ScriptMethod, DisplayName="Set Mesh To Per Vertex Normals"))
 	static UPARAM(DisplayName = "Target Mesh") UDynamicMesh* 
 	SetPerVertexNormals( 
 		UDynamicMesh* TargetMesh, 
 		UGeometryScriptDebug* Debug = nullptr);
 
-	UFUNCTION(BlueprintCallable, Category = "GeometryScript|Normals", meta=(ScriptMethod))
+	/**
+	 * Recompute the normals of TargetMesh by setting the normals of each triangle vertex to the triangle/face normal.
+	 * Each vertex will have a unique normal in each triangle, ie there will be hard edges / split normals at every mesh edge
+	 */
+	UFUNCTION(BlueprintCallable, Category = "GeometryScript|Normals", meta=(ScriptMethod, DisplayName = "Set Mesh To Facet Normals"))
 	static UPARAM(DisplayName = "Target Mesh") UDynamicMesh* 
 	SetPerFaceNormals( 
 		UDynamicMesh* TargetMesh, 
 		UGeometryScriptDebug* Debug = nullptr);
 
+	/**
+	 * Recompute the normals of TargetMesh using the given CalculateOptions. This method will preserve any existing hard
+	 * edges, ie each shared triangle-vertex normal is recomputed by averaging the face normals of triangles that reference
+	 * that shared triangle-vertex normal
+	 */
 	UFUNCTION(BlueprintCallable, Category = "GeometryScript|Normals", meta=(ScriptMethod))
 	static UPARAM(DisplayName = "Target Mesh") UDynamicMesh* 
 	RecomputeNormals(  
@@ -101,6 +124,11 @@ public:
 		FGeometryScriptCalculateNormalsOptions CalculateOptions,
 		UGeometryScriptDebug* Debug = nullptr);
 
+	/**
+	 * Recompute hard edges / split-normals for TargetMesh based on the provided SplitOptions, and then 
+	 * recompute the new shared triangle-vertex normals using the given CalculateOptions. 
+	 * The normal recomputation is identical to calling RecomputeNormals.
+	 */
 	UFUNCTION(BlueprintCallable, Category = "GeometryScript|Normals", meta=(ScriptMethod))
 	static UPARAM(DisplayName = "Target Mesh") UDynamicMesh* 
 	ComputeSplitNormals( 
@@ -109,14 +137,13 @@ public:
 		FGeometryScriptCalculateNormalsOptions CalculateOptions,
 		UGeometryScriptDebug* Debug = nullptr);
 
-	UFUNCTION(BlueprintCallable, Category = "GeometryScript|Normals", meta=(ScriptMethod))
-	static UPARAM(DisplayName = "Target Mesh") UDynamicMesh* 
-	ComputeTangents( 
-		UDynamicMesh* TargetMesh, 
-		FGeometryScriptTangentsOptions Options,
-		UGeometryScriptDebug* Debug = nullptr);
-
-
+	/**
+	 * Set the triangle-vertex normals for the given TriangleID on the TargetMesh. This will
+	 * create unique triangle-vertex normals, ie it will create hard edges / split normals in 
+	 * the normal overlay for each edge of the triangle. 
+	 * @param bIsValidTriangle will be returned as false if TriangleID does not refer to a valid triangle
+	 * @param bDeferChangeNotifications if true, no mesh change notification will be sent. Set to true if changing many normals in a loop.
+	 */
 	UFUNCTION(BlueprintCallable, Category = "GeometryScript|Normals", meta=(ScriptMethod))
 	static UPARAM(DisplayName = "Target Mesh") UDynamicMesh* 
 	SetMeshTriangleNormals( 
@@ -126,6 +153,16 @@ public:
 		bool& bIsValidTriangle, 
 		bool bDeferChangeNotifications = false );
 
+	/**
+	 * Set all normals in the TargetMesh Normals Overlay to the specified per-vertex normals
+	 * @param VertexNormalList per-vertex normals. Size must be equal to the MaxVertexID of TargetMesh  (ie non-compact TargetMesh is supported)
+	 */
+	UFUNCTION(BlueprintCallable, Category = "GeometryScript|Normals", meta=(ScriptMethod))
+	static UPARAM(DisplayName = "Target Mesh") UDynamicMesh*
+	SetMeshPerVertexNormals(
+		UDynamicMesh* TargetMesh,
+		FGeometryScriptVectorList VertexNormalList,
+		UGeometryScriptDebug* Debug = nullptr);
 
 
 	/**
@@ -146,4 +183,103 @@ public:
 		bool& bHasVertexIDGaps,
 		bool bAverageSplitVertexValues = true);
 
+
+
+	/**
+	 * Check if the TargetMesh has a Tangents Attribute Layer enabled
+	 */
+	UFUNCTION(BlueprintCallable, Category = "GeometryScript|Normals", meta=(ScriptMethod))
+	static UPARAM(DisplayName = "Target Mesh") UDynamicMesh* 
+	GetMeshHasTangents( 
+		UDynamicMesh* TargetMesh,
+		bool& bHasTangents,
+		UGeometryScriptDebug* Debug = nullptr);
+
+	/**
+	 * Remove any existing Tangents Attribute Layer from the TargetMesh
+	 */
+	UFUNCTION(BlueprintCallable, Category = "GeometryScript|Normals", meta=(ScriptMethod))
+	static UPARAM(DisplayName = "Target Mesh") UDynamicMesh* 
+	DiscardTangents( 
+		UDynamicMesh* TargetMesh,
+		UGeometryScriptDebug* Debug = nullptr);
+
+	/**
+	 * Recompute Tangents for the TargetMesh, using the method and settings specified by FGeometryScriptTangentsOptions
+	 */
+	UFUNCTION(BlueprintCallable, Category = "GeometryScript|Normals", meta=(ScriptMethod))
+	static UPARAM(DisplayName = "Target Mesh") UDynamicMesh* 
+	ComputeTangents( 
+		UDynamicMesh* TargetMesh, 
+		FGeometryScriptTangentsOptions Options,
+		UGeometryScriptDebug* Debug = nullptr);
+
+
+	/**
+	 * Set all tangents in the TargetMesh Tangents Overlays to the specified per-vertex tangents
+	 * @param TangentXList per-vertex tangent vectors. Size must be equal to the MaxVertexID of TargetMesh  (ie non-compact TargetMesh is supported)
+	 * @param TangentYList per-vertex bitangent/binormal vectors. Size must be equal to TangentXList
+	 */
+	UFUNCTION(BlueprintCallable, Category = "GeometryScript|Normals", meta=(ScriptMethod))
+	static UPARAM(DisplayName = "Target Mesh") UDynamicMesh*
+	SetMeshPerVertexTangents(
+		UDynamicMesh* TargetMesh,
+		FGeometryScriptVectorList TangentXList,
+		FGeometryScriptVectorList TangentYList,
+		UGeometryScriptDebug* Debug = nullptr);
+
+	/**
+	 * Get a list of single tangent vectors for each mesh vertex in the TargetMesh, derived from the Tangents Overlays.
+	 * The Tangents Overlay may store multiple tangents for a single vertex (ie split tangents)
+	 * In such cases the tangents can either be averaged, or the last tangent seen will be used, depending on the bAverageSplitVertexValues parameter.
+	 * @param TangentXList output Tangent "X" vectors list will be stored here. Size will be equal to the MaxVertexID of TargetMesh  (not the VertexCount!)
+	 * @param TangentYList output Tangent "Y" vectors (Binormal/Bitangent) list will be stored here. Size will be equal to TangentXList
+	 * @param bIsValidTangentSet will be set to true if the Tangent Overlay was valid
+	 * @param bHasVertexIDGaps will be set to true if some vertex indices in TargetMesh were invalid, ie MaxVertexID > VertexCount 
+	 * @param bAverageSplitVertexValues control how multiple tangents at the same vertex should be interpreted
+	 */
+	UFUNCTION(BlueprintCallable, Category = "GeometryScript|Normals", meta=(ScriptMethod))
+	static UPARAM(DisplayName = "Target Mesh") UDynamicMesh* 
+	GetMeshPerVertexTangents( 
+		UDynamicMesh* TargetMesh, 
+		FGeometryScriptVectorList& TangentXList, 
+		FGeometryScriptVectorList& TangentYList,
+		bool& bIsValidTangentSet,
+		bool& bHasVertexIDGaps,
+		bool bAverageSplitVertexValues = true);
+
+
+	/**
+	 * Update the Normals and/or Tangents at VertexID of TargetMesh. Note that the specified vertex may have "split normals"
+	 * or "split tangents", ie in the case of hard/crease normals, UV seams, and so on. In these situations, by default
+	 * each of the unique normals/tangents at the vertex will be updated, but they will not be "merged", ie they will remain split.
+	 * However if bMergeSplitValues=true, then the vertex will be "un-split", ie after the function call the vertex will have
+	 * a single unique shared normal and/or tangents. 
+	 * 
+	 * Note that this function requires that some normals/tangents already exist on the TargetMesh. If this is not the case, 
+	 * functions like SetPerVertexNormals and ComputeTangents can be used to initialize the normals/tangents first.
+	 * 	
+	 * @param bUpdateNormal if true (default) then the normals overlay is updated
+	 * @param NewNormal the new normal vector. This vector will not be normalized, it must be normalized by the calling code.
+	 * @param bUpdateTangents if true then the tangents overlay will be updated. If the tangents overlay does not exist, this function returns an error.
+	 * @param NewTangentX the new tangent vector. This vector will not be normalized, it must be normalized by the calling code.
+	 * @param NewTangentY the new bitangent/binormal vector. This vector will not be normalized, it must be normalized by the calling code.
+	 * @param bIsValidVertex will be set to true on return if the VertexID was valid, ie had valid normals and tangents
+	 * @param bMergeSplitValues if true, any split normals/tangents at the vertex will be cleared, and a unique normal/tangent element will be set in the connected triangles
+	 * @param bDeferChangeNotifications if true, no mesh change notification will be sent. Set to true if changing many normals in a loop.
+	 */
+	UFUNCTION(BlueprintCallable, Category = "GeometryScript|MeshEdits", meta=(ScriptMethod, bUpdateNormal="true"))
+	static UPARAM(DisplayName = "Target Mesh") UDynamicMesh* 
+	UpdateVertexNormal( 
+		UDynamicMesh* TargetMesh, 
+		int VertexID, 
+		bool bUpdateNormal,
+		FVector NewNormal, 
+		bool bUpdateTangents,
+		FVector NewTangentX,
+		FVector NewTangentY,
+		bool& bIsValidVertex, 
+		bool bMergeSplitValues,
+		bool bDeferChangeNotifications = false );
+
 };
\ No newline at end of file
diff --git a/Engine/Plugins/Experimental/MeshModelingToolsetExp/Source/MeshModelingToolsEditorOnlyExp/Private/AttributeEditorTool.cpp b/Engine/Plugins/Experimental/MeshModelingToolsetExp/Source/MeshModelingToolsEditorOnlyExp/Private/AttributeEditorTool.cpp
--- a/Engine/Plugins/Experimental/MeshModelingToolsetExp/Source/MeshModelingToolsEditorOnlyExp/Private/AttributeEditorTool.cpp
+++ b/Engine/Plugins/Experimental/MeshModelingToolsetExp/Source/MeshModelingToolsEditorOnlyExp/Private/AttributeEditorTool.cpp
@@ -396,6 +396,9 @@ void UAttributeEditorTool::OnTick(float DeltaTime)
 	case EAttributeEditorToolActions::ClearNormals:
 		ClearNormals();
 		break;
+	case EAttributeEditorToolActions::DiscardTangents:
+		DiscardTangents();
+		break;
 	case EAttributeEditorToolActions::ClearAllUVs:
 		ClearUVs();
 		break;
@@ -470,6 +473,38 @@ void UAttributeEditorTool::ClearNormals()
 
 
 
+void UAttributeEditorTool::DiscardTangents()
+{
+	GetToolManager()->BeginUndoTransaction(LOCTEXT("DiscardTangentsTransactionMessage", "Discard Tangents"));
+
+	for (int32 ComponentIdx = 0; ComponentIdx < Targets.Num(); ComponentIdx++)
+	{
+		FMeshDescription EditedMesh = UE::ToolTarget::GetMeshDescriptionCopy(Targets[ComponentIdx]);
+
+		TVertexInstanceAttributesRef<FVector3f> VertexInstanceTangents = EditedMesh.VertexInstanceAttributes().GetAttributesRef<FVector3f>(MeshAttribute::VertexInstance::Tangent);
+		TVertexInstanceAttributesRef<float> VertexInstanceBinormalSigns = EditedMesh.VertexInstanceAttributes().GetAttributesRef<float>(MeshAttribute::VertexInstance::BinormalSign);
+		const FVertexInstanceArray& Instances = EditedMesh.VertexInstances();
+		if (VertexInstanceTangents.IsValid())
+		{
+			for (const FVertexInstanceID InstanceID : Instances.GetElementIDs())
+			{
+				VertexInstanceTangents.Set(InstanceID, FVector3f::ZeroVector);
+			}
+		}
+		if (VertexInstanceBinormalSigns.IsValid())
+		{
+			for (const FVertexInstanceID InstanceID : Instances.GetElementIDs())
+			{
+				VertexInstanceBinormalSigns.Set(InstanceID, 1.0f);
+			}
+		}
+
+		UE::ToolTarget::CommitMeshDescriptionUpdate(Targets[ComponentIdx], MoveTemp(EditedMesh));
+	}
+	GetToolManager()->EndUndoTransaction();
+}
+
+
 
 void UAttributeEditorTool::ClearUVs()
 {
diff --git a/Engine/Plugins/Experimental/MeshModelingToolsetExp/Source/MeshModelingToolsEditorOnlyExp/Public/AttributeEditorTool.h b/Engine/Plugins/Experimental/MeshModelingToolsetExp/Source/MeshModelingToolsEditorOnlyExp/Public/AttributeEditorTool.h
--- a/Engine/Plugins/Experimental/MeshModelingToolsetExp/Source/MeshModelingToolsEditorOnlyExp/Public/AttributeEditorTool.h
+++ b/Engine/Plugins/Experimental/MeshModelingToolsetExp/Source/MeshModelingToolsEditorOnlyExp/Public/AttributeEditorTool.h
@@ -90,6 +90,7 @@ enum class EAttributeEditorToolActions
 	NoAction,
 
 	ClearNormals,
+	DiscardTangents,
 	ClearAllUVs,
 	ClearSelectedUVs,
 	AddUVSet,
@@ -128,12 +129,19 @@ class MESHMODELINGTOOLSEDITORONLYEXP_API UAttributeEditorNormalsActions : public
 {
 	GENERATED_BODY()
 public:
-	/** Remove any hard edges / split normals, setting all normals to a single vertex normal */
+	/** Remove any hard edges / split normals, setting all normals to a single averaged vertex normal */
 	UFUNCTION(CallInEditor, Category = Normals, meta = (DisplayPriority = 1))
 	void ResetHardNormals()
 	{
 		PostAction(EAttributeEditorToolActions::ClearNormals);
 	}
+
+	/** Clear Tangents from the mesh */
+	UFUNCTION(CallInEditor, Category = Normals, meta = (DisplayPriority = 2))
+	void DiscardTangents()
+	{
+		PostAction(EAttributeEditorToolActions::DiscardTangents);
+	}
 };
 
 
@@ -372,6 +380,7 @@ protected:
 
 	EAttributeEditorToolActions PendingAction = EAttributeEditorToolActions::NoAction;
 	void ClearNormals();
+	void DiscardTangents();
 	void ClearUVs();
 	void DeleteSelectedUVSet();
 	void DuplicateSelectedUVSet();
diff --git a/Engine/Plugins/Experimental/MeshModelingToolsetExp/Source/MeshModelingToolsExp/Private/MeshInspectorTool.cpp b/Engine/Plugins/Experimental/MeshModelingToolsetExp/Source/MeshModelingToolsExp/Private/MeshInspectorTool.cpp
--- a/Engine/Plugins/Experimental/MeshModelingToolsetExp/Source/MeshModelingToolsExp/Private/MeshInspectorTool.cpp
+++ b/Engine/Plugins/Experimental/MeshModelingToolsetExp/Source/MeshModelingToolsExp/Private/MeshInspectorTool.cpp
@@ -282,6 +282,7 @@ void UMeshInspectorTool::Precompute()
 	UVSeamEdges.Reset();
 	UVBowties.Reset();
 	NormalSeamEdges.Reset();
+	TangentSeamEdges.Reset();
 	GroupBoundaryEdges.Reset();
 	MissingUVTriangleEdges.Reset();
 
@@ -290,10 +291,15 @@ void UMeshInspectorTool::Precompute()
 		TargetMesh->HasAttributes() ? TargetMesh->Attributes()->PrimaryUV() : nullptr;
 	const FDynamicMeshNormalOverlay* NormalOverlay =
 		TargetMesh->HasAttributes() ? TargetMesh->Attributes()->PrimaryNormals() : nullptr;
+	const FDynamicMeshNormalOverlay* TangentXOverlay =
+		(TargetMesh->HasAttributes() && TargetMesh->Attributes()->HasTangentSpace()) ? TargetMesh->Attributes()->PrimaryTangents() : nullptr;
+	const FDynamicMeshNormalOverlay* TangentYOverlay =
+		(TargetMesh->HasAttributes() && TargetMesh->Attributes()->HasTangentSpace()) ? TargetMesh->Attributes()->PrimaryBiTangents() : nullptr;
 
 	for (int eid : TargetMesh->EdgeIndicesItr())
 	{
-		if (TargetMesh->IsBoundaryEdge(eid))
+		bool bIsBoundaryEdge = TargetMesh->IsBoundaryEdge(eid);
+		if (bIsBoundaryEdge)
 		{
 			BoundaryEdges.Add(eid);
 		}
@@ -301,10 +307,14 @@ void UMeshInspectorTool::Precompute()
 		{
 			UVSeamEdges.Add(eid);
 		}
-		if (NormalOverlay != nullptr && NormalOverlay->IsSeamEdge(eid))
+		if (bIsBoundaryEdge == false && NormalOverlay != nullptr && NormalOverlay->IsSeamEdge(eid))
 		{
 			NormalSeamEdges.Add(eid);
 		}
+		if (bIsBoundaryEdge == false && TangentXOverlay != nullptr && (TangentXOverlay->IsSeamEdge(eid) || TangentYOverlay->IsSeamEdge(eid)) )
+		{
+			TangentSeamEdges.Add(eid);
+		}
 		if (TargetMesh->IsGroupBoundaryEdge(eid))
 		{
 			GroupBoundaryEdges.Add(eid);
@@ -558,6 +568,8 @@ void UMeshInspectorTool::UpdateVisualization()
 	const float MissingUVThickness = LineWidthMultiplier * 2.0;
 	FColor NormalSeamColor(15, 240, 240);
 	float NormalSeamThickness = LineWidthMultiplier * 2.0;
+	FColor TangentSeamColor(64, 240, 240);
+	float TangentSeamThickness = LineWidthMultiplier * 2.0;
 	FColor PolygonBorderColor(240, 15, 240);
 	float PolygonBorderThickness = LineWidthMultiplier * 2.0;
 	FColor NormalColor(15, 15, 240);
@@ -570,6 +582,7 @@ void UMeshInspectorTool::UpdateVisualization()
 	float UVSeamDepthBias = 0.3f;
 	const float MissingUVDepthBias = 0.3f;
 	float NormalSeamDepthBias = 0.3f;
+	float TangentSeamDepthBias = 0.3f;
 	float PolygonBorderDepthBias = 0.2f;
 	float NormalDepthBias = 0.0f;
 	float TangentDepthBias = 0.35f;
@@ -618,6 +631,15 @@ void UMeshInspectorTool::UpdateVisualization()
 		}
 	}
 
+	if (Settings->bTangentSeams)
+	{
+		for (int eid : TangentSeamEdges)
+		{
+			TargetMesh->GetEdgeV(eid, A, B);
+			DrawnLineSet->AddLine((FVector)A, (FVector)B, TangentSeamColor, TangentSeamThickness, TangentSeamDepthBias);
+		}
+	}
+
 	if (Settings->bPolygonBorders)
 	{
 		for (int eid : GroupBoundaryEdges)
diff --git a/Engine/Plugins/Experimental/MeshModelingToolsetExp/Source/MeshModelingToolsExp/Public/MeshInspectorTool.h b/Engine/Plugins/Experimental/MeshModelingToolsetExp/Source/MeshModelingToolsExp/Public/MeshInspectorTool.h
--- a/Engine/Plugins/Experimental/MeshModelingToolsetExp/Source/MeshModelingToolsExp/Public/MeshInspectorTool.h
+++ b/Engine/Plugins/Experimental/MeshModelingToolsetExp/Source/MeshModelingToolsExp/Public/MeshInspectorTool.h
@@ -84,6 +84,10 @@ public:
 	UPROPERTY(EditAnywhere, Category = Options)
 	bool bNormalSeams = false;
 
+	/** Toggle visibility of Tangent seam edges */
+	UPROPERTY(EditAnywhere, Category = Options)
+	bool bTangentSeams = false;
+
 	/** Toggle visibility of normal vectors */
 	UPROPERTY(EditAnywhere, Category = Options)
 	bool bNormalVectors = false;
@@ -270,6 +274,7 @@ protected:
 	TArray<int> UVSeamEdges;
 	TArray<int> UVBowties;
 	TArray<int> NormalSeamEdges;
+	TArray<int> TangentSeamEdges;
 	TArray<int> GroupBoundaryEdges;
 	TArray<int> MissingUVTriangleEdges;
 
