MeshModelingToolset: Fix spelling for PolyGroup in Shapes tools

diff --git a/Engine/Plugins/Experimental/GeometryScripting/Source/GeometryScriptingCore/Public/GeometryScript/MeshPrimitiveFunctions.h b/Engine/Plugins/Experimental/GeometryScripting/Source/GeometryScriptingCore/Public/GeometryScript/MeshPrimitiveFunctions.h
--- a/Engine/Plugins/Experimental/GeometryScripting/Source/GeometryScriptingCore/Public/GeometryScript/MeshPrimitiveFunctions.h
+++ b/Engine/Plugins/Experimental/GeometryScripting/Source/GeometryScriptingCore/Public/GeometryScript/MeshPrimitiveFunctions.h
@@ -36,7 +36,7 @@ struct GEOMETRYSCRIPTINGCORE_API FGeometryScriptPrimitiveOptions
 {
 	GENERATED_BODY()
 public:
-	UPROPERTY(EditAnywhere, BlueprintReadWrite, Category = Options)
+	UPROPERTY(EditAnywhere, BlueprintReadWrite, Category = Options, meta = (DisplayName = "PolyGroup Mode"))
 	EGeometryScriptPrimitivePolygroupMode PolygroupMode = EGeometryScriptPrimitivePolygroupMode::PerFace;
 
 	UPROPERTY(EditAnywhere, BlueprintReadWrite, Category = Options)
diff --git a/Engine/Plugins/Runtime/MeshModelingToolset/Source/MeshModelingTools/Public/AddPrimitiveTool.h b/Engine/Plugins/Runtime/MeshModelingToolset/Source/MeshModelingTools/Public/AddPrimitiveTool.h
--- a/Engine/Plugins/Runtime/MeshModelingToolset/Source/MeshModelingTools/Public/AddPrimitiveTool.h
+++ b/Engine/Plugins/Runtime/MeshModelingToolset/Source/MeshModelingTools/Public/AddPrimitiveTool.h
@@ -81,8 +81,8 @@ class MESHMODELINGTOOLS_API UProceduralShapeToolProperties : public UInteractive
 
 public:
 
-	/** How Polygroups are assigned to shape primitives. */
-	UPROPERTY(EditAnywhere, Category = Shape, meta = (ProceduralShapeSetting))
+	/** How PolyGroups are assigned to shape primitives. */
+	UPROPERTY(EditAnywhere, Category = Shape, meta = (ProceduralShapeSetting, DisplayName = "PolyGroup Mode"))
 	EMakeMeshPolygroupMode PolygroupMode = EMakeMeshPolygroupMode::PerFace;
 
 	/** How the shape is placed in the scene. */
