GeometryScripting: Fix some crashes

diff --git a/Engine/Plugins/Experimental/GeometryScripting/Source/GeometryScriptingCore/Private/ListUtilityFunctions.cpp b/Engine/Plugins/Experimental/GeometryScripting/Source/GeometryScriptingCore/Private/ListUtilityFunctions.cpp
--- a/Engine/Plugins/Experimental/GeometryScripting/Source/GeometryScriptingCore/Private/ListUtilityFunctions.cpp
+++ b/Engine/Plugins/Experimental/GeometryScripting/Source/GeometryScriptingCore/Private/ListUtilityFunctions.cpp
@@ -17,7 +17,7 @@ int UGeometryScriptLibrary_ListUtilityFunctions::GetIndexListLength(FGeometryScr
 
 int UGeometryScriptLibrary_ListUtilityFunctions::GetIndexListLastIndex(FGeometryScriptIndexList IndexList)
 {
-	return (IndexList.List.IsValid()) ? FMath::Max(IndexList.List->Num()-1,0) : 0;
+	return IndexList.List.IsValid() ? IndexList.List->Num()-1 : -1;
 }
 
 int UGeometryScriptLibrary_ListUtilityFunctions::GetIndexListItem(FGeometryScriptIndexList IndexList, int Index, bool& bIsValidIndex)
@@ -65,7 +65,7 @@ void UGeometryScriptLibrary_ListUtilityFunctions::DuplicateIndexList(FGeometrySc
 
 void UGeometryScriptLibrary_ListUtilityFunctions::ClearIndexList(FGeometryScriptIndexList& IndexList, int ClearValue)
 {
-	int32 Num = IndexList.List->Num();
+	int Num = GetIndexListLength(IndexList);
 	IndexList.Reset(IndexList.IndexType);
 	IndexList.List->Init(ClearValue, Num);
 }
@@ -120,7 +120,7 @@ int UGeometryScriptLibrary_ListUtilityFunctions::GetScalarListLength(FGeometrySc
 
 int UGeometryScriptLibrary_ListUtilityFunctions::GetScalarListLastIndex(FGeometryScriptScalarList ScalarList)
 {
-	return (ScalarList.List.IsValid()) ? FMath::Max(ScalarList.List->Num()-1,0) : 0;
+	return ScalarList.List.IsValid() ? ScalarList.List->Num()-1 : -1;
 }
 
 double UGeometryScriptLibrary_ListUtilityFunctions::GetScalarListItem(FGeometryScriptScalarList ScalarList, int Index, bool& bIsValidIndex)
@@ -167,7 +167,7 @@ void UGeometryScriptLibrary_ListUtilityFunctions::DuplicateScalarList(FGeometryS
 
 void UGeometryScriptLibrary_ListUtilityFunctions::ClearScalarList(FGeometryScriptScalarList& ScalarList, double ClearValue)
 {
-	int32 Num = ScalarList.List->Num();
+	int Num = GetScalarListLength(ScalarList);
 	ScalarList.Reset();
 	ScalarList.List->Init(ClearValue, Num);
 }
@@ -183,7 +183,7 @@ int UGeometryScriptLibrary_ListUtilityFunctions::GetVectorListLength(FGeometrySc
 
 int UGeometryScriptLibrary_ListUtilityFunctions::GetVectorListLastIndex(FGeometryScriptVectorList VectorList)
 {
-	return (VectorList.List.IsValid()) ? FMath::Max(VectorList.List->Num()-1,0) : 0;
+	return VectorList.List.IsValid() ? VectorList.List->Num()-1 : -1;
 }
 
 
@@ -231,7 +231,7 @@ void UGeometryScriptLibrary_ListUtilityFunctions::DuplicateVectorList(FGeometryS
 
 void UGeometryScriptLibrary_ListUtilityFunctions::ClearVectorList(FGeometryScriptVectorList& VectorList, FVector ClearValue)
 {
-	int32 Num = VectorList.List->Num();
+	int Num = GetVectorListLength(VectorList);
 	VectorList.Reset();
 	VectorList.List->Init(ClearValue, Num);
 }
@@ -245,7 +245,7 @@ int UGeometryScriptLibrary_ListUtilityFunctions::GetUVListLength(FGeometryScript
 
 int UGeometryScriptLibrary_ListUtilityFunctions::GetUVListLastIndex(FGeometryScriptUVList UVList)
 {
-	return (UVList.List.IsValid()) ? FMath::Max(UVList.List->Num()-1,0) : 0;
+	return UVList.List.IsValid() ? UVList.List->Num()-1 : -1;
 }
 
 FVector2D UGeometryScriptLibrary_ListUtilityFunctions::GetUVListItem(FGeometryScriptUVList UVList, int Index, bool& bIsValidIndex)
@@ -292,7 +292,7 @@ void UGeometryScriptLibrary_ListUtilityFunctions::DuplicateUVList(FGeometryScrip
 
 void UGeometryScriptLibrary_ListUtilityFunctions::ClearUVList(FGeometryScriptUVList& UVList, FVector2D ClearUV)
 {
-	int32 Num = UVList.List->Num();
+	int Num = GetUVListLength(UVList);
 	UVList.Reset();
 	UVList.List->Init(ClearUV, Num);
 }
@@ -307,7 +307,7 @@ int UGeometryScriptLibrary_ListUtilityFunctions::GetColorListLength(FGeometryScr
 
 int UGeometryScriptLibrary_ListUtilityFunctions::GetColorListLastIndex(FGeometryScriptColorList ColorList)
 {
-	return (ColorList.List.IsValid()) ? FMath::Max(ColorList.List->Num()-1,0) : 0;
+	return ColorList.List.IsValid() ? ColorList.List->Num()-1 : -1;
 }
 
 FLinearColor UGeometryScriptLibrary_ListUtilityFunctions::GetColorListItem(FGeometryScriptColorList ColorList, int Index, bool& bIsValidIndex)
@@ -356,7 +356,7 @@ void UGeometryScriptLibrary_ListUtilityFunctions::DuplicateColorList(FGeometrySc
 
 void UGeometryScriptLibrary_ListUtilityFunctions::ClearColorList(FGeometryScriptColorList& ColorList, FLinearColor ClearColor)
 {
-	int32 Num = ColorList.List->Num();
+	int Num = GetColorListLength(ColorList);
 	ColorList.Reset();
 	ColorList.List->Init(ClearColor, Num);
 }
diff --git a/Engine/Plugins/Experimental/GeometryScripting/Source/GeometryScriptingCore/Private/MeshDeformFunctions.cpp b/Engine/Plugins/Experimental/GeometryScripting/Source/GeometryScriptingCore/Private/MeshDeformFunctions.cpp
--- a/Engine/Plugins/Experimental/GeometryScripting/Source/GeometryScriptingCore/Private/MeshDeformFunctions.cpp
+++ b/Engine/Plugins/Experimental/GeometryScripting/Source/GeometryScriptingCore/Private/MeshDeformFunctions.cpp
@@ -537,16 +537,16 @@ UDynamicMesh* UGeometryScriptLibrary_MeshDeformFunctions::ApplyDisplaceFromPerVe
 		return TargetMesh;
 	}
 
-	const TArray<FVector>& Vectors = *VectorList.List;
-
 	TargetMesh->EditMesh([&](FDynamicMesh3& EditMesh)
 	{
-		if (Vectors.Num() < EditMesh.MaxVertexID())
+		if (VectorList.List.IsValid() == false || VectorList.List->Num() < EditMesh.MaxVertexID())
 		{
 			UE::Geometry::AppendError(Debug, EGeometryScriptErrorType::InvalidInputs, LOCTEXT("ApplyDisplaceFromPerVertexVectors_InvalidVectorLength", "ApplyDisplaceFromPerVertexVectors: VectorList Length is less than TargetMesh MaxVertexID"));
 			return;
 		}
 
+		const TArray<FVector>& Vectors = *VectorList.List;
+
 		auto UpdateVertex = [&EditMesh, &Magnitude, &Vectors](int32 VertexID)
 		{
 			FVector Position = EditMesh.GetVertex(VertexID);
diff --git a/Engine/Plugins/Experimental/GeometryScripting/Source/GeometryScriptingCore/Private/TextureMapFunctions.cpp b/Engine/Plugins/Experimental/GeometryScripting/Source/GeometryScriptingCore/Private/TextureMapFunctions.cpp
--- a/Engine/Plugins/Experimental/GeometryScripting/Source/GeometryScriptingCore/Private/TextureMapFunctions.cpp
+++ b/Engine/Plugins/Experimental/GeometryScripting/Source/GeometryScriptingCore/Private/TextureMapFunctions.cpp
@@ -19,6 +19,12 @@ void UGeometryScriptLibrary_TextureMapFunctions::SampleTexture2DAtUVPositions(
 	FGeometryScriptColorList& ColorList,
 	UGeometryScriptDebug* Debug)
 {
+	ColorList.Reset();
+	if (UVList.List.IsValid() == false || UVList.List->Num() == 0)
+	{
+		return;
+	}
+	
 	if (TextureAsset == nullptr)
 	{
 		UE::Geometry::AppendError(Debug, EGeometryScriptErrorType::InvalidInputs, LOCTEXT("SampleTexture2DAtUVPositions_InvalidInput2", "SampleTexture2DAtUVPositions: Texture is Null"));
@@ -35,7 +41,6 @@ void UGeometryScriptLibrary_TextureMapFunctions::SampleTexture2DAtUVPositions(
 	const TArray<FVector2D>& UVs = *UVList.List;
 	int32 NumUVs = UVs.Num();
 
-	ColorList.Reset();
 	TArray<FLinearColor>& Colors = *ColorList.List;
 	Colors.SetNumUninitialized(NumUVs);
 
