GeometryScripting: Fix geometry script list and selection structs not working in Python

diff --git a/Engine/Plugins/Experimental/GeometryScripting/Source/GeometryScriptingCore/Public/GeometryScript/GeometryScriptSelectionTypes.h b/Engine/Plugins/Experimental/GeometryScripting/Source/GeometryScriptingCore/Public/GeometryScript/GeometryScriptSelectionTypes.h
--- a/Engine/Plugins/Experimental/GeometryScripting/Source/GeometryScriptingCore/Public/GeometryScript/GeometryScriptSelectionTypes.h
+++ b/Engine/Plugins/Experimental/GeometryScripting/Source/GeometryScriptingCore/Public/GeometryScript/GeometryScriptSelectionTypes.h
@@ -114,10 +114,29 @@ struct GEOMETRYSCRIPTINGCORE_API FGeometryScriptMeshSelection
 		TFunctionRef<void(int32)> PerVertexFunc,
 		bool bProcessAllVertsIfSelectionEmpty = false) const;
 
+
+	// Required by TStructOpsTypeTraits interface
+	bool operator==(const FGeometryScriptMeshSelection& Other) const
+	{
+		return GeoSelection.Get() == Other.GeoSelection.Get(); 
+	}
+	bool operator!=(const FGeometryScriptMeshSelection& Other) const
+	{
+		return GeoSelection.Get() != Other.GeoSelection.Get(); 
+	}
+
 private:
 	// keeping this private for now in case it needs to be revised in 5.2+
 	TSharedPtr<UE::Geometry::FGeometrySelection> GeoSelection;
 };
 
 
+template<>
+struct TStructOpsTypeTraits<FGeometryScriptMeshSelection> : public TStructOpsTypeTraitsBase2<FGeometryScriptMeshSelection>
+{
+	enum
+	{
+		WithIdenticalViaEquality = true,
+	};
+};
 
diff --git a/Engine/Plugins/Experimental/GeometryScripting/Source/GeometryScriptingCore/Public/GeometryScript/GeometryScriptTypes.h b/Engine/Plugins/Experimental/GeometryScripting/Source/GeometryScriptingCore/Public/GeometryScript/GeometryScriptTypes.h
--- a/Engine/Plugins/Experimental/GeometryScripting/Source/GeometryScriptingCore/Public/GeometryScript/GeometryScriptTypes.h
+++ b/Engine/Plugins/Experimental/GeometryScripting/Source/GeometryScriptingCore/Public/GeometryScript/GeometryScriptTypes.h
@@ -192,6 +192,13 @@ public:
 // List Types
 //
 
+// By default structs exposed to Python will use a per-UPROPERTY comparison. When this doesn't give correct results
+// (e.g., for structs with no properties, which will compare equal in all cases because there are no properties to
+// compare), it is necessary to define explicity equality operators and add the following WithIdenticalViaEquality trait.
+// Note that users can write blueprints/python scripts which pass these lists to many function calls so we would like to
+// avoid copying them (very slow if they have millions of elements) so we defined the equality operations using pointer
+// equality but this is potentially confusing if users expect that different lists with the same elements compare equal
+
 UENUM(BlueprintType)
 enum class EGeometryScriptIndexType : uint8
 {
@@ -227,6 +234,25 @@ public:
 	{
 		return IndexType == OtherType || IndexType == EGeometryScriptIndexType::Any;
 	}
+
+	// Required by TStructOpsTypeTraits interface
+	bool operator==(const FGeometryScriptIndexList& Other) const
+	{
+		return IndexType == Other.IndexType && List.Get() == Other.List.Get(); 
+	}
+	bool operator!=(const FGeometryScriptIndexList& Other) const
+	{
+		return !(*this == Other);
+	}
+};
+
+template<>
+struct TStructOpsTypeTraits<FGeometryScriptIndexList> : public TStructOpsTypeTraitsBase2<FGeometryScriptIndexList>
+{
+	enum
+	{
+		WithIdenticalViaEquality = true,
+	};
 };
 
 
@@ -245,6 +271,25 @@ public:
 		}
 		List->Reset();
 	}
+
+	// Required by TStructOpsTypeTraits interface
+	bool operator==(const FGeometryScriptTriangleList& Other) const
+	{
+		return List.Get() == Other.List.Get(); 
+	}
+	bool operator!=(const FGeometryScriptTriangleList& Other) const
+	{
+		return List.Get() != Other.List.Get(); 
+	}
+};
+
+template<>
+struct TStructOpsTypeTraits<FGeometryScriptTriangleList> : public TStructOpsTypeTraitsBase2<FGeometryScriptTriangleList>
+{
+	enum
+	{
+		WithIdenticalViaEquality = true,
+	};
 };
 
 
@@ -263,6 +308,25 @@ public:
 		}
 		List->Reset();
 	}
+
+	// Required by TStructOpsTypeTraits interface
+	bool operator==(const FGeometryScriptScalarList& Other) const
+	{
+		return List.Get() == Other.List.Get(); 
+	}
+	bool operator!=(const FGeometryScriptScalarList& Other) const
+	{
+		return List.Get() != Other.List.Get(); 
+	}
+};
+
+template<>
+struct TStructOpsTypeTraits<FGeometryScriptScalarList> : public TStructOpsTypeTraitsBase2<FGeometryScriptScalarList>
+{
+	enum
+	{
+		WithIdenticalViaEquality = true,
+	};
 };
 
 
@@ -282,8 +346,26 @@ public:
 		}
 		List->Reset();
 	}
+
+	// Required by TStructOpsTypeTraits interface
+	bool operator==(const FGeometryScriptVectorList& Other) const
+	{
+		return List.Get() == Other.List.Get(); 
+	}
+	bool operator!=(const FGeometryScriptVectorList& Other) const
+	{
+		return List.Get() != Other.List.Get(); 
+	}
 };
 
+template<>
+struct TStructOpsTypeTraits<FGeometryScriptVectorList> : public TStructOpsTypeTraitsBase2<FGeometryScriptVectorList>
+{
+	enum
+	{
+		WithIdenticalViaEquality = true,
+	};
+};
 
 
 USTRUCT(BlueprintType, meta = (DisplayName = "UV List"))
@@ -301,6 +383,25 @@ public:
 		}
 		List->Reset();
 	}
+
+	// Required by TStructOpsTypeTraits interface
+	bool operator==(const FGeometryScriptUVList& Other) const
+	{
+		return List.Get() == Other.List.Get(); 
+	}
+	bool operator!=(const FGeometryScriptUVList& Other) const
+	{
+		return List.Get() != Other.List.Get(); 
+	}
+};
+
+template<>
+struct TStructOpsTypeTraits<FGeometryScriptUVList> : public TStructOpsTypeTraitsBase2<FGeometryScriptUVList>
+{
+	enum
+	{
+		WithIdenticalViaEquality = true,
+	};
 };
 
 
@@ -320,6 +421,25 @@ public:
 		}
 		List->Reset();
 	}
+
+	// Required by TStructOpsTypeTraits interface
+	bool operator==(const FGeometryScriptColorList& Other) const
+	{
+		return List.Get() == Other.List.Get(); 
+	}
+	bool operator!=(const FGeometryScriptColorList& Other) const
+	{
+		return List.Get() != Other.List.Get(); 
+	}
+};
+
+template<>
+struct TStructOpsTypeTraits<FGeometryScriptColorList> : public TStructOpsTypeTraitsBase2<FGeometryScriptColorList>
+{
+	enum
+	{
+		WithIdenticalViaEquality = true,
+	};
 };
 
 
@@ -341,8 +461,29 @@ public:
 		}
 		Path->Reset();
 	}
+
+	// Required by TStructOpsTypeTraits interface
+	bool operator==(const FGeometryScriptPolyPath& Other) const
+	{
+		return bClosedLoop == Other.bClosedLoop && Path.Get() == Other.Path.Get();
+	}
+	bool operator!=(const FGeometryScriptPolyPath& Other) const
+	{
+		return !(*this == Other);
+	}
 };
 
+template<>
+struct TStructOpsTypeTraits<FGeometryScriptPolyPath> : public TStructOpsTypeTraitsBase2<FGeometryScriptPolyPath>
+{
+	enum
+	{
+		WithIdenticalViaEquality = true,
+	};
+};
+
+
+
 
 //
 // Spatial data structures
@@ -357,6 +498,25 @@ public:
 
 	TSharedPtr<UE::Geometry::FDynamicMeshAABBTree3> Spatial;
 	TSharedPtr<UE::Geometry::TFastWindingTree<UE::Geometry::FDynamicMesh3>> FWNTree;
+	
+	// Required by TStructOpsTypeTraits interface
+	bool operator==(const FGeometryScriptDynamicMeshBVH& Other) const
+	{
+		return Spatial.Get() == Other.Spatial.Get() && FWNTree.Get() == Other.FWNTree.Get();
+	}
+	bool operator!=(const FGeometryScriptDynamicMeshBVH& Other) const
+	{
+		return !(*this == Other);
+	}
+};
+
+template<>
+struct TStructOpsTypeTraits<FGeometryScriptDynamicMeshBVH> : public TStructOpsTypeTraitsBase2<FGeometryScriptDynamicMeshBVH>
+{
+	enum
+	{
+		WithIdenticalViaEquality = true,
+	};
 };
 
 
