GeometryScript: add small library of geometric functions for working with FRay in Blueprints/Python. Add TranslatePivotToLocation helper function. Add SetAllMeshVertexPositions. Fix bug in IntersectionUtil.h function that had never been called.

diff --git a/Engine/Plugins/Experimental/GeometryScripting/Source/GeometryScriptingCore/Private/MeshBasicEditFunctions.cpp b/Engine/Plugins/Experimental/GeometryScripting/Source/GeometryScriptingCore/Private/MeshBasicEditFunctions.cpp
--- a/Engine/Plugins/Experimental/GeometryScripting/Source/GeometryScriptingCore/Private/MeshBasicEditFunctions.cpp
+++ b/Engine/Plugins/Experimental/GeometryScripting/Source/GeometryScriptingCore/Private/MeshBasicEditFunctions.cpp
@@ -52,6 +52,43 @@ UDynamicMesh* UGeometryScriptLibrary_MeshBasicEditFunctions::SetVertexPosition(U
 }
 
 
+UDynamicMesh* UGeometryScriptLibrary_MeshBasicEditFunctions::SetAllMeshVertexPositions(
+	UDynamicMesh* TargetMesh,
+	FGeometryScriptVectorList PositionList,
+	UGeometryScriptDebug* Debug)
+{
+	if (TargetMesh == nullptr)
+	{
+		UE::Geometry::AppendError(Debug, EGeometryScriptErrorType::InvalidInputs, LOCTEXT("SetAllMeshVertexPositions_InvalidMesh", "SetAllMeshVertexPositions: TargetMesh is Null"));
+		return TargetMesh;
+	}
+	if (PositionList.List.IsValid() == false || PositionList.List->Num() == 0)
+	{
+		UE::Geometry::AppendError(Debug, EGeometryScriptErrorType::InvalidInputs, LOCTEXT("SetAllMeshVertexPositions_InvalidList", "SetAllMeshVertexPositions: List is empty"));
+		return TargetMesh;
+	}
+
+	TargetMesh->EditMesh([&](FDynamicMesh3& EditMesh) 
+	{
+		const TArray<FVector>& VertexPositions = *PositionList.List;
+		if (VertexPositions.Num() < EditMesh.MaxVertexID())
+		{
+			UE::Geometry::AppendError(Debug, EGeometryScriptErrorType::InvalidInputs, LOCTEXT("SetAllMeshVertexPositions_IncorrectCount", "SetAllMeshVertexPositions: size of provided PositionList is smaller than MaxVertexID of Mesh"));
+		}
+		else
+		{
+			for (int32 VertexID : EditMesh.VertexIndicesItr())
+			{
+				EditMesh.SetVertex( VertexID, VertexPositions[VertexID] );
+			}
+		}
+	}, EDynamicMeshChangeType::GeneralEdit, EDynamicMeshAttributeChangeFlags::Unknown, false);
+
+	return TargetMesh;
+}
+
+
+
 
 UDynamicMesh* UGeometryScriptLibrary_MeshBasicEditFunctions::AddVertexToMesh(
 	UDynamicMesh* TargetMesh,
diff --git a/Engine/Plugins/Experimental/GeometryScripting/Source/GeometryScriptingCore/Private/ShapeFunctions.cpp b/Engine/Plugins/Experimental/GeometryScripting/Source/GeometryScriptingCore/Private/ShapeFunctions.cpp
--- a/Engine/Plugins/Experimental/GeometryScripting/Source/GeometryScriptingCore/Private/ShapeFunctions.cpp
+++ b/Engine/Plugins/Experimental/GeometryScripting/Source/GeometryScriptingCore/Private/ShapeFunctions.cpp
@@ -2,6 +2,8 @@
 
 #include "GeometryScript/ShapeFunctions.h"
 #include "VectorTypes.h"
+#include "Intersection/IntersectionUtil.h"
+#include "Intersection/IntrRay3AxisAlignedBox3.h"
 
 #include UE_INLINE_GENERATED_CPP_BY_NAME(ShapeFunctions)
 
@@ -9,6 +11,67 @@ using namespace UE::Geometry;
 
 #define LOCTEXT_NAMESPACE "UGeometryScriptLibrary_BoxFunctions"
 
+
+
+FRay UGeometryScriptLibrary_RayFunctions::MakeRayFromPoints(const FVector& A, const FVector& B)
+{
+	return FRay(A, Normalized(B-A), true);
+}
+
+FRay UGeometryScriptLibrary_RayFunctions::MakeRayFromPointDirection(const FVector& Origin, const FVector& Direction, bool bDirectionIsNormalized)
+{
+	return FRay(Origin, Direction, bDirectionIsNormalized);
+}
+
+FRay UGeometryScriptLibrary_RayFunctions::GetTransformedRay(const FRay& Ray, const FTransform& TransformIn, bool bInvert)
+{
+	FTransformSRT3d Transform(TransformIn);
+	return (bInvert) ? Transform.InverseTransformRay(Ray) : Transform.TransformRay(Ray);
+}
+
+FVector UGeometryScriptLibrary_RayFunctions::GetRayPoint(const FRay& Ray, double Distance)
+{
+	return Ray.PointAt(Distance);
+}
+
+double UGeometryScriptLibrary_RayFunctions::GetRayParameter(const FRay& Ray, FVector& Point)
+{
+	return Ray.GetParameter(Point);
+}
+
+double UGeometryScriptLibrary_RayFunctions::GetRayPointDistance(const FRay& Ray, FVector& Point)
+{
+	return Ray.Dist(Point);
+}
+
+FVector UGeometryScriptLibrary_RayFunctions::GetRayClosestPoint(const FRay& Ray, FVector& Point)
+{
+	return Ray.ClosestPoint(Point);
+}
+
+bool UGeometryScriptLibrary_RayFunctions::GetRaySphereIntersection(const FRay& Ray, const FVector& SphereCenter, double SphereRadius, double& Distance1, double& Distance2)
+{
+	Distance1 = Distance2 = (double)TNumericLimits<float>::Max();
+	FLinearIntersection Intersection;
+	bool bIntersects = IntersectionUtil::RaySphereIntersection(Ray.Origin, Ray.Direction, SphereCenter, SphereRadius, Intersection);
+	if (bIntersects)
+	{
+		Distance1 = Intersection.parameter.Min;
+		Distance2 = (Intersection.numIntersections > 1) ? Intersection.parameter.Max : Intersection.parameter.Min;
+	}
+	return bIntersects;
+}
+
+bool UGeometryScriptLibrary_RayFunctions::GetRayBoxIntersection(const FRay& Ray, const FBox& Box, double& HitDistance)
+{
+	HitDistance = (double)TNumericLimits<float>::Max();
+	return FIntrRay3AxisAlignedBox3d::FindIntersection(Ray, FAxisAlignedBox3d(Box), HitDistance);
+}
+
+
+
+
+
 FBox UGeometryScriptLibrary_BoxFunctions::MakeBoxFromCenterSize(const FVector& Center, const FVector& Dimensions)
 {
 	FVector Extents( FMathd::Max(0, Dimensions.X * 0.5), FMathd::Max(0, Dimensions.Y * 0.5), FMathd::Max(0, Dimensions.Z * 0.5) );
diff --git a/Engine/Plugins/Experimental/GeometryScripting/Source/GeometryScriptingCore/Public/GeometryScript/MeshBasicEditFunctions.h b/Engine/Plugins/Experimental/GeometryScripting/Source/GeometryScriptingCore/Public/GeometryScript/MeshBasicEditFunctions.h
--- a/Engine/Plugins/Experimental/GeometryScripting/Source/GeometryScriptingCore/Public/GeometryScript/MeshBasicEditFunctions.h
+++ b/Engine/Plugins/Experimental/GeometryScripting/Source/GeometryScriptingCore/Public/GeometryScript/MeshBasicEditFunctions.h
@@ -109,6 +109,17 @@ public:
 		bool& bIsValidVertex, 
 		bool bDeferChangeNotifications = false );
 
+	/**
+	 * Set all vertex positions in the TargetMesh to the specified Positions
+	 * @param PositionList new vertex Positions. Size must be less than or equal to the MaxVertexID of TargetMesh  (ie gaps are supported)
+	 */
+	UFUNCTION(BlueprintCallable, Category = "GeometryScript|MeshEdits", meta=(ScriptMethod))
+	static UPARAM(DisplayName = "Target Mesh") UDynamicMesh*
+	SetAllMeshVertexPositions(
+		UDynamicMesh* TargetMesh,
+		FGeometryScriptVectorList PositionList,
+		UGeometryScriptDebug* Debug = nullptr);
+
 
 	UFUNCTION(BlueprintCallable, Category = "GeometryScript|MeshEdits", meta=(ScriptMethod))
 	static UPARAM(DisplayName = "Target Mesh") UDynamicMesh* 
diff --git a/Engine/Plugins/Experimental/GeometryScripting/Source/GeometryScriptingCore/Public/GeometryScript/MeshTransformFunctions.h b/Engine/Plugins/Experimental/GeometryScripting/Source/GeometryScriptingCore/Public/GeometryScript/MeshTransformFunctions.h
--- a/Engine/Plugins/Experimental/GeometryScripting/Source/GeometryScriptingCore/Public/GeometryScript/MeshTransformFunctions.h
+++ b/Engine/Plugins/Experimental/GeometryScripting/Source/GeometryScriptingCore/Public/GeometryScript/MeshTransformFunctions.h
@@ -87,4 +87,20 @@ public:
 		UGeometryScriptDebug* Debug = nullptr);
 
 
+	/**
+	 * Set the Pivot Location for the Mesh. Since the Pivot of a Mesh object is always the point at (0,0,0),
+	 * this function simply translates the mesh by -PivotLocation.
+	 */
+	UFUNCTION(BlueprintCallable, Category = "GeometryScript|Transforms", meta=(ScriptMethod))
+	static UPARAM(DisplayName = "Target Mesh") UDynamicMesh* 
+	TranslatePivotToLocation(
+		UDynamicMesh* TargetMesh,
+		FVector PivotLocation,
+		UGeometryScriptDebug* Debug = nullptr)
+	{
+		// note: this function is redundant, however *many* users do not intuitively understand the relationship
+		// between "setting the pivot" and translation. This function is for those users.
+		return TranslateMesh(TargetMesh, -PivotLocation, Debug);
+	}
+
 };
\ No newline at end of file
diff --git a/Engine/Plugins/Experimental/GeometryScripting/Source/GeometryScriptingCore/Public/GeometryScript/ShapeFunctions.h b/Engine/Plugins/Experimental/GeometryScripting/Source/GeometryScriptingCore/Public/GeometryScript/ShapeFunctions.h
--- a/Engine/Plugins/Experimental/GeometryScripting/Source/GeometryScriptingCore/Public/GeometryScript/ShapeFunctions.h
+++ b/Engine/Plugins/Experimental/GeometryScripting/Source/GeometryScriptingCore/Public/GeometryScript/ShapeFunctions.h
@@ -8,6 +8,68 @@
 #include "ShapeFunctions.generated.h"
 
 
+UCLASS(meta = (ScriptName = "GeometryScript_Ray"))
+class GEOMETRYSCRIPTINGCORE_API UGeometryScriptLibrary_RayFunctions : public UBlueprintFunctionLibrary
+{
+	GENERATED_BODY()
+public:
+
+	/**
+	 * Create a Ray from two points, placing the Origin at A and the Direction as Normalize(B-A)
+	 */
+	UFUNCTION(BlueprintPure, Category = "GeometryScript|Shapes|Ray")
+	static FRay MakeRayFromPoints(const FVector& A, const FVector& B);
+
+	/**
+	 * Create a Ray from an Origin and Direction, with optionally non-normalized Direction
+	 */
+	UFUNCTION(BlueprintPure, Category = "GeometryScript|Shapes|Ray")
+	static FRay MakeRayFromPointDirection(const FVector& Origin, const FVector& Direction, bool bDirectionIsNormalized = true);
+
+	/**
+	 * Apply the given Transform to the given Ray, or optionally the Transform Inverse, and return the new transformed Ray
+	 */
+	UFUNCTION(BlueprintPure, Category = "GeometryScript|Shapes|Ray", meta=(ScriptMethod))
+	static FRay GetTransformedRay(const FRay& Ray, const FTransform& Transform, bool bInvert = false);
+
+	/**
+	 * Get a Point at the given Distance along the Ray (Origin + Distance*Direction)
+	 */
+	UFUNCTION(BlueprintPure, Category = "GeometryScript|Shapes|Ray", meta=(ScriptMethod))
+	static FVector GetRayPoint(const FRay& Ray, double Distance);
+
+	/**
+	 * Project the given Point onto the closest point along the Ray, and return the Ray Parameter/Distance at that Point
+	 */
+	UFUNCTION(BlueprintPure, Category = "GeometryScript|Shapes|Ray", meta=(ScriptMethod))
+	static double GetRayParameter(const FRay& Ray, FVector& Point);
+
+	/**
+	 * Get the distance from Point to the closest point on the Ray
+	 */
+	UFUNCTION(BlueprintPure, Category = "GeometryScript|Shapes|Ray", meta=(ScriptMethod))
+	static double GetRayPointDistance(const FRay& Ray, FVector& Point);
+
+	/**
+	 * Get the closest point on the Ray to the given Point
+	 */
+	UFUNCTION(BlueprintPure, Category = "GeometryScript|Shapes|Ray", meta=(ScriptMethod))
+	static FVector GetRayClosestPoint(const FRay& Ray, FVector& Point);
+
+	/**
+	 * Check if the Ray intersects a Sphere defined by the SphereCenter and SphereRadius
+	 */
+	UFUNCTION(BlueprintPure, Category = "GeometryScript|Shapes|Ray", meta=(ScriptMethod))
+	static bool GetRaySphereIntersection(const FRay& Ray, const FVector& SphereCenter, double SphereRadius, double& Distance1, double& Distance2);
+
+	/**
+	 * Check if the Ray intersects a Sphere defined by the SphereCenter and SphereRadius
+	 */
+	UFUNCTION(BlueprintPure, Category = "GeometryScript|Shapes|Ray", meta=(ScriptMethod))
+	static bool GetRayBoxIntersection(const FRay& Ray, const FBox& Box, double& HitDistance);
+
+};
+
 
 UCLASS(meta = (ScriptName = "GeometryScript_Box"))
 class GEOMETRYSCRIPTINGCORE_API UGeometryScriptLibrary_BoxFunctions : public UBlueprintFunctionLibrary
diff --git a/Engine/Source/Runtime/GeometryCore/Public/Intersection/IntersectionUtil.h b/Engine/Source/Runtime/GeometryCore/Public/Intersection/IntersectionUtil.h
--- a/Engine/Source/Runtime/GeometryCore/Public/Intersection/IntersectionUtil.h
+++ b/Engine/Source/Runtime/GeometryCore/Public/Intersection/IntersectionUtil.h
@@ -239,7 +239,7 @@ namespace IntersectionUtil
 		// adapted from GeometricTools GTEngine
 		// https://www.geometrictools.com/GTEngine/Include/Mathematics/GteIntrRay3Sphere3.h
 
-		RaySphereIntersection(RayOrigin, RayDirection, SphereCenter, SphereRadius, Result);
+		LineSphereIntersection(RayOrigin, RayDirection, SphereCenter, SphereRadius, Result);
 		if (Result.intersects) 
 		{
 			// The line containing the ray intersects the sphere; the t-interval
