GeometryScript: fix up const mistake in new ray functions and improve documentation/labeling

diff --git a/Engine/Plugins/Experimental/GeometryScripting/Source/GeometryScriptingCore/Private/ShapeFunctions.cpp b/Engine/Plugins/Experimental/GeometryScripting/Source/GeometryScriptingCore/Private/ShapeFunctions.cpp
--- a/Engine/Plugins/Experimental/GeometryScripting/Source/GeometryScriptingCore/Private/ShapeFunctions.cpp
+++ b/Engine/Plugins/Experimental/GeometryScripting/Source/GeometryScriptingCore/Private/ShapeFunctions.cpp
@@ -15,7 +15,12 @@ using namespace UE::Geometry;
 
 FRay UGeometryScriptLibrary_RayFunctions::MakeRayFromPoints(const FVector& A, const FVector& B)
 {
-	return FRay(A, Normalized(B-A), true);
+	FVector Direction(B-A);
+	if (Normalize(Direction) < FMathf::Epsilon)
+	{
+		Direction = FVector(0,0,1);
+	}
+	return FRay(A, Direction, true);
 }
 
 FRay UGeometryScriptLibrary_RayFunctions::MakeRayFromPointDirection(const FVector& Origin, const FVector& Direction, bool bDirectionIsNormalized)
@@ -34,17 +39,17 @@ FVector UGeometryScriptLibrary_RayFunctions::GetRayPoint(const FRay& Ray, double
 	return Ray.PointAt(Distance);
 }
 
-double UGeometryScriptLibrary_RayFunctions::GetRayParameter(const FRay& Ray, FVector& Point)
+double UGeometryScriptLibrary_RayFunctions::GetRayParameter(const FRay& Ray, const FVector& Point)
 {
 	return Ray.GetParameter(Point);
 }
 
-double UGeometryScriptLibrary_RayFunctions::GetRayPointDistance(const FRay& Ray, FVector& Point)
+double UGeometryScriptLibrary_RayFunctions::GetRayPointDistance(const FRay& Ray, const FVector& Point)
 {
 	return Ray.Dist(Point);
 }
 
-FVector UGeometryScriptLibrary_RayFunctions::GetRayClosestPoint(const FRay& Ray, FVector& Point)
+FVector UGeometryScriptLibrary_RayFunctions::GetRayClosestPoint(const FRay& Ray, const FVector& Point)
 {
 	return Ray.ClosestPoint(Point);
 }
diff --git a/Engine/Plugins/Experimental/GeometryScripting/Source/GeometryScriptingCore/Public/GeometryScript/ShapeFunctions.h b/Engine/Plugins/Experimental/GeometryScripting/Source/GeometryScriptingCore/Public/GeometryScript/ShapeFunctions.h
--- a/Engine/Plugins/Experimental/GeometryScripting/Source/GeometryScriptingCore/Public/GeometryScript/ShapeFunctions.h
+++ b/Engine/Plugins/Experimental/GeometryScripting/Source/GeometryScriptingCore/Public/GeometryScript/ShapeFunctions.h
@@ -18,55 +18,72 @@ public:
 	 * Create a Ray from two points, placing the Origin at A and the Direction as Normalize(B-A)
 	 */
 	UFUNCTION(BlueprintPure, Category = "GeometryScript|Shapes|Ray")
-	static FRay MakeRayFromPoints(const FVector& A, const FVector& B);
+	static UPARAM(DisplayName="Ray") FRay
+	MakeRayFromPoints(const FVector& A, const FVector& B);
 
 	/**
 	 * Create a Ray from an Origin and Direction, with optionally non-normalized Direction
 	 */
 	UFUNCTION(BlueprintPure, Category = "GeometryScript|Shapes|Ray")
-	static FRay MakeRayFromPointDirection(const FVector& Origin, const FVector& Direction, bool bDirectionIsNormalized = true);
+	static UPARAM(DisplayName="Ray") FRay 
+	MakeRayFromPointDirection(const FVector& Origin, const FVector& Direction, bool bDirectionIsNormalized = true);
 
 	/**
 	 * Apply the given Transform to the given Ray, or optionally the Transform Inverse, and return the new transformed Ray
 	 */
 	UFUNCTION(BlueprintPure, Category = "GeometryScript|Shapes|Ray", meta=(ScriptMethod))
-	static FRay GetTransformedRay(const FRay& Ray, const FTransform& Transform, bool bInvert = false);
+	static UPARAM(DisplayName="Transformed Ray") FRay 
+	GetTransformedRay(const FRay& Ray, const FTransform& Transform, bool bInvert = false);
 
 	/**
 	 * Get a Point at the given Distance along the Ray (Origin + Distance*Direction)
 	 */
 	UFUNCTION(BlueprintPure, Category = "GeometryScript|Shapes|Ray", meta=(ScriptMethod))
-	static FVector GetRayPoint(const FRay& Ray, double Distance);
+	static UPARAM(DisplayName="Point") FVector 
+	GetRayPoint(const FRay& Ray, double Distance);
 
 	/**
 	 * Project the given Point onto the closest point along the Ray, and return the Ray Parameter/Distance at that Point
 	 */
 	UFUNCTION(BlueprintPure, Category = "GeometryScript|Shapes|Ray", meta=(ScriptMethod))
-	static double GetRayParameter(const FRay& Ray, FVector& Point);
+	static UPARAM(DisplayName="Ray Paramater") double 
+	GetRayParameter(const FRay& Ray, const FVector& Point);
 
 	/**
 	 * Get the distance from Point to the closest point on the Ray
 	 */
 	UFUNCTION(BlueprintPure, Category = "GeometryScript|Shapes|Ray", meta=(ScriptMethod))
-	static double GetRayPointDistance(const FRay& Ray, FVector& Point);
+	static UPARAM(DisplayName="Distance") double 
+	GetRayPointDistance(const FRay& Ray, const FVector& Point);
 
 	/**
 	 * Get the closest point on the Ray to the given Point
 	 */
 	UFUNCTION(BlueprintPure, Category = "GeometryScript|Shapes|Ray", meta=(ScriptMethod))
-	static FVector GetRayClosestPoint(const FRay& Ray, FVector& Point);
+	static UPARAM(DisplayName="Closest Point") FVector 
+	GetRayClosestPoint(const FRay& Ray, const FVector& Point);
 
 	/**
-	 * Check if the Ray intersects a Sphere defined by the SphereCenter and SphereRadius
+	 * Check if the Ray intersects a Sphere defined by the SphereCenter and SphereRadius.
+	 * This function returns two intersection distances (ray parameters). If the ray grazes the sphere, both
+	 * distances will be the same, and if it misses, they will be MAX_FLOAT. 
+	 * Use the function GetRayPoint to convert the distances to points on the ray/sphere.
+	 * 
+	 * @param Distance1 Distance along ray (Ray Parameter) to first/closer intersection point with sphere
+	 * @param Distance2 Distance along ray (Ray Parameter) to second/further intersection point with sphere
+	 * @return true if ray intersects sphere
 	 */
 	UFUNCTION(BlueprintPure, Category = "GeometryScript|Shapes|Ray", meta=(ScriptMethod))
-	static bool GetRaySphereIntersection(const FRay& Ray, const FVector& SphereCenter, double SphereRadius, double& Distance1, double& Distance2);
+	static UPARAM(DisplayName="Intersects") bool 
+	GetRaySphereIntersection(const FRay& Ray, const FVector& SphereCenter, double SphereRadius, double& Distance1, double& Distance2);
 
 	/**
-	 * Check if the Ray intersects a Sphere defined by the SphereCenter and SphereRadius
+	 * Check if the Ray intersects a Sphere defined by the SphereCenter and SphereRadius.
+	 * @param HitDistance Distance along the ray (Ray Parameter) to first intersection point with the Box
 	 */
 	UFUNCTION(BlueprintPure, Category = "GeometryScript|Shapes|Ray", meta=(ScriptMethod))
-	static bool GetRayBoxIntersection(const FRay& Ray, const FBox& Box, double& HitDistance);
+	static UPARAM(DisplayName="Intersects") bool 
+	GetRayBoxIntersection(const FRay& Ray, const FBox& Box, double& HitDistance);
 
 };
 
@@ -81,7 +98,8 @@ public:
 	 * Create a Box from a Center point and X/Y/Z Dimensions (*not* Extents, which are half-dimensions)
 	 */
 	UFUNCTION(BlueprintPure, Category = "GeometryScript|Shapes|Box")
-	static FBox MakeBoxFromCenterSize(const FVector& Center, const FVector& Dimensions);
+	static UPARAM(DisplayName="Box") FBox 
+	MakeBoxFromCenterSize(const FVector& Center, const FVector& Dimensions);
 
 	/**
 	 * Get the Center point and X/Y/Z Dimensions of a Box (full dimensions, not Extents)
