GeometryCore: add ability to return TriangleIDs of sample points in MeshSurfacePointSampling, and also add optional computation of barycentric coordinates for each sample
GeometryScripting: add return of TriangleIDs to geo-script point sampling functions. Add utility function GetRayStartEnd and MakeBoxFromCenterExtents to simplify common BP operations, improve some display names.

diff --git a/Engine/Plugins/Experimental/GeometryScripting/Source/GeometryScriptingCore/Private/MeshSamplingFunctions.cpp b/Engine/Plugins/Experimental/GeometryScripting/Source/GeometryScriptingCore/Private/MeshSamplingFunctions.cpp
--- a/Engine/Plugins/Experimental/GeometryScripting/Source/GeometryScriptingCore/Private/MeshSamplingFunctions.cpp
+++ b/Engine/Plugins/Experimental/GeometryScripting/Source/GeometryScriptingCore/Private/MeshSamplingFunctions.cpp
@@ -16,6 +16,7 @@ UDynamicMesh* UGeometryScriptLibrary_MeshSamplingFunctions::ComputePointSampling
 	UDynamicMesh* TargetMesh,
 	FGeometryScriptMeshPointSamplingOptions Options,
 	TArray<FTransform>& Samples,
+	FGeometryScriptIndexList& TriangleIDs,
 	UGeometryScriptDebug* Debug)
 {
 	if (TargetMesh == nullptr)
@@ -37,9 +38,15 @@ UDynamicMesh* UGeometryScriptLibrary_MeshSamplingFunctions::ComputePointSampling
 		Sampler.ComputePoissonSampling(ReadMesh);
 	});
 
-	for (FFrame3d Frame : Sampler.Samples)
+	TriangleIDs.Reset(EGeometryScriptIndexType::Triangle);
+	TArray<int>& IndexList = *TriangleIDs.List;
+	int32 NumSamples = Sampler.Samples.Num();
+	for ( int32 k = 0; k < NumSamples; ++k )
 	{
+		const FFrame3d& Frame = Sampler.Samples[k];
 		Samples.Add( Frame.ToFTransform() );
+
+		IndexList.Add( Sampler.TriangleIDs[k] );
 	}
 
 	return TargetMesh;
@@ -52,6 +59,7 @@ UDynamicMesh* UGeometryScriptLibrary_MeshSamplingFunctions::ComputeNonUniformPoi
 	FGeometryScriptNonUniformPointSamplingOptions NonUniformOptions,
 	TArray<FTransform>& Samples,
 	TArray<double>& SampleRadii,
+	FGeometryScriptIndexList& TriangleIDs,
 	UGeometryScriptDebug* Debug)
 {
 	if (TargetMesh == nullptr)
@@ -79,10 +87,15 @@ UDynamicMesh* UGeometryScriptLibrary_MeshSamplingFunctions::ComputeNonUniformPoi
 		Sampler.ComputePoissonSampling(ReadMesh);
 	});
 
-	for ( int32 k = 0; k < Sampler.Samples.Num(); ++k)
+	TriangleIDs.Reset(EGeometryScriptIndexType::Triangle);
+	TArray<int>& IndexList = *TriangleIDs.List;
+	int32 NumSamples = Sampler.Samples.Num();
+	for ( int32 k = 0; k < NumSamples; ++k )
 	{
-		Samples.Add( Sampler.Samples[k].ToFTransform());
-		SampleRadii.Add( Sampler.Radii[k] );
+		Samples.Add( Sampler.Samples[k].ToFTransform() );
+		SampleRadii.Add( Sampler.Radii[k]);
+
+		IndexList.Add( Sampler.TriangleIDs[k] );
 	}
 
 	return TargetMesh;
@@ -99,6 +112,7 @@ UDynamicMesh* UGeometryScriptLibrary_MeshSamplingFunctions::ComputeVertexWeighte
 	FGeometryScriptScalarList VertexWeights,
 	TArray<FTransform>& Samples,
 	TArray<double>& SampleRadii,
+	FGeometryScriptIndexList& TriangleIDs,
 	UGeometryScriptDebug* Debug)
 {
 	if (TargetMesh == nullptr)
@@ -135,10 +149,15 @@ UDynamicMesh* UGeometryScriptLibrary_MeshSamplingFunctions::ComputeVertexWeighte
 		Sampler.ComputePoissonSampling(ReadMesh);
 	});
 
-	for ( int32 k = 0; k < Sampler.Samples.Num(); ++k)
+	TriangleIDs.Reset(EGeometryScriptIndexType::Triangle);
+	TArray<int>& IndexList = *TriangleIDs.List;
+	int32 NumSamples = Sampler.Samples.Num();
+	for ( int32 k = 0; k < NumSamples; ++k )
 	{
-		Samples.Add( Sampler.Samples[k].ToFTransform());
-		SampleRadii.Add( Sampler.Radii[k] );
+		Samples.Add( Sampler.Samples[k].ToFTransform() );
+		SampleRadii.Add( Sampler.Radii[k]);
+
+		IndexList.Add( Sampler.TriangleIDs[k] );
 	}
 
 	return TargetMesh;
diff --git a/Engine/Plugins/Experimental/GeometryScripting/Source/GeometryScriptingCore/Private/ShapeFunctions.cpp b/Engine/Plugins/Experimental/GeometryScripting/Source/GeometryScriptingCore/Private/ShapeFunctions.cpp
--- a/Engine/Plugins/Experimental/GeometryScripting/Source/GeometryScriptingCore/Private/ShapeFunctions.cpp
+++ b/Engine/Plugins/Experimental/GeometryScripting/Source/GeometryScriptingCore/Private/ShapeFunctions.cpp
@@ -39,6 +39,12 @@ FVector UGeometryScriptLibrary_RayFunctions::GetRayPoint(const FRay& Ray, double
 	return Ray.PointAt(Distance);
 }
 
+void UGeometryScriptLibrary_RayFunctions::GetRayStartEnd(const FRay& Ray, double StartDistance, double EndDistance, FVector& StartPoint, FVector& EndPoint)
+{
+	StartPoint = Ray.PointAt(StartDistance);
+	EndPoint = Ray.PointAt(EndDistance == 0 ? TNumericLimits<float>::Max() : EndDistance );
+}
+
 double UGeometryScriptLibrary_RayFunctions::GetRayParameter(const FRay& Ray, const FVector& Point)
 {
 	return Ray.GetParameter(Point);
@@ -83,6 +89,12 @@ FBox UGeometryScriptLibrary_BoxFunctions::MakeBoxFromCenterSize(const FVector& C
 	return FBox(Center - Extents, Center + Extents);
 }
 
+FBox UGeometryScriptLibrary_BoxFunctions::MakeBoxFromCenterExtents(const FVector& Center, const FVector& Extents)
+{
+	return FBox(Center - Extents, Center + Extents);
+}
+
+
 void UGeometryScriptLibrary_BoxFunctions::GetBoxCenterSize(const FBox& Box, FVector& Center, FVector& Dimensions)
 {
 	FVector Extents;
diff --git a/Engine/Plugins/Experimental/GeometryScripting/Source/GeometryScriptingCore/Public/GeometryScript/MeshSamplingFunctions.h b/Engine/Plugins/Experimental/GeometryScripting/Source/GeometryScriptingCore/Public/GeometryScript/MeshSamplingFunctions.h
--- a/Engine/Plugins/Experimental/GeometryScripting/Source/GeometryScriptingCore/Public/GeometryScript/MeshSamplingFunctions.h
+++ b/Engine/Plugins/Experimental/GeometryScripting/Source/GeometryScriptingCore/Public/GeometryScript/MeshSamplingFunctions.h
@@ -103,6 +103,8 @@ public:
 	 * Compute a set of sample points lying on the surface of TargetMesh based on the provided sampling Options.
 	 * Samples are approximately uniformly distributed, and non-overlapping relative to the provided Options.SamplingRadius,
 	 * ie the distance between any pair of samples if >= 2*SamplingRadius.
+	 * @param Samples output list of sample points. Transform Location is sample position, Rotation orients Z with the triangle normal
+	 * @param TriangleIDs TriangleID that contains each sample point. Length is the same as Samples array.
 	 */
 	UFUNCTION(BlueprintCallable, Category = "GeometryScript|MeshSampling", meta=(ScriptMethod))
 	static UPARAM(DisplayName = "Target Mesh") UDynamicMesh* 
@@ -110,6 +112,7 @@ public:
 		UDynamicMesh* TargetMesh, 
 		FGeometryScriptMeshPointSamplingOptions Options,
 		TArray<FTransform>& Samples,
+		FGeometryScriptIndexList& TriangleIDs,
 		UGeometryScriptDebug* Debug = nullptr);
 
 	/**
@@ -125,6 +128,7 @@ public:
 		FGeometryScriptNonUniformPointSamplingOptions NonUniformOptions,
 		TArray<FTransform>& Samples,
 		TArray<double>& SampleRadii,
+		FGeometryScriptIndexList& TriangleIDs,
 		UGeometryScriptDebug* Debug = nullptr);
 
 	/**
@@ -142,6 +146,7 @@ public:
 		FGeometryScriptScalarList VertexWeights,
 		TArray<FTransform>& Samples,
 		TArray<double>& SampleRadii,
+		FGeometryScriptIndexList& TriangleIDs,
 		UGeometryScriptDebug* Debug = nullptr);
 
 };
\ No newline at end of file
diff --git a/Engine/Plugins/Experimental/GeometryScripting/Source/GeometryScriptingCore/Public/GeometryScript/ShapeFunctions.h b/Engine/Plugins/Experimental/GeometryScripting/Source/GeometryScriptingCore/Public/GeometryScript/ShapeFunctions.h
--- a/Engine/Plugins/Experimental/GeometryScripting/Source/GeometryScriptingCore/Public/GeometryScript/ShapeFunctions.h
+++ b/Engine/Plugins/Experimental/GeometryScripting/Source/GeometryScriptingCore/Public/GeometryScript/ShapeFunctions.h
@@ -42,6 +42,15 @@ public:
 	static UPARAM(DisplayName="Point") FVector 
 	GetRayPoint(const FRay& Ray, double Distance);
 
+	/**
+	 * Get two points along the ray. 
+	 * @param StartPoint returned as Origin + StartDistance*Direction
+	 * @param EndPoint returned as Origin + EndDistance*Direction, Unless EndDistance = 0, then MaxFloat is used as the Distance
+	 */
+	UFUNCTION(BlueprintPure, Category = "GeometryScript|Shapes|Ray", meta=(ScriptMethod, StartDistance="0", EndDistance = "0"))
+	static void 
+	GetRayStartEnd(const FRay& Ray, double StartDistance, double EndDistance, FVector& StartPoint, FVector& EndPoint);
+
 	/**
 	 * Project the given Point onto the closest point along the Ray, and return the Ray Parameter/Distance at that Point
 	 */
@@ -101,6 +110,13 @@ public:
 	static UPARAM(DisplayName="Box") FBox 
 	MakeBoxFromCenterSize(const FVector& Center, const FVector& Dimensions);
 
+	/**
+	 * Create a Box from a Center point and X/Y/Z Extents (Extents are half-dimensions)
+	 */
+	UFUNCTION(BlueprintPure, Category = "GeometryScript|Shapes|Box")
+	static UPARAM(DisplayName="Box") FBox 
+	MakeBoxFromCenterExtents(const FVector& Center, const FVector& Extents);
+
 	/**
 	 * Get the Center point and X/Y/Z Dimensions of a Box (full dimensions, not Extents)
 	 */
@@ -114,7 +130,8 @@ public:
 	 * 7 is the Max corner, and 4/5/6 are -Z/-Y/-X from the Max corner.
 	 */
 	UFUNCTION(BlueprintPure, Category = "GeometryScript|Shapes|Box", meta=(ScriptMethod))
-	static FVector GetBoxCorner(const FBox& Box, int CornerIndex);
+	static UPARAM(DisplayName="Corner Point") FVector 
+	GetBoxCorner(const FBox& Box, int CornerIndex);
 
 	/**
 	 * Get the position of the center of a face of the Box. Faces are indexed from 0 to 5,
@@ -122,7 +139,8 @@ public:
 	 * @param FaceNormal returned Normal vector of the identified face
 	 */
 	UFUNCTION(BlueprintPure, Category = "GeometryScript|Shapes|Box", meta=(ScriptMethod))
-	static FVector GetBoxFaceCenter(const FBox& Box, int FaceIndex, FVector& FaceNormal);
+	static UPARAM(DisplayName="Center Point") FVector 
+	GetBoxFaceCenter(const FBox& Box, int FaceIndex, FVector& FaceNormal);
 
 	/**
 	 * Get the Volume and Surface Area of a Box
@@ -135,39 +153,45 @@ public:
 	 * Dimensions will be clamped to the center point if any of ExpandBy are larger than half the box size
 	 */
 	UFUNCTION(BlueprintPure, Category = "GeometryScript|Shapes|Box", meta=(ScriptMethod))
-	static FBox GetExpandedBox(const FBox& Box, const FVector& ExpandBy);
+	static UPARAM(DisplayName="Expanded Box") FBox 
+	GetExpandedBox(const FBox& Box, const FVector& ExpandBy);
 
 	/**
 	 * Apply the input Transform to the corners of the input Box, and return the new Box containing those points
 	 */
 	UFUNCTION(BlueprintPure, Category = "GeometryScript|Shapes|Box", meta=(ScriptMethod))
-	static FBox GetTransformedBox(const FBox& Box, const FTransform& Transform);
+	static UPARAM(DisplayName="Transformed Box") FBox 
+	GetTransformedBox(const FBox& Box, const FTransform& Transform);
 
 	/**
 	 * Test if Box1 and Box2 intersect
 	 */
 	UFUNCTION(BlueprintPure, Category = "GeometryScript|Shapes|Box", meta=(ScriptMethod))
-	static bool TestBoxBoxIntersection(const FBox& Box1, const FBox& Box2);
+	static UPARAM(DisplayName="Intersects") bool 
+	TestBoxBoxIntersection(const FBox& Box1, const FBox& Box2);
 
 	/**
 	 * Find the Box formed by the intersection of Box1 and Box2
 	 * @param bIsIntersecting if the boxes do not intersect, this will be returned as false, otherwise true
 	 */
 	UFUNCTION(BlueprintPure, Category = "GeometryScript|Shapes|Box", meta=(ScriptMethod))
-	static FBox FindBoxBoxIntersection(const FBox& Box1, const FBox& Box2, bool& bIsIntersecting);
+	static UPARAM(DisplayName="Intersection Box") FBox 
+	FindBoxBoxIntersection(const FBox& Box1, const FBox& Box2, bool& bIsIntersecting);
 
 	/**
 	 * Calculate the minimum distance between Box1 and Box2
 	 */
 	UFUNCTION(BlueprintPure, Category = "GeometryScript|Shapes|Box", meta=(ScriptMethod))
-	static double GetBoxBoxDistance(const FBox& Box1, const FBox& Box2);
+	static UPARAM(DisplayName="Distance") double 
+	GetBoxBoxDistance(const FBox& Box1, const FBox& Box2);
 
 	/**
 	 * Test if a Point is inside the Box, returning true if so, otherwise false
 	 * @param bConsiderOnBoxAsInside if true, a point lying on the box face is considered "inside", otherwise it is considered "outside"
 	 */
 	UFUNCTION(BlueprintPure, Category = "GeometryScript|Shapes|Box", meta=(ScriptMethod))
-	static bool TestPointInsideBox(const FBox& Box, const FVector& Point, bool bConsiderOnBoxAsInside = true);
+	static UPARAM(DisplayName="Is Inside") bool 
+	TestPointInsideBox(const FBox& Box, const FVector& Point, bool bConsiderOnBoxAsInside = true);
 
 	/**
 	 * Find the point on the faces of the Box that is closest to the input Point.
@@ -175,19 +199,22 @@ public:
 	 * @param bIsInside if the Point is inside the Box, this will return as true, otherwise false
 	 */
 	UFUNCTION(BlueprintPure, Category = "GeometryScript|Shapes|Box", meta=(ScriptMethod))
-	static FVector FindClosestPointOnBox(const FBox& Box, const FVector& Point, bool& bIsInside);
+	static UPARAM(DisplayName="Closest Point") FVector 
+	FindClosestPointOnBox(const FBox& Box, const FVector& Point, bool& bIsInside);
 
 	/**
 	 * Calculate the minimum distance between the Box and the Point
 	 */
 	UFUNCTION(BlueprintPure, Category = "GeometryScript|Shapes|Box", meta=(ScriptMethod))
-	static double GetBoxPointDistance(const FBox& Box, const FVector& Point);
+	static UPARAM(DisplayName="Distance") double 
+	GetBoxPointDistance(const FBox& Box, const FVector& Point);
 
 	/**
 	 * Check if the Box intersects a Sphere defined by the SphereCenter and SphereRadius
 	 */
 	UFUNCTION(BlueprintPure, Category = "GeometryScript|Shapes|Box", meta=(ScriptMethod))
-	static bool TestBoxSphereIntersection(const FBox& Box, const FVector& SphereCenter, double SphereRadius);
+	static UPARAM(DisplayName="Intersects") bool 
+	TestBoxSphereIntersection(const FBox& Box, const FVector& SphereCenter, double SphereRadius);
 
 };
 
diff --git a/Engine/Source/Runtime/GeometryCore/Private/Sampling/MeshSurfacePointSampling.cpp b/Engine/Source/Runtime/GeometryCore/Private/Sampling/MeshSurfacePointSampling.cpp
--- a/Engine/Source/Runtime/GeometryCore/Private/Sampling/MeshSurfacePointSampling.cpp
+++ b/Engine/Source/Runtime/GeometryCore/Private/Sampling/MeshSurfacePointSampling.cpp
@@ -634,11 +634,14 @@ void FMeshSurfacePointSampling::ComputePoissonSampling(const FDynamicMesh3& Mesh
 
 	Samples.Reset();
 	Radii.Reset();
+	TriangleIDs.Reset();
+	BarycentricCoords.Reset();
 	auto AddSampleFunc = [&](FVector3d Position, int TriangleID, double Radius)
-		{
-			Samples.Add(FFrame3d(Position, Mesh.GetTriNormal(TriangleID)));
-			Radii.Add(Radius);
-		};
+	{
+		Samples.Add(FFrame3d(Position, Mesh.GetTriNormal(TriangleID)));
+		Radii.Add(Radius);
+		TriangleIDs.Add(TriangleID);
+	};
 
 	if (MaxSampleRadius > SampleRadius)
 	{
@@ -676,5 +679,17 @@ void FMeshSurfacePointSampling::ComputePoissonSampling(const FDynamicMesh3& Mesh
 			SampleRadius, MaxSamples, SubSampleDensity, RandomSeed, MaxSubSamplePoints, Progress );
 	}
 
+	if (bComputeBarycentrics)
+	{
+		FVector3d A,B,C;
+		int32 N = Samples.Num();
+		BarycentricCoords.SetNum(N);
+		for (int32 k = 0; k < N; ++k)
+		{
+			Mesh.GetTriVertices(TriangleIDs[k], A, B, C);
+			BarycentricCoords[k] = VectorUtil::BarycentricCoords( Samples[k].Origin, A, B, C);
+		}
+	}
+
 	Result.SetSuccess(true, Progress);
 }
\ No newline at end of file
diff --git a/Engine/Source/Runtime/GeometryCore/Public/Sampling/MeshSurfacePointSampling.h b/Engine/Source/Runtime/GeometryCore/Public/Sampling/MeshSurfacePointSampling.h
--- a/Engine/Source/Runtime/GeometryCore/Public/Sampling/MeshSurfacePointSampling.h
+++ b/Engine/Source/Runtime/GeometryCore/Public/Sampling/MeshSurfacePointSampling.h
@@ -90,6 +90,8 @@ public:
 	/** If true, weights are inverted */
 	bool bInvertWeights = false;
 
+	/** If true, barycentric coordinates output array will be populated */
+	bool bComputeBarycentrics = false;
 
 	//
 	// TODO: when MaxSamples is set, it would be useful to be able to use Weight to modulate
@@ -110,6 +112,11 @@ public:
 	/** Radius of each Sample Point, length is the same as Samples array */
 	TArray<double> Radii;
 
+	/** Triangle that contains each Sample Point, length is the same as Samples array */
+	TArray<int32> TriangleIDs;
+
+	/** Barycentric Coordinates of each Sample Point in it's respective Triangle. Only computed if bComputeBarycentrics = true */
+	TArray<FVector3d> BarycentricCoords;
 
 public:
 
