GeometryScript: Support sampling UTextureRenderTarget2D at UV positions

diff --git a/Engine/Plugins/Experimental/GeometryScripting/Source/GeometryScriptingCore/Private/TextureMapFunctions.cpp b/Engine/Plugins/Experimental/GeometryScripting/Source/GeometryScriptingCore/Private/TextureMapFunctions.cpp
--- a/Engine/Plugins/Experimental/GeometryScripting/Source/GeometryScriptingCore/Private/TextureMapFunctions.cpp
+++ b/Engine/Plugins/Experimental/GeometryScripting/Source/GeometryScriptingCore/Private/TextureMapFunctions.cpp
@@ -4,6 +4,8 @@
 
 #include "Async/ParallelFor.h"
 #include "AssetUtils/Texture2DUtil.h"
+#include "Engine/TextureRenderTarget2D.h"
+#include "TextureResource.h"
 #include "Spatial/SampledScalarField2.h"
 
 #include UE_INLINE_GENERATED_CPP_BY_NAME(TextureMapFunctions)
@@ -12,31 +14,22 @@ using namespace UE::Geometry;
 
 #define LOCTEXT_NAMESPACE "UGeometryScriptLibrary_TextureMapFunctions"
 
-void UGeometryScriptLibrary_TextureMapFunctions::SampleTexture2DAtUVPositions(
-	FGeometryScriptUVList UVList,
-	UTexture2D* TextureAsset,
-	FGeometryScriptSampleTextureOptions SampleOptions,
-	FGeometryScriptColorList& ColorList,
-	UGeometryScriptDebug* Debug)
+namespace UE::Geometry
+{
+namespace
+{
+
+void SampleImageBuilderAtUVPositions(
+	const TImageBuilder<FVector4f>& ImageData,
+	const FGeometryScriptUVList& UVList,
+	const FGeometryScriptSampleTextureOptions& SampleOptions,
+	FGeometryScriptColorList& ColorList)
 {
 	ColorList.Reset();
 	if (UVList.List.IsValid() == false || UVList.List->Num() == 0)
 	{
 		return;
 	}
-	
-	if (TextureAsset == nullptr)
-	{
-		UE::Geometry::AppendError(Debug, EGeometryScriptErrorType::InvalidInputs, LOCTEXT("SampleTexture2DAtUVPositions_InvalidInput2", "SampleTexture2DAtUVPositions: Texture is Null"));
-		return;
-	}
-
-	TImageBuilder<FVector4f> ImageData;
-	if (UE::AssetUtils::ReadTexture(TextureAsset, ImageData, false) == false)
-	{
-		UE::Geometry::AppendError(Debug, EGeometryScriptErrorType::OperationFailed, LOCTEXT("SampleTexture2DAtUVPositions_TexReadFailed", "SampleTexture2DAtUVPositions: Error reading source texture data. If using this function at Runtime, The Compression Settings type on the UTexture2D Asset must be set to VectorDisplacementmap (RGBA8)."));
-		return;
-	}
 
 	const TArray<FVector2D>& UVs = *UVList.List;
 	int32 NumUVs = UVs.Num();
@@ -67,8 +60,68 @@ void UGeometryScriptLibrary_TextureMapFunctions::SampleTexture2DAtUVPositions(
 
 		Colors[k] = (FLinearColor)InterpValue;
 	}
+}
+
+} // end anonymous namespace
+} // end UE::Geometry namespace
+
+void UGeometryScriptLibrary_TextureMapFunctions::SampleTexture2DAtUVPositions(
+	FGeometryScriptUVList UVList,
+	UTexture2D* TextureAsset,
+	FGeometryScriptSampleTextureOptions SampleOptions,
+	FGeometryScriptColorList& ColorList,
+	UGeometryScriptDebug* Debug)
+{
+	if (TextureAsset == nullptr)
+	{
+		UE::Geometry::AppendError(Debug, EGeometryScriptErrorType::InvalidInputs, LOCTEXT("SampleTexture2DAtUVPositions_InvalidInput2", "SampleTexture2DAtUVPositions: Texture is Null"));
+		return;
+	}
+
+	TImageBuilder<FVector4f> ImageData;
+	if (UE::AssetUtils::ReadTexture(TextureAsset, ImageData, false) == false)
+	{
+		UE::Geometry::AppendError(Debug, EGeometryScriptErrorType::OperationFailed, LOCTEXT("SampleTexture2DAtUVPositions_TexReadFailed", "SampleTexture2DAtUVPositions: Error reading source texture data. If using this function at Runtime, The Compression Settings type on the UTexture2D Asset must be set to VectorDisplacementmap (RGBA8)."));
+		return;
+	}
+
+	SampleImageBuilderAtUVPositions(ImageData, UVList, SampleOptions, ColorList);
+}
+
+
+void UGeometryScriptLibrary_TextureMapFunctions::SampleTextureRenderTarget2DAtUVPositions(  
+	FGeometryScriptUVList UVList, 
+	UTextureRenderTarget2D* Texture,
+	FGeometryScriptSampleTextureOptions SampleOptions,
+	FGeometryScriptColorList& ColorList,
+	UGeometryScriptDebug* Debug)
+{
+	if (Texture == nullptr)
+	{
+		UE::Geometry::AppendError(Debug, EGeometryScriptErrorType::InvalidInputs, LOCTEXT("SampleTextureRenderTarget2DAtUVPositions_InvalidInput", "SampleTextureRenderTarget2DAtUVPositions: Texture is Null"));
+		return;
+	}
+
+	FTextureRenderTargetResource* RenderTargetResource = Texture->GameThread_GetRenderTargetResource();
+	if (RenderTargetResource == nullptr)
+	{
+		UE::Geometry::AppendError(Debug, EGeometryScriptErrorType::OperationFailed, LOCTEXT("SampleTextureRenderTarget2DAtUVPositions_NullResource", "SampleTextureRenderTarget2DAtUVPositions: Render target resource is Null"));
+		return;
+	}
 
+	FReadSurfaceDataFlags ReadSurfaceDataFlags(RCM_MinMax);
+	ReadSurfaceDataFlags.SetLinearToGamma(false);
+	
+	TImageBuilder<FVector4f> ImageData;
+	ImageData.SetDimensions(FImageDimensions(Texture->SizeX, Texture->SizeY));
+	FLinearColor* Buffer = reinterpret_cast<FLinearColor*>(ImageData.GetImageBuffer().GetData());
+	if (RenderTargetResource->ReadLinearColorPixelsPtr(Buffer, ReadSurfaceDataFlags, FIntRect(0, 0, Texture->SizeX, Texture->SizeY)) == false)
+	{
+		UE::Geometry::AppendError(Debug, EGeometryScriptErrorType::OperationFailed, LOCTEXT("SampleTextureRenderTarget2DAtUVPositions_GpuReadFailed", "SampleTextureRenderTarget2DAtUVPositions: Error reading texture data from the GPU."));
+		return;
+	}
 
+	SampleImageBuilderAtUVPositions(ImageData, UVList, SampleOptions, ColorList);
 }
 
 
diff --git a/Engine/Plugins/Experimental/GeometryScripting/Source/GeometryScriptingCore/Public/GeometryScript/TextureMapFunctions.h b/Engine/Plugins/Experimental/GeometryScripting/Source/GeometryScriptingCore/Public/GeometryScript/TextureMapFunctions.h
--- a/Engine/Plugins/Experimental/GeometryScripting/Source/GeometryScriptingCore/Public/GeometryScript/TextureMapFunctions.h
+++ b/Engine/Plugins/Experimental/GeometryScripting/Source/GeometryScriptingCore/Public/GeometryScript/TextureMapFunctions.h
@@ -57,5 +57,18 @@ public:
 		FGeometryScriptColorList& ColorList,
 		UGeometryScriptDebug* Debug = nullptr);
 
+	/**
+	 * Sample the the given TextureMap at the list of UV positions and return the color at each position in ColorList output.
+	 * This function fetches GPU data before sampling so, depending on your application, it can be inefficient and slow!
+	 */
+	UFUNCTION(BlueprintCallable, Category = "GeometryScript|TextureUtils")
+	static void 
+	SampleTextureRenderTarget2DAtUVPositions(  
+		FGeometryScriptUVList UVList, 
+		UTextureRenderTarget2D* Texture,
+		FGeometryScriptSampleTextureOptions SampleOptions,
+		FGeometryScriptColorList& ColorList,
+		UGeometryScriptDebug* Debug = nullptr);
+
 
 };
\ No newline at end of file
diff --git a/Engine/Source/Runtime/GeometryCore/Public/Image/ImageBuilder.h b/Engine/Source/Runtime/GeometryCore/Public/Image/ImageBuilder.h
--- a/Engine/Source/Runtime/GeometryCore/Public/Image/ImageBuilder.h
+++ b/Engine/Source/Runtime/GeometryCore/Public/Image/ImageBuilder.h
@@ -118,6 +118,17 @@ public:
 		return Image[LinearIndex];
 	}
 
+	TArrayView64<PixelType> GetImageBuffer()
+	{
+		TArray64<PixelType>& Buffer = Image.GridValues();
+		return TArrayView64<PixelType>(Buffer.GetData(), Buffer.Num());
+	}
+
+	TConstArrayView64<PixelType> GetImageBuffer() const
+	{
+		const TArray64<PixelType>& Buffer = Image.GridValues();
+		return TConstArrayView64<PixelType>(Buffer.GetData(), Buffer.Num());
+	}
 
 	/**
 	 * Set the Pixel at the given X/Y coordinates to the given PixelType
diff --git a/Engine/Source/Runtime/GeometryCore/Public/Spatial/DenseGrid2.h b/Engine/Source/Runtime/GeometryCore/Public/Spatial/DenseGrid2.h
--- a/Engine/Source/Runtime/GeometryCore/Public/Spatial/DenseGrid2.h
+++ b/Engine/Source/Runtime/GeometryCore/Public/Spatial/DenseGrid2.h
@@ -129,6 +129,10 @@ public:
 		return Buffer[X + DimensionX * Y];
 	}
 
+	TArray64<ElemType>& GridValues()
+	{
+		return Buffer;
+	}
 	const TArray64<ElemType>& GridValues() const
 	{
 		return Buffer;
