Geometry Scripting: Asset creation and binding functionality for skeletal meshes

diff --git a/Engine/Plugins/Experimental/GeometryScripting/Source/GeometryScriptingCore/Private/MeshBoneWeightFunctions.cpp b/Engine/Plugins/Experimental/GeometryScripting/Source/GeometryScriptingCore/Private/MeshBoneWeightFunctions.cpp
--- a/Engine/Plugins/Experimental/GeometryScripting/Source/GeometryScriptingCore/Private/MeshBoneWeightFunctions.cpp
+++ b/Engine/Plugins/Experimental/GeometryScripting/Source/GeometryScriptingCore/Private/MeshBoneWeightFunctions.cpp
@@ -2,11 +2,13 @@
 
 #include "GeometryScript/MeshBoneWeightFunctions.h"
 
+#include "Animation/Skeleton.h"
+#include "BoneWeights.h"
 #include "DynamicMesh/DynamicMesh3.h"
 #include "DynamicMesh/DynamicMeshAttributeSet.h"
 #include "DynamicMesh/DynamicVertexSkinWeightsAttribute.h"
+#include "SkinningOps/SkinBindingOp.h"
 #include "UDynamicMesh.h"
-#include "BoneWeights.h"
 
 #include UE_INLINE_GENERATED_CPP_BY_NAME(MeshBoneWeightFunctions)
 
@@ -227,6 +229,83 @@ UDynamicMesh* UGeometryScriptLibrary_MeshBoneWeightFunctions::SetVertexBoneWeigh
 }
 
 
+UDynamicMesh* UGeometryScriptLibrary_MeshBoneWeightFunctions::SetAllVertexBoneWeights(
+	UDynamicMesh* TargetMesh,
+	const TArray<FGeometryScriptBoneWeight>& BoneWeights, 
+	FGeometryScriptBoneWeightProfile Profile
+	)
+{
+	bool bHasBoneWeightProfile = false;
+	const int32 Num = BoneWeights.Num();
+	TArray<FBoneWeight> NewWeightsList;
+	for (int32 k = 0; k < Num; ++k)
+	{
+		FBoneWeight NewWeight;
+		NewWeight.SetBoneIndex(BoneWeights[k].BoneIndex);
+		NewWeight.SetWeight(BoneWeights[k].Weight);
+		NewWeightsList.Add(NewWeight);
+	}
+	const FBoneWeights NewBoneWeights = FBoneWeights::Create(NewWeightsList);
+	
+	SimpleMeshBoneWeightEdit<bool>(TargetMesh, Profile, bHasBoneWeightProfile, false,
+		[&](const FDynamicMesh3& Mesh, FDynamicMeshVertexSkinWeightsAttribute& SkinWeights)
+	{
+		for (const int32 VertexID : Mesh.VertexIndicesItr())
+		{
+			SkinWeights.SetValue(VertexID, NewBoneWeights);
+		}
+		return true;
+	});
+
+	return TargetMesh;
+}
+
+
+UDynamicMesh* UGeometryScriptLibrary_MeshBoneWeightFunctions::ComputeSmoothBoneWeights(
+	UDynamicMesh* TargetMesh,
+	USkeleton* Skeleton, 
+	FGeometryScriptSmoothBoneWeightsOptions Options, 
+	FGeometryScriptBoneWeightProfile Profile,
+	UGeometryScriptDebug* Debug
+	)
+{
+	if (TargetMesh == nullptr)
+	{
+		AppendError(Debug, EGeometryScriptErrorType::InvalidInputs, LOCTEXT("ComputeSmoothBoneWeights_InvalidInput", "ComputeSmoothBoneWeights: TargetMesh is Null"));
+		return TargetMesh;
+	}
+	if (Skeleton == nullptr)
+	{
+		AppendError(Debug, EGeometryScriptErrorType::InvalidInputs, LOCTEXT("ComputeSmoothBoneWeights_InvalidSkeleton", "ComputeSmoothBoneWeights: Skeleton is Null"));
+		return TargetMesh;
+	}
+
+	TargetMesh->EditMesh([&](FDynamicMesh3& EditMesh)
+	{
+		FSkinBindingOp SkinBindingOp;
+		SkinBindingOp.OriginalMesh = MakeShared<FDynamicMesh3>(MoveTemp(EditMesh));
+		SkinBindingOp.SetTransformHierarchyFromReferenceSkeleton(Skeleton->GetReferenceSkeleton());
+		SkinBindingOp.ProfileName = Profile.ProfileName;
+		switch(Options.DistanceWeighingType)
+		{
+		case EGeometryScriptSmoothBoneWeightsType::DirectDistance:
+			SkinBindingOp.BindType = ESkinBindingType::DirectDistance;
+			break;
+		case EGeometryScriptSmoothBoneWeightsType::GeodesicVoxel:
+			SkinBindingOp.BindType = ESkinBindingType::GeodesicVoxel;
+			break;
+		}
+		SkinBindingOp.Stiffness = Options.Stiffness;
+		SkinBindingOp.MaxInfluences = Options.MaxInfluences;
+		SkinBindingOp.VoxelResolution = Options.VoxelResolution;
+
+		SkinBindingOp.CalculateResult(nullptr);
+
+		EditMesh = MoveTemp(*SkinBindingOp.ExtractResult().Release());
+	}, EDynamicMeshChangeType::AttributeEdit, EDynamicMeshAttributeChangeFlags::Unknown, false);
+
+	return TargetMesh;
+}
 
 
 #undef LOCTEXT_NAMESPACE
diff --git a/Engine/Plugins/Experimental/GeometryScripting/Source/GeometryScriptingCore/Public/GeometryScript/MeshBoneWeightFunctions.h b/Engine/Plugins/Experimental/GeometryScripting/Source/GeometryScriptingCore/Public/GeometryScript/MeshBoneWeightFunctions.h
--- a/Engine/Plugins/Experimental/GeometryScripting/Source/GeometryScriptingCore/Public/GeometryScript/MeshBoneWeightFunctions.h
+++ b/Engine/Plugins/Experimental/GeometryScripting/Source/GeometryScriptingCore/Public/GeometryScript/MeshBoneWeightFunctions.h
@@ -36,6 +36,45 @@ struct GEOMETRYSCRIPTINGCORE_API FGeometryScriptBoneWeightProfile
 };
 
 
+UENUM(BlueprintType)
+enum class EGeometryScriptSmoothBoneWeightsType : uint8
+{
+	DirectDistance = 0,		/** Compute weighting by using Euclidean distance from bone to vertex */
+	GeodesicVoxel = 1,		/** Compute weighting by using geodesic distance from bone to vertex */
+};
+
+
+USTRUCT(BlueprintType)
+struct GEOMETRYSCRIPTINGCORE_API FGeometryScriptSmoothBoneWeightsOptions
+{
+	GENERATED_BODY()
+
+	/** The type of algorithm to use for computing the bone weight for each vertex */
+	UPROPERTY(EditAnywhere, BlueprintReadWrite, Category = Options)
+	EGeometryScriptSmoothBoneWeightsType DistanceWeighingType = EGeometryScriptSmoothBoneWeightsType::DirectDistance;
+
+	/** How rigid the binding should be. Higher values result in a more rigid binding (greater influence by bones
+	 *  closer to the vertex than those further away).
+	 */
+	UPROPERTY(EditAnywhere, BlueprintReadWrite, Category = Options)
+	float Stiffness = 0.2f;
+
+	/** Maximum number of bones that contribute to each weight. Set to 1 for a completely rigid binding. Higher values
+	 *  to have more distant bones make additional contributions to the deformation at each vertex. 
+	 */
+	UPROPERTY(EditAnywhere, BlueprintReadWrite, Category = Options)
+	int32 MaxInfluences = 5;
+
+	/** The resolution to build the voxelized representation of the mesh, for computing geodesic distance. Higher values
+	 *  result in greater fidelity and less chance of disconnected parts contributing, but slower rate of computation and
+	 *  more memory usage.
+	 */
+	UPROPERTY(EditAnywhere, BlueprintReadWrite, Category = Options, meta=(EditCondition="DistanceWeighingType==EGeometryScriptSmoothBoneWeightsType::GeodesicVoxel"))
+	int32 VoxelResolution = 256;
+};
+
+
+
 
 UCLASS(meta = (ScriptName = "GeometryScript_BoneWeights"))
 class GEOMETRYSCRIPTINGCORE_API UGeometryScriptLibrary_MeshBoneWeightFunctions : public UBlueprintFunctionLibrary
@@ -117,5 +156,31 @@ public:
 		bool& bIsValidVertexID,
 		FGeometryScriptBoneWeightProfile Profile = FGeometryScriptBoneWeightProfile() );
 
+	/**
+	 * Set all vertices of the TargetMesh to the given Bone/Skin Weights
+	 * @param BoneWeights input array of bone index/weight pairs for the Vertex
+	 * @param Profile identifier for the bone/skin weight profile
+	 */
+	UFUNCTION(BlueprintCallable, Category = "GeometryScript|MeshQueries|BoneWeights", meta=(ScriptMethod))
+	static UPARAM(DisplayName = "Target Mesh") UDynamicMesh* 
+	SetAllVertexBoneWeights( 
+		UDynamicMesh* TargetMesh,
+		const TArray<FGeometryScriptBoneWeight>& BoneWeights,
+		FGeometryScriptBoneWeightProfile Profile = FGeometryScriptBoneWeightProfile() );
 
+	/** 
+	 *  Computes a smooth skin binding for the given mesh to the skeleton provided.
+	 *  @param Skeleton The skeleton to compute binding for the skin weights.
+	 *  @param Options The options to set for the binding algorithm.
+	 *  @param Profile The skin weight profile to update with the smooth binding.
+	 */
+	UFUNCTION(BlueprintCallable, Category = "GeometryScript|MeshQueries|BoneWeights", meta=(ScriptMethod))
+	static UPARAM(DisplayName = "Target Mesh") UDynamicMesh*
+	ComputeSmoothBoneWeights(
+		UDynamicMesh* TargetMesh, 
+		USkeleton* Skeleton,
+		FGeometryScriptSmoothBoneWeightsOptions Options,
+		FGeometryScriptBoneWeightProfile Profile = FGeometryScriptBoneWeightProfile(),
+		UGeometryScriptDebug* Debug = nullptr);
+	
 };
\ No newline at end of file
diff --git a/Engine/Plugins/Experimental/GeometryScripting/Source/GeometryScriptingEditor/Private/CreateNewAssetUtilityFunctions.cpp b/Engine/Plugins/Experimental/GeometryScripting/Source/GeometryScriptingEditor/Private/CreateNewAssetUtilityFunctions.cpp
--- a/Engine/Plugins/Experimental/GeometryScripting/Source/GeometryScriptingEditor/Private/CreateNewAssetUtilityFunctions.cpp
+++ b/Engine/Plugins/Experimental/GeometryScripting/Source/GeometryScriptingEditor/Private/CreateNewAssetUtilityFunctions.cpp
@@ -22,12 +22,12 @@
 #include "AssetRegistry/AssetRegistryModule.h"
 #include "Misc/Paths.h"
 
-#include "MeshDescriptionToDynamicMesh.h"
-#include "DynamicMeshToMeshDescription.h"
 #include "ConversionUtils/DynamicMeshToVolume.h"
 #include "AssetUtils/CreateStaticMeshUtil.h"
+#include "AssetUtils/CreateSkeletalMeshUtil.h"
 #include "AssetUtils/CreateTexture2DUtil.h"
 #include "ModelingObjectsCreationAPI.h"
+#include "Engine/SkinnedAssetCommon.h"
 
 #include UE_INLINE_GENERATED_CPP_BY_NAME(CreateNewAssetUtilityFunctions)
 
@@ -213,6 +213,87 @@ UStaticMesh* UGeometryScriptLibrary_CreateNewAssetFunctions::CreateNewStaticMesh
 }
 
 
+USkeletalMesh* UGeometryScriptLibrary_CreateNewAssetFunctions::CreateNewSkeletalMeshAssetFromMesh(
+	UDynamicMesh* FromDynamicMesh,
+	USkeleton* InSkeleton,
+	FString AssetPathAndName, 
+	FGeometryScriptCreateNewSkeletalMeshAssetOptions Options,
+	TEnumAsByte<EGeometryScriptOutcomePins>& Outcome, 
+	UGeometryScriptDebug* Debug)
+{
+	using namespace UE::AssetUtils;
+
+	Outcome = EGeometryScriptOutcomePins::Failure;
+	if (FromDynamicMesh == nullptr)
+	{
+		AppendError(Debug, EGeometryScriptErrorType::InvalidInputs, LOCTEXT("CreateNewSkeletalMeshAssetFromMesh_InvalidInput1", "CreateNewSkeletalMeshAssetFromMesh: FromDynamicMesh is Null"));
+		return nullptr;
+	}
+	if (FromDynamicMesh->GetTriangleCount() == 0)
+	{
+		AppendError(Debug, EGeometryScriptErrorType::InvalidInputs, LOCTEXT("CreateNewSkeletalMeshAssetFromMesh_InvalidInput2", "CreateNewSkeletalMeshAssetFromMesh: FromDynamicMesh has zero triangles"));
+		return nullptr;
+	}
+	if (FromDynamicMesh->GetMeshRef().Attributes()->GetSkinWeightsAttributes().Num() == 0)
+	{
+		AppendError(Debug, EGeometryScriptErrorType::InvalidInputs, LOCTEXT("CreateNewSkeletalMeshAssetFromMesh_InvalidInput3", "CreateNewSkeletalMeshAssetFromMesh: FromDynamicMesh has no skin weight attributes"));
+		return nullptr;
+	}
+	
+	// todo: other safety checks
+
+	GEditor->BeginTransaction(LOCTEXT("CreateNewSkeletalMeshAssetFromMesh_Transaction", "Create SkeletalMesh"));
+
+	FSkeletalMeshAssetOptions AssetOptions;
+	AssetOptions.NewAssetPath = AssetPathAndName;
+
+	AssetOptions.NumSourceModels = 1;
+
+	if (!Options.Materials.IsEmpty())
+	{
+		TArray<FSkeletalMaterial> Materials;
+
+		for (const TPair<FName, TObjectPtr<UMaterialInterface>>& Item : Options.Materials)
+		{
+			Materials.Add(FSkeletalMaterial{Item.Value, Item.Key});
+		}
+		AssetOptions.SkeletalMaterials = MoveTemp(Materials);
+		AssetOptions.NumMaterialSlots = AssetOptions.SkeletalMaterials.Num(); 
+	}
+	else
+	{
+		AssetOptions.NumMaterialSlots = 1;
+	}
+	
+	AssetOptions.bEnableRecomputeNormals = Options.bEnableRecomputeNormals;
+	AssetOptions.bEnableRecomputeTangents = Options.bEnableRecomputeTangents;
+
+	FDynamicMesh3 LODMesh;
+	FromDynamicMesh->ProcessMesh([&](const FDynamicMesh3& ReadMesh)
+	{
+		LODMesh = ReadMesh;
+	});
+	AssetOptions.SourceMeshes.DynamicMeshes.Add(&LODMesh);
+
+	FSkeletalMeshResults ResultData;
+	const ECreateSkeletalMeshResult AssetResult = CreateSkeletalMeshAsset(AssetOptions, ResultData);
+
+	if (AssetResult != ECreateSkeletalMeshResult::Ok)
+	{
+		AppendError(Debug, EGeometryScriptErrorType::OperationFailed, LOCTEXT("CreateNewSkeletalMeshAssetFromMesh_Failed", "CreateNewSkeletalMeshAssetFromMesh: Failed to create new Asset"));
+		return nullptr;
+	}
+
+	GEditor->EndTransaction();
+
+	// publish new asset so that asset editor updates
+	FAssetRegistryModule::AssetCreated(ResultData.SkeletalMesh);
+
+	Outcome = EGeometryScriptOutcomePins::Success;
+	return ResultData.SkeletalMesh;
+}
+
+
 UTexture2D* UGeometryScriptLibrary_CreateNewAssetFunctions::CreateNewTexture2DAsset(
 		UTexture2D* FromTexture, 
 		FString AssetPathAndName,
diff --git a/Engine/Plugins/Experimental/GeometryScripting/Source/GeometryScriptingEditor/Public/GeometryScript/CreateNewAssetUtilityFunctions.h b/Engine/Plugins/Experimental/GeometryScripting/Source/GeometryScriptingEditor/Public/GeometryScript/CreateNewAssetUtilityFunctions.h
--- a/Engine/Plugins/Experimental/GeometryScripting/Source/GeometryScriptingEditor/Public/GeometryScript/CreateNewAssetUtilityFunctions.h
+++ b/Engine/Plugins/Experimental/GeometryScripting/Source/GeometryScriptingEditor/Public/GeometryScript/CreateNewAssetUtilityFunctions.h
@@ -9,6 +9,8 @@
 #include "BodySetupEnums.h"
 #include "CreateNewAssetUtilityFunctions.generated.h"
 
+
+class UMaterialInterface;
 class UStaticMesh;
 class UDynamicMesh;
 class AVolume;
@@ -70,6 +72,22 @@ public:
 };
 
 
+USTRUCT(BlueprintType)
+struct GEOMETRYSCRIPTINGEDITOR_API FGeometryScriptCreateNewSkeletalMeshAssetOptions
+{
+	GENERATED_BODY()
+public:
+	UPROPERTY(BlueprintReadWrite, Category = Options)
+	bool bEnableRecomputeNormals = false;
+
+	UPROPERTY(BlueprintReadWrite, Category = Options)
+	bool bEnableRecomputeTangents = false;
+
+	UPROPERTY(BlueprintReadWrite, Category = Options)
+	TMap<FName, TObjectPtr<UMaterialInterface>> Materials;
+};
+
+
 USTRUCT(BlueprintType)
 struct GEOMETRYSCRIPTINGEDITOR_API FGeometryScriptCreateNewTexture2DAssetOptions
 {
@@ -118,6 +136,16 @@ public:
 		FGeometryScriptCreateNewStaticMeshAssetOptions Options,
 		TEnumAsByte<EGeometryScriptOutcomePins>& Outcome,
 		UGeometryScriptDebug* Debug = nullptr);
+	
+	UFUNCTION(BlueprintCallable, Category = "GeometryScript|AssetManagement", meta = (ExpandEnumAsExecs = "Outcome"))
+	static UPARAM(DisplayName = "Skeletal Mesh Asset") USkeletalMesh* 
+	CreateNewSkeletalMeshAssetFromMesh(
+		UDynamicMesh* FromDynamicMesh, 
+		USkeleton* InSkeleton,
+		FString AssetPathAndName,
+		FGeometryScriptCreateNewSkeletalMeshAssetOptions Options,
+		TEnumAsByte<EGeometryScriptOutcomePins>& Outcome,
+		UGeometryScriptDebug* Debug = nullptr);
 
 	UFUNCTION(BlueprintCallable, Category = "GeometryScript|AssetManagement", meta = (ExpandEnumAsExecs = "Outcome"))
 	static UPARAM(DisplayName = "Texture 2D Asset") UTexture2D* 
diff --git a/Engine/Plugins/Runtime/MeshModelingToolset/Source/ModelingComponentsEditorOnly/ModelingComponentsEditorOnly.Build.cs b/Engine/Plugins/Runtime/MeshModelingToolset/Source/ModelingComponentsEditorOnly/ModelingComponentsEditorOnly.Build.cs
--- a/Engine/Plugins/Runtime/MeshModelingToolset/Source/ModelingComponentsEditorOnly/ModelingComponentsEditorOnly.Build.cs
+++ b/Engine/Plugins/Runtime/MeshModelingToolset/Source/ModelingComponentsEditorOnly/ModelingComponentsEditorOnly.Build.cs
@@ -38,7 +38,8 @@ public class ModelingComponentsEditorOnly : ModuleRules
 				"Slate",
                 "RHI",
 				"AssetTools",
-				"UnrealEd"			// required for asset factories
+				"UnrealEd", 
+				"SkeletalMeshUtilitiesCommon" // required for asset factories
 			}
 			);
 
diff --git a/Engine/Plugins/Runtime/MeshModelingToolset/Source/ModelingComponentsEditorOnly/Private/AssetUtils/CreateSkeletalMeshUtil.cpp b/Engine/Plugins/Runtime/MeshModelingToolset/Source/ModelingComponentsEditorOnly/Private/AssetUtils/CreateSkeletalMeshUtil.cpp
new file mode 100644
--- /dev/null
+++ b/Engine/Plugins/Runtime/MeshModelingToolset/Source/ModelingComponentsEditorOnly/Private/AssetUtils/CreateSkeletalMeshUtil.cpp
@@ -0,0 +1,111 @@
+﻿// Copyright Epic Games, Inc. All Rights Reserved.
+
+#include "AssetUtils/CreateSkeletalMeshUtil.h"
+
+#include "Animation/Skeleton.h"
+#include "DynamicMeshToMeshDescription.h"
+#include "Engine/SkeletalMesh.h"
+#include "Engine/SkinnedAssetCommon.h"
+#include "Misc/PackageName.h"
+#include "StaticToSkeletalMeshConverter.h"
+#include "UObject/Package.h"
+#include "UObject/UObjectGlobals.h"
+
+
+UE::AssetUtils::ECreateSkeletalMeshResult UE::AssetUtils::CreateSkeletalMeshAsset(
+	const FSkeletalMeshAssetOptions& Options,
+	FSkeletalMeshResults& ResultsOut
+	)
+{
+	const FString NewObjectName = FPackageName::GetLongPackageAssetName(Options.NewAssetPath);
+
+	UPackage* UsePackage;
+	if (Options.UsePackage != nullptr)
+	{
+		UsePackage = Options.UsePackage;
+	}
+	else
+	{
+		UsePackage = CreatePackage(*Options.NewAssetPath);
+	}
+	if (ensure(UsePackage != nullptr) == false)
+	{
+		return ECreateSkeletalMeshResult::InvalidPackage;
+	}
+
+	constexpr EObjectFlags UseFlags = RF_Public | RF_Standalone;
+	USkeletalMesh* NewSkeletalMesh = NewObject<USkeletalMesh>(UsePackage, FName(*NewObjectName), UseFlags);
+	if (ensure(NewSkeletalMesh != nullptr) == false)
+	{
+		return ECreateSkeletalMeshResult::UnknownError;
+	}
+
+	if (!ensure(Options.Skeleton))
+	{
+		return ECreateSkeletalMeshResult::InvalidSkeleton;
+	}
+
+	const int32 UseNumSourceModels = FMath::Max(1, Options.NumSourceModels);
+	
+	TArray<const FMeshDescription*> MeshDescriptions;
+	TArray<FMeshDescription> ConstructedMeshDescriptions;
+	if (!Options.SourceMeshes.MoveMeshDescriptions.IsEmpty())
+	{
+		if (ensure(Options.SourceMeshes.MoveMeshDescriptions.Num() == UseNumSourceModels))
+		{
+			MeshDescriptions.Append(Options.SourceMeshes.MoveMeshDescriptions);
+		}
+	}
+	else if (!Options.SourceMeshes.MeshDescriptions.IsEmpty())
+	{
+		if (ensure(Options.SourceMeshes.MeshDescriptions.Num() == UseNumSourceModels))
+		{
+			MeshDescriptions.Append(Options.SourceMeshes.MeshDescriptions);
+		}
+	}
+	else if (!Options.SourceMeshes.DynamicMeshes.IsEmpty())
+	{
+		if (ensure(Options.SourceMeshes.DynamicMeshes.Num() == UseNumSourceModels))
+		{
+			for (const FDynamicMesh3* DynamicMesh : Options.SourceMeshes.DynamicMeshes)
+			{
+				ConstructedMeshDescriptions.AddDefaulted();
+				FDynamicMeshToMeshDescription Converter;
+				Converter.Convert(DynamicMesh, ConstructedMeshDescriptions.Last(), !Options.bEnableRecomputeTangents);
+				MeshDescriptions.Add(&ConstructedMeshDescriptions.Last());
+			}
+		}
+	}
+
+	TArray<FSkeletalMaterial> Materials;
+	TConstArrayView<FSkeletalMaterial> MaterialView;
+	if (!Options.SkeletalMaterials.IsEmpty())
+	{
+		MaterialView = Options.SkeletalMaterials;
+	}
+	else if (!Options.AssetMaterials.IsEmpty())
+	{
+		for (UMaterialInterface* MaterialInterface : Options.AssetMaterials)
+		{
+			Materials.Add(FSkeletalMaterial(MaterialInterface));
+		}
+	}
+	
+	if (!FStaticToSkeletalMeshConverter::InitializeSkeletalMeshFromMeshDescriptions(
+		NewSkeletalMesh, MeshDescriptions, MaterialView, Options.Skeleton->GetReferenceSkeleton(),
+		Options.bEnableRecomputeNormals, Options.bEnableRecomputeTangents))
+	{
+		return ECreateSkeletalMeshResult::UnknownError;
+	}
+
+	// Update the skeletal mesh and the skeleton so that their ref skeletons are in sync and the skeleton's preview mesh
+	// is the one we just created.
+	NewSkeletalMesh->SetSkeleton(Options.Skeleton);
+	Options.Skeleton->MergeAllBonesToBoneTree(NewSkeletalMesh);
+	if (!Options.Skeleton->GetPreviewMesh())
+	{
+		Options.Skeleton->SetPreviewMesh(NewSkeletalMesh);
+	}
+	
+	return ECreateSkeletalMeshResult::Ok;
+}
diff --git a/Engine/Plugins/Runtime/MeshModelingToolset/Source/ModelingComponentsEditorOnly/Public/AssetUtils/CreateSkeletalMeshUtil.h b/Engine/Plugins/Runtime/MeshModelingToolset/Source/ModelingComponentsEditorOnly/Public/AssetUtils/CreateSkeletalMeshUtil.h
new file mode 100644
--- /dev/null
+++ b/Engine/Plugins/Runtime/MeshModelingToolset/Source/ModelingComponentsEditorOnly/Public/AssetUtils/CreateSkeletalMeshUtil.h
@@ -0,0 +1,96 @@
+﻿// Copyright Epic Games, Inc. All Rights Reserved.
+
+#pragma once
+
+#include "DynamicMesh/DynamicMesh3.h"
+
+
+class UMaterialInterface;
+class USkeletalMesh;
+class USkeleton;
+struct FMeshDescription;
+struct FSkeletalMaterial;
+
+
+namespace UE::AssetUtils
+{
+	using namespace Geometry;
+
+	/**
+	 * Result enum returned by Create functions below to indicate success/error conditions
+	 */
+	enum class ECreateSkeletalMeshResult
+	{
+		Ok = 0,
+		InvalidPackage = 1,
+		InvalidSkeleton = 2,
+
+		UnknownError = 100
+	};
+
+	/**
+	 * Set of input meshes for StaticMesh Create() functions below.
+	 * Only one of the arrays should be initialized, and it's size should be equal to 
+	 * the number of LODs specified for the Asset.
+	 */
+	struct MODELINGCOMPONENTSEDITORONLY_API FSkeletalMeshAssetMeshes
+	{
+		TArray<const FDynamicMesh3*> DynamicMeshes;
+		TArray<const FMeshDescription*> MeshDescriptions;
+		TArray<FMeshDescription*> MoveMeshDescriptions;		// these will be MoveTemp()'d into the SkeletalMesh, so if initialized, this array will be invalid after Create call()
+	};
+
+	/**
+	 * Options for new USkeletalMesh asset created by Create() functions below.
+	 */
+	struct MODELINGCOMPONENTSEDITORONLY_API FSkeletalMeshAssetOptions
+	{
+		// This package will be used if it is not nullptr, otherwise a new package will be created at NewAssetPath
+		UPackage* UsePackage = nullptr;
+		FString NewAssetPath;
+		
+		// The associated USkeleton object
+		USkeleton* Skeleton = nullptr;
+
+		// Number of SourceModels (ie LODs)
+		int32 NumSourceModels = 1;
+		// Number of Material Slots desired on the Asset. At least one will always be created.
+		int32 NumMaterialSlots = 1;
+
+		// Controls whether or not the RecomputeNormals option will be enabled on the Asset
+		bool bEnableRecomputeNormals = false;
+		// Controls whether or not the RecomputeTangents option will be enabled on the Asset
+		bool bEnableRecomputeTangents = true;
+
+		// List of skeletal materials. These will be used over the AssetMaterials, if provided.
+		TArray<FSkeletalMaterial> SkeletalMaterials;
+		
+		// Optional list of materials to initialize the Asset with. If defined, size must be the same as NumMaterialSlots
+		TArray<UMaterialInterface*> AssetMaterials;
+
+		// Optional list of meshes to use to initialize SourceModels. If defined, only one array should be non-empty, and must be the same length as NumSourceModels.
+		FSkeletalMeshAssetMeshes SourceMeshes;
+	};
+
+	/**
+	 * Output information about a newly-created StaticMesh, returned by Create functions below.
+	 * Some fields may be null, if the relevant function did not create that type of object
+	 */
+	struct MODELINGCOMPONENTSEDITORONLY_API FSkeletalMeshResults
+	{
+		/** USkeletalMesh asset that was created */
+		USkeletalMesh* SkeletalMesh = nullptr;
+	};
+
+
+	/**
+	 * Create a new UStaticMesh Asset based on the input Options
+	 * @param Options defines configuration for the new UStaticMesh Asset
+	 * @param ResultsOut new StaticMesh is returned here
+	 * @return Ok, or error flag if some part of the creation process failed. On failure, no Asset is created.
+	 */
+	MODELINGCOMPONENTSEDITORONLY_API ECreateSkeletalMeshResult CreateSkeletalMeshAsset(
+		const FSkeletalMeshAssetOptions& Options,
+		FSkeletalMeshResults& ResultsOut);
+
+}
\ No newline at end of file
