GeometryScript: change several enum types used as OutcomePins on functions to be 'enum class : uint8' instead of just plain enum, and update functions that were using those types as outcome pins

diff --git a/Engine/Plugins/Experimental/GeometryScripting/Source/GeometryScriptingCore/Private/MeshAssetFunctions.cpp b/Engine/Plugins/Experimental/GeometryScripting/Source/GeometryScriptingCore/Private/MeshAssetFunctions.cpp
--- a/Engine/Plugins/Experimental/GeometryScripting/Source/GeometryScriptingCore/Private/MeshAssetFunctions.cpp
+++ b/Engine/Plugins/Experimental/GeometryScripting/Source/GeometryScriptingCore/Private/MeshAssetFunctions.cpp
@@ -42,7 +42,7 @@ static UDynamicMesh* CopyMeshFromStaticMesh_SourceData(
 	UDynamicMesh* ToDynamicMesh, 
 	FGeometryScriptCopyMeshFromAssetOptions AssetOptions,
 	FGeometryScriptMeshReadLOD RequestedLOD,
-	TEnumAsByte<EGeometryScriptOutcomePins>& Outcome,
+	EGeometryScriptOutcomePins& Outcome,
 	UGeometryScriptDebug* Debug
 )
 {
@@ -119,7 +119,7 @@ static UDynamicMesh* CopyMeshFromStaticMesh_RenderData(
 	UDynamicMesh* ToDynamicMesh, 
 	FGeometryScriptCopyMeshFromAssetOptions AssetOptions,
 	FGeometryScriptMeshReadLOD RequestedLOD,
-	TEnumAsByte<EGeometryScriptOutcomePins>& Outcome,
+	EGeometryScriptOutcomePins& Outcome,
 	UGeometryScriptDebug* Debug
 )
 {
@@ -173,7 +173,7 @@ UDynamicMesh*  UGeometryScriptLibrary_StaticMeshFunctions::CopyMeshFromStaticMes
 	UDynamicMesh* ToDynamicMesh, 
 	FGeometryScriptCopyMeshFromAssetOptions AssetOptions,
 	FGeometryScriptMeshReadLOD RequestedLOD,
-	TEnumAsByte<EGeometryScriptOutcomePins>& Outcome,
+	EGeometryScriptOutcomePins& Outcome,
 	UGeometryScriptDebug* Debug)
 {
 	Outcome = EGeometryScriptOutcomePins::Failure;
@@ -214,7 +214,7 @@ UDynamicMesh*  UGeometryScriptLibrary_StaticMeshFunctions::CopyMeshToStaticMesh(
 	UStaticMesh* ToStaticMeshAsset,
 	FGeometryScriptCopyMeshToAssetOptions Options,
 	FGeometryScriptMeshWriteLOD TargetLOD,
-	TEnumAsByte<EGeometryScriptOutcomePins>& Outcome,
+	EGeometryScriptOutcomePins& Outcome,
 	UGeometryScriptDebug* Debug)
 {
 	Outcome = EGeometryScriptOutcomePins::Failure;
@@ -393,7 +393,7 @@ void UGeometryScriptLibrary_StaticMeshFunctions::GetSectionMaterialListFromStati
 	FGeometryScriptMeshReadLOD RequestedLOD,
 	TArray<UMaterialInterface*>& MaterialList,
 	TArray<int32>& MaterialIndex,
-	TEnumAsByte<EGeometryScriptOutcomePins>& Outcome,
+	EGeometryScriptOutcomePins& Outcome,
 	UGeometryScriptDebug* Debug)
 {
 	Outcome = EGeometryScriptOutcomePins::Failure;
@@ -428,7 +428,7 @@ UDynamicMesh* UGeometryScriptLibrary_StaticMeshFunctions::CopyMeshFromSkeletalMe
 		UDynamicMesh* ToDynamicMesh,
 		FGeometryScriptCopyMeshFromAssetOptions AssetOptions,
 		FGeometryScriptMeshReadLOD RequestedLOD,
-		TEnumAsByte<EGeometryScriptOutcomePins>& Outcome,
+		EGeometryScriptOutcomePins& Outcome,
 		UGeometryScriptDebug* Debug)
 {
 	Outcome = EGeometryScriptOutcomePins::Failure;
@@ -494,7 +494,7 @@ UDynamicMesh* UGeometryScriptLibrary_StaticMeshFunctions::CopyMeshToSkeletalMesh
 		USkeletalMesh* ToSkeletalMeshAsset,
 		FGeometryScriptCopyMeshToAssetOptions Options,
 		FGeometryScriptMeshWriteLOD TargetLOD,
-		TEnumAsByte<EGeometryScriptOutcomePins>& Outcome,
+		EGeometryScriptOutcomePins& Outcome,
 		UGeometryScriptDebug* Debug)
 {
 	Outcome = EGeometryScriptOutcomePins::Failure;
diff --git a/Engine/Plugins/Experimental/GeometryScripting/Source/GeometryScriptingCore/Private/MeshSpatialFunctions.cpp b/Engine/Plugins/Experimental/GeometryScripting/Source/GeometryScriptingCore/Private/MeshSpatialFunctions.cpp
--- a/Engine/Plugins/Experimental/GeometryScripting/Source/GeometryScriptingCore/Private/MeshSpatialFunctions.cpp
+++ b/Engine/Plugins/Experimental/GeometryScripting/Source/GeometryScriptingCore/Private/MeshSpatialFunctions.cpp
@@ -115,7 +115,7 @@ UDynamicMesh* UGeometryScriptLibrary_MeshSpatial::FindNearestPointOnMesh(
 	FVector QueryPoint,
 	FGeometryScriptSpatialQueryOptions Options,
 	FGeometryScriptTrianglePoint& NearestResult,
-	TEnumAsByte<EGeometryScriptSearchOutcomePins>& Outcome,
+	EGeometryScriptSearchOutcomePins& Outcome,
 	UGeometryScriptDebug* Debug)
 {
 	Outcome = EGeometryScriptSearchOutcomePins::NotFound;
@@ -172,7 +172,7 @@ UDynamicMesh* UGeometryScriptLibrary_MeshSpatial::FindNearestRayIntersectionWith
 	FVector RayDirection,
 	FGeometryScriptSpatialQueryOptions Options,
 	FGeometryScriptRayHitResult& HitResult,
-	TEnumAsByte<EGeometryScriptSearchOutcomePins>& Outcome,
+	EGeometryScriptSearchOutcomePins& Outcome,
 	UGeometryScriptDebug* Debug)
 {
 	Outcome = EGeometryScriptSearchOutcomePins::NotFound;
@@ -232,7 +232,7 @@ UDynamicMesh* UGeometryScriptLibrary_MeshSpatial::IsPointInsideMesh(
 	FVector QueryPoint,
 	FGeometryScriptSpatialQueryOptions Options,
 	bool& bIsInside,
-	TEnumAsByte<EGeometryScriptContainmentOutcomePins>& Outcome,
+	EGeometryScriptContainmentOutcomePins& Outcome,
 	UGeometryScriptDebug* Debug)
 {
 	Outcome = EGeometryScriptContainmentOutcomePins::Outside;
diff --git a/Engine/Plugins/Experimental/GeometryScripting/Source/GeometryScriptingCore/Private/SceneUtilityFunctions.cpp b/Engine/Plugins/Experimental/GeometryScripting/Source/GeometryScriptingCore/Private/SceneUtilityFunctions.cpp
--- a/Engine/Plugins/Experimental/GeometryScripting/Source/GeometryScriptingCore/Private/SceneUtilityFunctions.cpp
+++ b/Engine/Plugins/Experimental/GeometryScripting/Source/GeometryScriptingCore/Private/SceneUtilityFunctions.cpp
@@ -37,7 +37,7 @@ UDynamicMesh* UGeometryScriptLibrary_SceneUtilityFunctions::CopyMeshFromComponen
 	FGeometryScriptCopyMeshFromComponentOptions Options,
 	bool bTransformToWorld,
 	FTransform& LocalToWorld,
-	TEnumAsByte<EGeometryScriptOutcomePins>& Outcome,
+	EGeometryScriptOutcomePins& Outcome,
 	UGeometryScriptDebug* Debug)
 {
 	Outcome = EGeometryScriptOutcomePins::Failure;
@@ -177,7 +177,7 @@ UDynamicMesh* UGeometryScriptLibrary_SceneUtilityFunctions::CopyCollisionMeshesF
 	UDynamicMesh* ToDynamicMesh,
 	bool bTransformToWorld,
 	FTransform& LocalToWorld,
-	TEnumAsByte<EGeometryScriptOutcomePins>& Outcome,
+	EGeometryScriptOutcomePins& Outcome,
 	bool bUseComplexCollision,
 	int SphereResolution,
 	UGeometryScriptDebug* Debug)
diff --git a/Engine/Plugins/Experimental/GeometryScripting/Source/GeometryScriptingCore/Public/GeometryScript/GeometryScriptTypes.h b/Engine/Plugins/Experimental/GeometryScripting/Source/GeometryScriptingCore/Public/GeometryScript/GeometryScriptTypes.h
--- a/Engine/Plugins/Experimental/GeometryScripting/Source/GeometryScriptingCore/Public/GeometryScript/GeometryScriptTypes.h
+++ b/Engine/Plugins/Experimental/GeometryScripting/Source/GeometryScriptingCore/Public/GeometryScript/GeometryScriptTypes.h
@@ -13,21 +13,21 @@ PREDECLARE_GEOMETRY(typedef TMeshAABBTree3<FDynamicMesh3> FDynamicMeshAABBTree3)
 
 
 UENUM(BlueprintType)
-enum EGeometryScriptOutcomePins : int
+enum class EGeometryScriptOutcomePins : uint8
 {
 	Failure,
 	Success
 };
 
 UENUM(BlueprintType)
-enum EGeometryScriptSearchOutcomePins : int
+enum class EGeometryScriptSearchOutcomePins : uint8
 {
 	Found,
 	NotFound
 };
 
 UENUM(BlueprintType)
-enum EGeometryScriptContainmentOutcomePins : int
+enum class EGeometryScriptContainmentOutcomePins : uint8
 {
 	Inside, 
 	Outside
diff --git a/Engine/Plugins/Experimental/GeometryScripting/Source/GeometryScriptingCore/Public/GeometryScript/MeshAssetFunctions.h b/Engine/Plugins/Experimental/GeometryScripting/Source/GeometryScriptingCore/Public/GeometryScript/MeshAssetFunctions.h
--- a/Engine/Plugins/Experimental/GeometryScripting/Source/GeometryScriptingCore/Public/GeometryScript/MeshAssetFunctions.h
+++ b/Engine/Plugins/Experimental/GeometryScripting/Source/GeometryScriptingCore/Public/GeometryScript/MeshAssetFunctions.h
@@ -112,7 +112,7 @@ public:
 		UDynamicMesh* ToDynamicMesh, 
 		FGeometryScriptCopyMeshFromAssetOptions AssetOptions,
 		FGeometryScriptMeshReadLOD RequestedLOD,
-		TEnumAsByte<EGeometryScriptOutcomePins>& Outcome,
+		EGeometryScriptOutcomePins& Outcome,
 		UGeometryScriptDebug* Debug = nullptr);
 
 
@@ -123,7 +123,7 @@ public:
 		UStaticMesh* ToStaticMeshAsset, 
 		FGeometryScriptCopyMeshToAssetOptions Options,
 		FGeometryScriptMeshWriteLOD TargetLOD,
-		TEnumAsByte<EGeometryScriptOutcomePins>& Outcome,
+		EGeometryScriptOutcomePins& Outcome,
 		UGeometryScriptDebug* Debug = nullptr);
 
 
@@ -135,7 +135,7 @@ public:
 		FGeometryScriptMeshReadLOD RequestedLOD,
 		TArray<UMaterialInterface*>& MaterialList,
 		TArray<int32>& MaterialIndex,
-		TEnumAsByte<EGeometryScriptOutcomePins>& Outcome,
+		EGeometryScriptOutcomePins& Outcome,
 		UGeometryScriptDebug* Debug = nullptr);
 
 
@@ -146,7 +146,7 @@ public:
 		UDynamicMesh* ToDynamicMesh,
 		FGeometryScriptCopyMeshFromAssetOptions AssetOptions,
 		FGeometryScriptMeshReadLOD RequestedLOD,
-		TEnumAsByte<EGeometryScriptOutcomePins>& Outcome,
+		EGeometryScriptOutcomePins& Outcome,
 		UGeometryScriptDebug* Debug = nullptr);
 
 
@@ -157,7 +157,7 @@ public:
 		USkeletalMesh* ToSkeletalMeshAsset,
 		FGeometryScriptCopyMeshToAssetOptions Options,
 		FGeometryScriptMeshWriteLOD TargetLOD,
-		TEnumAsByte<EGeometryScriptOutcomePins>& Outcome,
+		EGeometryScriptOutcomePins& Outcome,
 		UGeometryScriptDebug* Debug = nullptr);
 };
 
diff --git a/Engine/Plugins/Experimental/GeometryScripting/Source/GeometryScriptingCore/Public/GeometryScript/MeshSpatialFunctions.h b/Engine/Plugins/Experimental/GeometryScripting/Source/GeometryScriptingCore/Public/GeometryScript/MeshSpatialFunctions.h
--- a/Engine/Plugins/Experimental/GeometryScripting/Source/GeometryScriptingCore/Public/GeometryScript/MeshSpatialFunctions.h
+++ b/Engine/Plugins/Experimental/GeometryScripting/Source/GeometryScriptingCore/Public/GeometryScript/MeshSpatialFunctions.h
@@ -91,7 +91,7 @@ public:
 		FVector QueryPoint,
 		FGeometryScriptSpatialQueryOptions Options,
 		FGeometryScriptTrianglePoint& NearestResult,
-		TEnumAsByte<EGeometryScriptSearchOutcomePins>& Outcome,
+		EGeometryScriptSearchOutcomePins& Outcome,
 		UGeometryScriptDebug* Debug = nullptr );
 
 
@@ -104,7 +104,7 @@ public:
 		FVector RayDirection,
 		FGeometryScriptSpatialQueryOptions Options,
 		FGeometryScriptRayHitResult& HitResult,
-		TEnumAsByte<EGeometryScriptSearchOutcomePins>& Outcome,
+		EGeometryScriptSearchOutcomePins& Outcome,
 		UGeometryScriptDebug* Debug = nullptr );
 
 
@@ -116,7 +116,7 @@ public:
 		FVector QueryPoint,
 		FGeometryScriptSpatialQueryOptions Options,
 		bool& bIsInside,
-		TEnumAsByte<EGeometryScriptContainmentOutcomePins>& Outcome,
+		EGeometryScriptContainmentOutcomePins& Outcome,
 		UGeometryScriptDebug* Debug = nullptr );
 
 };
\ No newline at end of file
diff --git a/Engine/Plugins/Experimental/GeometryScripting/Source/GeometryScriptingCore/Public/GeometryScript/SceneUtilityFunctions.h b/Engine/Plugins/Experimental/GeometryScripting/Source/GeometryScriptingCore/Public/GeometryScript/SceneUtilityFunctions.h
--- a/Engine/Plugins/Experimental/GeometryScripting/Source/GeometryScriptingCore/Public/GeometryScript/SceneUtilityFunctions.h
+++ b/Engine/Plugins/Experimental/GeometryScripting/Source/GeometryScriptingCore/Public/GeometryScript/SceneUtilityFunctions.h
@@ -58,7 +58,7 @@ public:
 		FGeometryScriptCopyMeshFromComponentOptions Options,
 		bool bTransformToWorld,
 		FTransform& LocalToWorld,
-		TEnumAsByte<EGeometryScriptOutcomePins>& Outcome,
+		EGeometryScriptOutcomePins& Outcome,
 		UGeometryScriptDebug* Debug = nullptr);
 
 
@@ -90,7 +90,7 @@ public:
 		UDynamicMesh* ToDynamicMesh,
 		bool bTransformToWorld,
 		FTransform& LocalToWorld,
-		TEnumAsByte<EGeometryScriptOutcomePins>& Outcome,
+		EGeometryScriptOutcomePins& Outcome,
 		bool bUseComplexCollision = false,
 		int SphereResolution = 16,
 		UGeometryScriptDebug* Debug = nullptr);
diff --git a/Engine/Plugins/Experimental/GeometryScripting/Source/GeometryScriptingEditor/Private/CreateNewAssetUtilityFunctions.cpp b/Engine/Plugins/Experimental/GeometryScripting/Source/GeometryScriptingEditor/Private/CreateNewAssetUtilityFunctions.cpp
--- a/Engine/Plugins/Experimental/GeometryScripting/Source/GeometryScriptingEditor/Private/CreateNewAssetUtilityFunctions.cpp
+++ b/Engine/Plugins/Experimental/GeometryScripting/Source/GeometryScriptingEditor/Private/CreateNewAssetUtilityFunctions.cpp
@@ -46,7 +46,7 @@ void UGeometryScriptLibrary_CreateNewAssetFunctions::CreateUniqueNewAssetPathNam
 	FString& UniqueAssetPathAndName,
 	FString& UniqueAssetName,
 	FGeometryScriptUniqueAssetNameOptions Options,
-	TEnumAsByte<EGeometryScriptOutcomePins>& Outcome,
+	EGeometryScriptOutcomePins& Outcome,
 	UGeometryScriptDebug* Debug)
 {
 	FAssetRegistryModule& AssetRegistryModule = FModuleManager::LoadModuleChecked<FAssetRegistryModule>("AssetRegistry");
@@ -78,7 +78,7 @@ AVolume* UGeometryScriptLibrary_CreateNewAssetFunctions::CreateNewVolumeFromMesh
 	FTransform ActorTransform,
 	FString BaseActorName,
 	FGeometryScriptCreateNewVolumeFromMeshOptions Options,
-	TEnumAsByte<EGeometryScriptOutcomePins>& Outcome,
+	EGeometryScriptOutcomePins& Outcome,
 	UGeometryScriptDebug* Debug)
 {
 	Outcome = EGeometryScriptOutcomePins::Failure;
@@ -148,7 +148,7 @@ UStaticMesh* UGeometryScriptLibrary_CreateNewAssetFunctions::CreateNewStaticMesh
 	UDynamicMesh* FromDynamicMesh, 
 	FString AssetPathAndName,
 	FGeometryScriptCreateNewStaticMeshAssetOptions Options,
-	TEnumAsByte<EGeometryScriptOutcomePins>& Outcome,
+	EGeometryScriptOutcomePins& Outcome,
 	UGeometryScriptDebug* Debug)
 {
 	Outcome = EGeometryScriptOutcomePins::Failure;
@@ -218,7 +218,7 @@ USkeletalMesh* UGeometryScriptLibrary_CreateNewAssetFunctions::CreateNewSkeletal
 	USkeleton* InSkeleton,
 	FString AssetPathAndName, 
 	FGeometryScriptCreateNewSkeletalMeshAssetOptions Options,
-	TEnumAsByte<EGeometryScriptOutcomePins>& Outcome, 
+	EGeometryScriptOutcomePins& Outcome, 
 	UGeometryScriptDebug* Debug)
 {
 	using namespace UE::AssetUtils;
@@ -298,7 +298,7 @@ UTexture2D* UGeometryScriptLibrary_CreateNewAssetFunctions::CreateNewTexture2DAs
 		UTexture2D* FromTexture, 
 		FString AssetPathAndName,
 		FGeometryScriptCreateNewTexture2DAssetOptions Options,
-		TEnumAsByte<EGeometryScriptOutcomePins>& Outcome,
+		EGeometryScriptOutcomePins& Outcome,
 		UGeometryScriptDebug* Debug)
 {
 	Outcome = EGeometryScriptOutcomePins::Failure;
diff --git a/Engine/Plugins/Experimental/GeometryScripting/Source/GeometryScriptingEditor/Public/GeometryScript/CreateNewAssetUtilityFunctions.h b/Engine/Plugins/Experimental/GeometryScripting/Source/GeometryScriptingEditor/Public/GeometryScript/CreateNewAssetUtilityFunctions.h
--- a/Engine/Plugins/Experimental/GeometryScripting/Source/GeometryScriptingEditor/Public/GeometryScript/CreateNewAssetUtilityFunctions.h
+++ b/Engine/Plugins/Experimental/GeometryScripting/Source/GeometryScriptingEditor/Public/GeometryScript/CreateNewAssetUtilityFunctions.h
@@ -113,7 +113,7 @@ public:
 		FString& UniqueAssetPathAndName,
 		FString& UniqueAssetName,
 		FGeometryScriptUniqueAssetNameOptions Options,
-		TEnumAsByte<EGeometryScriptOutcomePins>& Outcome,
+		EGeometryScriptOutcomePins& Outcome,
 		UGeometryScriptDebug* Debug = nullptr);
 
 
@@ -125,7 +125,7 @@ public:
 		FTransform ActorTransform,
 		FString BaseActorName,
 		FGeometryScriptCreateNewVolumeFromMeshOptions Options,
-		TEnumAsByte<EGeometryScriptOutcomePins>& Outcome,
+		EGeometryScriptOutcomePins& Outcome,
 		UGeometryScriptDebug* Debug = nullptr);
 
 	UFUNCTION(BlueprintCallable, Category = "GeometryScript|AssetManagement", meta = (ExpandEnumAsExecs = "Outcome"))
@@ -134,7 +134,7 @@ public:
 		UDynamicMesh* FromDynamicMesh, 
 		FString AssetPathAndName,
 		FGeometryScriptCreateNewStaticMeshAssetOptions Options,
-		TEnumAsByte<EGeometryScriptOutcomePins>& Outcome,
+		EGeometryScriptOutcomePins& Outcome,
 		UGeometryScriptDebug* Debug = nullptr);
 	
 	UFUNCTION(BlueprintCallable, Category = "GeometryScript|AssetManagement", meta = (ExpandEnumAsExecs = "Outcome"))
@@ -144,7 +144,7 @@ public:
 		USkeleton* InSkeleton,
 		FString AssetPathAndName,
 		FGeometryScriptCreateNewSkeletalMeshAssetOptions Options,
-		TEnumAsByte<EGeometryScriptOutcomePins>& Outcome,
+		EGeometryScriptOutcomePins& Outcome,
 		UGeometryScriptDebug* Debug = nullptr);
 
 	UFUNCTION(BlueprintCallable, Category = "GeometryScript|AssetManagement", meta = (ExpandEnumAsExecs = "Outcome"))
@@ -153,7 +153,7 @@ public:
 		UTexture2D* FromTexture, 
 		FString AssetPathAndName,
 		FGeometryScriptCreateNewTexture2DAssetOptions Options,
-		TEnumAsByte<EGeometryScriptOutcomePins>& Outcome,
+		EGeometryScriptOutcomePins& Outcome,
 		UGeometryScriptDebug* Debug = nullptr);
 
 };
