GeometryScripting: Fix inconsistencies in "PolyGroup" strings in Blueprints

diff --git a/Engine/Plugins/Experimental/GeometryScripting/Source/GeometryScriptingCore/Public/GeometryScript/GeometryScriptSelectionTypes.h b/Engine/Plugins/Experimental/GeometryScripting/Source/GeometryScriptingCore/Public/GeometryScript/GeometryScriptSelectionTypes.h
--- a/Engine/Plugins/Experimental/GeometryScripting/Source/GeometryScriptingCore/Public/GeometryScript/GeometryScriptSelectionTypes.h
+++ b/Engine/Plugins/Experimental/GeometryScripting/Source/GeometryScriptingCore/Public/GeometryScript/GeometryScriptSelectionTypes.h
@@ -16,7 +16,7 @@ enum class EGeometryScriptMeshSelectionType : uint8
 {
 	Vertices = 0,
 	Triangles = 1,
-	Polygroups = 2
+	Polygroups = 2 UMETA(DisplayName = "PolyGroups")
 };
 
 /**
@@ -28,7 +28,7 @@ enum class EGeometryScriptMeshSelectionConversionType : uint8
 	NoConversion = 0,
 	ToVertices = 1,
 	ToTriangles = 2,
-	ToPolygroups = 3
+	ToPolygroups = 3 UMETA(DisplayName = "To PolyGroups")
 };
 
 /**
@@ -90,7 +90,7 @@ struct GEOMETRYSCRIPTINGCORE_API FGeometryScriptMeshSelection
 
 	/**
 	 * Convert the current selection to a TArray, optionally converting to ConvertToType.
-	 * For (Tri|Group)=>Vtx, all triangle vertices (in triangles or polygroups) are included.
+	 * For (Tri|Group)=>Vtx, all triangle vertices (in triangles or PolyGroups) are included.
 	 * For Vtx=>Tri, all one-ring vertices are included. For Group=>Tri, all Triangles are found via enumerating over mesh.
 	 * (Tri|Vtx)=>Group, all GroupIDs of all triangles/one-ring triangles are included
 	 */
@@ -99,7 +99,7 @@ struct GEOMETRYSCRIPTINGCORE_API FGeometryScriptMeshSelection
 	/**
 	 * Call PerTriangleFunc for each TriangleID in the Selection.
 	 * For Vertex Selections, Vertex one-rings are enumerated and accumulated in a TSet.
-	 * For Polygroup Selections, a full mesh iteration is used to find all Triangles in the Groups.
+	 * For PolyGroup Selections, a full mesh iteration is used to find all Triangles in the groups.
 	 */
 	void ProcessByTriangleID(const UE::Geometry::FDynamicMesh3& Mesh,
 		TFunctionRef<void(int32)> PerTriangleFunc,
@@ -108,7 +108,7 @@ struct GEOMETRYSCRIPTINGCORE_API FGeometryScriptMeshSelection
 	/**
 	 * Call PerVertexFunc for each VertexID in the Selection.
 	 * For Triangle Selections, Triangles Vertex tuples are enumerated and accumulated in a TSet.
-	 * For Polygroup Selections, a full mesh iteration is used to find all Triangle Vertices in the Groups (accumulated in a TSet)
+	 * For PolyGroup Selections, a full mesh iteration is used to find all Triangle Vertices in the groups (accumulated in a TSet)
 	 */
 	void ProcessByVertexID(const UE::Geometry::FDynamicMesh3& Mesh,
 		TFunctionRef<void(int32)> PerVertexFunc,
diff --git a/Engine/Plugins/Experimental/GeometryScripting/Source/GeometryScriptingCore/Public/GeometryScript/GeometryScriptTypes.h b/Engine/Plugins/Experimental/GeometryScripting/Source/GeometryScriptingCore/Public/GeometryScript/GeometryScriptTypes.h
--- a/Engine/Plugins/Experimental/GeometryScripting/Source/GeometryScriptingCore/Public/GeometryScript/GeometryScriptTypes.h
+++ b/Engine/Plugins/Experimental/GeometryScripting/Source/GeometryScriptingCore/Public/GeometryScript/GeometryScriptTypes.h
@@ -166,11 +166,11 @@ public:
 
 
 //
-// Polygroups
+// PolyGroups
 //
 
 /**
- * FGeometryScriptGroupLayer identifies a Polygroup Layer of a Mesh.
+ * FGeometryScriptGroupLayer identifies a PolyGroup Layer of a Mesh.
  * The Default Layer always exists, Extended layers may or may not exist.
  */
 USTRUCT(BlueprintType)
@@ -194,10 +194,10 @@ public:
 
 // By default structs exposed to Python will use a per-UPROPERTY comparison. When this doesn't give correct results
 // (e.g., for structs with no properties, which will compare equal in all cases because there are no properties to
-// compare), it is necessary to define explicity equality operators and add the following WithIdenticalViaEquality trait.
+// compare), it is necessary to define explicit equality operators and add the following WithIdenticalViaEquality trait.
 // Note that users can write blueprints/python scripts which pass these lists to many function calls so we would like to
 // avoid copying them (very slow if they have millions of elements) so we defined the equality operations using pointer
-// equality but this is potentially confusing if users expect that different lists with the same elements compare equal
+// equality but this is potentially confusing if users expect that different lists with the same elements compare equal.
 
 UENUM(BlueprintType)
 enum class EGeometryScriptIndexType : uint8
@@ -207,7 +207,7 @@ enum class EGeometryScriptIndexType : uint8
 	Triangle,
 	Vertex,
 	MaterialID,
-	PolygroupID
+	PolygroupID UMETA(DisplayName = "PolyGroup ID")
 };
 
 USTRUCT(BlueprintType, meta = (DisplayName = "Index List"))
diff --git a/Engine/Plugins/Experimental/GeometryScripting/Source/GeometryScriptingCore/Public/GeometryScript/MeshMaterialFunctions.h b/Engine/Plugins/Experimental/GeometryScripting/Source/GeometryScriptingCore/Public/GeometryScript/MeshMaterialFunctions.h
--- a/Engine/Plugins/Experimental/GeometryScripting/Source/GeometryScriptingCore/Public/GeometryScript/MeshMaterialFunctions.h
+++ b/Engine/Plugins/Experimental/GeometryScripting/Source/GeometryScriptingCore/Public/GeometryScript/MeshMaterialFunctions.h
@@ -91,7 +91,7 @@ public:
 	* @param MaterialID new Material ID to set
 	* @param bDeferChangeNotifications if true, the UDynamicMesh does not emit a change event/signal for this modification
 	*/
-	UFUNCTION(BlueprintCallable, Category = "GeometryScript|Polygroups", meta=(ScriptMethod))
+	UFUNCTION(BlueprintCallable, Category = "GeometryScript|PolyGroups", meta=(ScriptMethod))
 	static UPARAM(DisplayName = "Target Mesh") UDynamicMesh* 
 	SetMaterialIDForMeshSelection( 
 		UDynamicMesh* TargetMesh, 
@@ -101,20 +101,20 @@ public:
 		UGeometryScriptDebug* Debug = nullptr);
 
 	/**
-	* Set a new MaterialID on all the triangles of TargetMesh with the given Polygroup.
-	* @param GroupLayer Polygroup Layer to use as basis for polygroups
-	* @param PolygroupID Polygroup ID that specifies Triangles to set to new MaterialID
+	* Set a new MaterialID on all the triangles of TargetMesh with the given PolyGroup.
+	* @param GroupLayer PolyGroup Layer to use as basis for PolyGroups
+	* @param PolygroupID PolyGroup ID that specifies Triangles to set to new MaterialID
 	* @param MaterialID explicit new MaterialID to set
 	* @param bDeferChangeNotifications if true, the UDynamicMesh does not emit a change event/signal for this modification
 	*/
-	UFUNCTION(BlueprintCallable, Category = "GeometryScript|Materials", meta=(ScriptMethod))
+	UFUNCTION(BlueprintCallable, Category = "GeometryScript|Materials", meta=(ScriptMethod, DisplayName = "SetPolyGroupMaterialID"))
 	static UPARAM(DisplayName = "Target Mesh") UDynamicMesh* 
 	SetPolygroupMaterialID( 
 		UDynamicMesh* TargetMesh, 
 		FGeometryScriptGroupLayer GroupLayer,
-		int PolygroupID, 
+		UPARAM(DisplayName = "PolyGroup ID") int PolygroupID,
 		int MaterialID,
-		bool& bIsValidPolygroupID,
+		UPARAM(DisplayName = "Is Valid PolyGroup ID") bool& bIsValidPolygroupID,
 		bool bDeferChangeNotifications = false,
 		UGeometryScriptDebug* Debug = nullptr);
 
diff --git a/Engine/Plugins/Experimental/GeometryScripting/Source/GeometryScriptingCore/Public/GeometryScript/MeshModelingFunctions.h b/Engine/Plugins/Experimental/GeometryScripting/Source/GeometryScriptingCore/Public/GeometryScript/MeshModelingFunctions.h
--- a/Engine/Plugins/Experimental/GeometryScripting/Source/GeometryScriptingCore/Public/GeometryScript/MeshModelingFunctions.h
+++ b/Engine/Plugins/Experimental/GeometryScripting/Source/GeometryScriptingCore/Public/GeometryScript/MeshModelingFunctions.h
@@ -61,7 +61,7 @@ UENUM(BlueprintType)
 enum class EGeometryScriptPolyOperationArea : uint8
 {
 	EntireSelection = 0,
-	PerPolygroup = 1,
+	PerPolygroup = 1 UMETA(DisplayName = "Per PolyGroup"),
 	PerTriangle = 2
 };
 
@@ -205,7 +205,7 @@ public:
 
 
 	/**
-	 * If true the set of beveled polygroup edges is limited to those that 
+	 * If true the set of beveled PolyGroup edges is limited to those that 
 	 * are fully or partially contained within the (transformed) FilterBox
 	 */
 	UPROPERTY(BlueprintReadWrite, Category = FilterShape, AdvancedDisplay)
@@ -224,7 +224,7 @@ public:
 	FTransform FilterBoxTransform = FTransform::Identity;
 
 	/**
-	 * If true, then only polygroup edges that are fully contained within the filter box will be beveled,
+	 * If true, then only PolyGroup edges that are fully contained within the filter box will be beveled,
 	 * otherwise the edge will be beveled if any vertex is within the filter box.
 	 */
 	UPROPERTY(BlueprintReadWrite, Category = FilterShape, AdvancedDisplay)
@@ -241,10 +241,10 @@ enum class EGeometryScriptMeshBevelSelectionMode : uint8
 {
 	/** Convert the selection to Triangles and bevel the boundary edge loops of the triangle set */
 	TriangleArea = 0,
-	/** Convert the selection to Polygroups and bevel all the Polygroup Edges of the selected Polygroups */
-	AllPolygroupEdges = 1,
-	/** Convert the selection to Polygroups and bevel all the Polygroup Edges that are between selected Polygroups */
-	SharedPolygroupEdges = 2
+	/** Convert the selection to PolyGroups and bevel all the PolyGroup Edges of the selected PolyGroups */
+	AllPolygroupEdges = 1 UMETA(DisplayName = "All PolyGroup Edges"),
+	/** Convert the selection to PolyGroups and bevel all the PolyGroup Edges that are between selected PolyGroups */
+	SharedPolygroupEdges = 2 UMETA(DisplayName = "Shared PolyGroup Edges")
 };
 
 
@@ -360,7 +360,7 @@ public:
 	/**
 	 * Apply a Mesh Bevel operation to parts of TargetMesh using the BevelOptions settings.
 	 * @param Selection specifies which mesh edges to Bevel
-	 * @param BevelMode specifies how Selection should be converted to a Triangle Region or set of Polygroup Edges
+	 * @param BevelMode specifies how Selection should be converted to a Triangle Region or set of PolyGroup Edges
 	 * @param BevelOptions settings for the Bevel Operation
 	 */
 	UFUNCTION(BlueprintCallable, Category = "GeometryScript|Modeling", meta=(ScriptMethod))
@@ -373,7 +373,7 @@ public:
 		UGeometryScriptDebug* Debug = nullptr );
 
 
-	UFUNCTION(BlueprintCallable, Category = "GeometryScript|Modeling", meta=(ScriptMethod))
+	UFUNCTION(BlueprintCallable, Category = "GeometryScript|Modeling", meta=(ScriptMethod, DisplayName = "ApplyMeshPolyGroupBevel"))
 	static UPARAM(DisplayName = "Target Mesh") UDynamicMesh* 
 	ApplyMeshPolygroupBevel(
 		UDynamicMesh* TargetMesh,
diff --git a/Engine/Plugins/Experimental/GeometryScripting/Source/GeometryScriptingCore/Public/GeometryScript/MeshPolygroupFunctions.h b/Engine/Plugins/Experimental/GeometryScripting/Source/GeometryScriptingCore/Public/GeometryScript/MeshPolygroupFunctions.h
--- a/Engine/Plugins/Experimental/GeometryScripting/Source/GeometryScriptingCore/Public/GeometryScript/MeshPolygroupFunctions.h
+++ b/Engine/Plugins/Experimental/GeometryScripting/Source/GeometryScriptingCore/Public/GeometryScript/MeshPolygroupFunctions.h
@@ -11,24 +11,24 @@
 class UDynamicMesh;
 
 
-UCLASS(meta = (ScriptName = "GeometryScript_Polygroups"))
+UCLASS(meta = (ScriptName = "GeometryScript_PolyGroups"))
 class GEOMETRYSCRIPTINGCORE_API UGeometryScriptLibrary_MeshPolygroupFunctions : public UBlueprintFunctionLibrary
 {
 	GENERATED_BODY()
 public:
 
-	UFUNCTION(BlueprintCallable, Category = "GeometryScript|Polygroups", meta=(ScriptMethod))
+	UFUNCTION(BlueprintCallable, Category = "GeometryScript|PolyGroups", meta = (ScriptMethod, DisplayName = "EnablePolyGroups"))
 	static UPARAM(DisplayName = "Target Mesh") UDynamicMesh* 
 	EnablePolygroups( UDynamicMesh* TargetMesh, UGeometryScriptDebug* Debug = nullptr );
 
-	UFUNCTION(BlueprintCallable, Category = "GeometryScript|Polygroups", meta=(ScriptMethod))
+	UFUNCTION(BlueprintCallable, Category = "GeometryScript|PolyGroups", meta = (ScriptMethod, DisplayName = "SetNumExtendedPolyGroupLayers"))
 	static UPARAM(DisplayName = "Target Mesh") UDynamicMesh* 
 	SetNumExtendedPolygroupLayers( 
 		UDynamicMesh* TargetMesh, 
 		int NumLayers,
 		UGeometryScriptDebug* Debug = nullptr );
 
-	UFUNCTION(BlueprintCallable, Category = "GeometryScript|Polygroups", meta=(ScriptMethod))
+	UFUNCTION(BlueprintCallable, Category = "GeometryScript|PolyGroups", meta = (ScriptMethod, DisplayName = "ClearPolyGroups"))
 	static UPARAM(DisplayName = "Target Mesh") UDynamicMesh* 
 	ClearPolygroups( 
 		UDynamicMesh* TargetMesh, 
@@ -36,7 +36,7 @@ public:
 		int ClearValue = 0,
 		UGeometryScriptDebug* Debug = nullptr );
 
-	UFUNCTION(BlueprintCallable, Category = "GeometryScript|Polygroups", meta=(ScriptMethod))
+	UFUNCTION(BlueprintCallable, Category = "GeometryScript|PolyGroups", meta = (ScriptMethod, DisplayName = "CopyPolyGroupsLayer"))
 	static UPARAM(DisplayName = "Target Mesh") UDynamicMesh* 
 	CopyPolygroupsLayer( 
 		UDynamicMesh* TargetMesh, 
@@ -45,7 +45,7 @@ public:
 		UGeometryScriptDebug* Debug = nullptr );
 
 
-	UFUNCTION(BlueprintCallable, Category = "GeometryScript|Polygroups", meta=(ScriptMethod))
+	UFUNCTION(BlueprintCallable, Category = "GeometryScript|PolyGroups", meta = (ScriptMethod, DisplayName = "ConvertUVIslandsToPolyGroups"))
 	static UPARAM(DisplayName = "Target Mesh") UDynamicMesh* 
 	ConvertUVIslandsToPolygroups( 
 		UDynamicMesh* TargetMesh, 
@@ -53,14 +53,14 @@ public:
 		int UVLayer = 0,
 		UGeometryScriptDebug* Debug = nullptr );
 
-	UFUNCTION(BlueprintCallable, Category = "GeometryScript|Polygroups", meta=(ScriptMethod))
+	UFUNCTION(BlueprintCallable, Category = "GeometryScript|PolyGroups", meta = (ScriptMethod, DisplayName = "ConvertComponentsToPolyGroups"))
 	static UPARAM(DisplayName = "Target Mesh") UDynamicMesh* 
 	ConvertComponentsToPolygroups( 
 		UDynamicMesh* TargetMesh, 
 		FGeometryScriptGroupLayer GroupLayer,
 		UGeometryScriptDebug* Debug = nullptr );
 
-	UFUNCTION(BlueprintCallable, Category = "GeometryScript|Polygroups", meta=(ScriptMethod))
+	UFUNCTION(BlueprintCallable, Category = "GeometryScript|PolyGroups", meta = (ScriptMethod, DisplayName = "ComputePolyGroupsFromAngleThreshold"))
 	static UPARAM(DisplayName = "Target Mesh") UDynamicMesh* 
 	ComputePolygroupsFromAngleThreshold( 
 		UDynamicMesh* TargetMesh, 
@@ -69,7 +69,7 @@ public:
 		int MinGroupSize = 2,
 		UGeometryScriptDebug* Debug = nullptr );
 
-	UFUNCTION(BlueprintCallable, Category = "GeometryScript|Polygroups", meta=(ScriptMethod))
+	UFUNCTION(BlueprintCallable, Category = "GeometryScript|PolyGroups", meta = (ScriptMethod, DisplayName = "ComputePolyGroupsFromPolygonDetection"))
 	static UPARAM(DisplayName = "Target Mesh") UDynamicMesh* 
 	ComputePolygroupsFromPolygonDetection( 
 		UDynamicMesh* TargetMesh, 
@@ -81,72 +81,72 @@ public:
 		int MaxSearchRounds = 1,
 		UGeometryScriptDebug* Debug = nullptr );
 
-	UFUNCTION(BlueprintPure, Category = "GeometryScript|Polygroups", meta=(ScriptMethod))
-	static UPARAM(DisplayName = "Polygroup ID") int32
+	UFUNCTION(BlueprintPure, Category = "GeometryScript|PolyGroups", meta = (ScriptMethod, DisplayName = "GetTrianglePolyGroupID"))
+	static UPARAM(DisplayName = "PolyGroup ID") int32
 	GetTrianglePolygroupID( 
 		UDynamicMesh* TargetMesh, 
 		FGeometryScriptGroupLayer GroupLayer, 
 		int TriangleID, 
 		bool& bIsValidTriangle );
 
-	UFUNCTION(BlueprintCallable, Category = "GeometryScript|Polygroups", meta=(ScriptMethod))
+	UFUNCTION(BlueprintCallable, Category = "GeometryScript|PolyGroups", meta = (ScriptMethod, DisplayName = "DeleteTrianglesInPolyGroup"))
 	static UPARAM(DisplayName = "Target Mesh") UDynamicMesh* 
 	DeleteTrianglesInPolygroup( 
 		UDynamicMesh* TargetMesh, 
 		FGeometryScriptGroupLayer GroupLayer,
-		int PolygroupID,
+		UPARAM(DisplayName = "PolyGroup ID") int PolygroupID,
 		int& NumDeleted,
 		bool bDeferChangeNotifications = false,
 		UGeometryScriptDebug* Debug = nullptr);
 
 	/**
-	* Create list of per-triangle Polygroup IDs for the Polygroup in the Mesh
+	* Create list of per-triangle PolyGroup IDs for the PolyGroup in the Mesh
 	* @warning if the mesh is not Triangle-Compact (eg GetHasTriangleIDGaps == false) then the returned list will also have the same gaps
 	*/
-	UFUNCTION(BlueprintCallable, Category = "GeometryScript|Polygroups", meta=(ScriptMethod))
+	UFUNCTION(BlueprintCallable, Category = "GeometryScript|PolyGroups", meta = (ScriptMethod, DisplayName = "GetAllTrianglePolyGroupIDs"))
 	static UPARAM(DisplayName = "Target Mesh") UDynamicMesh* 
 	GetAllTrianglePolygroupIDs( 
 		UDynamicMesh* TargetMesh, 
 		FGeometryScriptGroupLayer GroupLayer, 
-		UPARAM(ref, DisplayName="Polygroup IDs Out") FGeometryScriptIndexList& PolygroupIDsOut );
+		UPARAM(ref, DisplayName="PolyGroup IDs Out") FGeometryScriptIndexList& PolygroupIDsOut );
 
 	/**
-	* Create list of all unique Polygroup IDs that exist in the Polygroup Layer in the Mesh
+	* Create list of all unique PolyGroup IDs that exist in the PolyGroup Layer in the Mesh
 	*/
-	UFUNCTION(BlueprintCallable, Category = "GeometryScript|Polygroups", meta=(ScriptMethod))
+	UFUNCTION(BlueprintCallable, Category = "GeometryScript|PolyGroups", meta = (ScriptMethod, DisplayName = "GetPolyGroupIDsInMesh"))
 	static UPARAM(DisplayName = "Target Mesh") UDynamicMesh* 
 	GetPolygroupIDsInMesh( 
 		UDynamicMesh* TargetMesh, 
 		FGeometryScriptGroupLayer GroupLayer, 
-		UPARAM(ref, DisplayName="Polygroup IDs Out") FGeometryScriptIndexList& PolygroupIDsOut );
+		UPARAM(ref, DisplayName="PolyGroup IDs Out") FGeometryScriptIndexList& PolygroupIDsOut );
 
 	/**
-	 * Create list of all triangles with the given Polygroup ID in the given GroupLayer (not necessarily a single connected-component)
+	 * Create list of all triangles with the given PolyGroup ID in the given GroupLayer (not necessarily a single connected-component)
 	 */
-	UFUNCTION(BlueprintCallable, Category = "GeometryScript|Polygroups", meta=(ScriptMethod))
+	UFUNCTION(BlueprintCallable, Category = "GeometryScript|Polygroups", meta = (ScriptMethod, DisplayName = "GetTrianglesInPolyGroup"))
 	static UPARAM(DisplayName = "Target Mesh") UDynamicMesh* 
 	GetTrianglesInPolygroup( 
 		UDynamicMesh* TargetMesh, 
 		FGeometryScriptGroupLayer GroupLayer, 
-		int PolygroupID, 
+		UPARAM(DisplayName = "PolyGroup ID") int PolygroupID, 
 		UPARAM(ref, DisplayName="Triangle IDs Out") FGeometryScriptIndexList& TriangleIDsOut );
 
 
 	/**
-	 * Set a new Polygroup on all the triangles of the given Selection, for the given GroupLayer.
-	 * @param SetPolygroupID explicit new PolygroupID to set
-	 * @param bGenerateNewPolygroup if true, SetPolygroupID is ignored and a new unique PolygroupID is generated
-	 * @param SetPolygroupIDOut the PolygroupID that was set on the triangles is returned here (whether explicit or auto-generated)
+	 * Set a new PolyGroup on all the triangles of the given Selection, for the given GroupLayer.
+	 * @param SetPolygroupID explicit new PolyGroupID to set
+	 * @param bGenerateNewPolygroup if true, SetPolyGroupID is ignored and a new unique PolyGroupID is generated
+	 * @param SetPolygroupIDOut the PolyGroupID that was set on the triangles is returned here (whether explicit or auto-generated)
 	 * @param bDeferChangeNotifications if true, the UDynamicMesh does not emit a change event/signal for this modification
 	 */
-	UFUNCTION(BlueprintCallable, Category = "GeometryScript|Polygroups", meta=(ScriptMethod))
+	UFUNCTION(BlueprintCallable, Category = "GeometryScript|PolyGroups", meta = (ScriptMethod, DisplayName = "SetPolyGroupForMeshSelection"))
 	static UPARAM(DisplayName = "Target Mesh") UDynamicMesh* 
 	SetPolygroupForMeshSelection( 
 		UDynamicMesh* TargetMesh, 
 		FGeometryScriptGroupLayer GroupLayer, 
 		FGeometryScriptMeshSelection Selection,
-		int& SetPolygroupIDOut,
-		int SetPolygroupID = 0,
-		bool bGenerateNewPolygroup = false,
+		UPARAM(DisplayName = "Set PolyGroup ID Out") int& SetPolygroupIDOut,
+		UPARAM(DisplayName = "Set PolyGroup ID") int SetPolygroupID = 0,
+		UPARAM(DisplayName = " Generate New PolyGroup") bool bGenerateNewPolygroup = false,
 		bool bDeferChangeNotifications = false);
 };
\ No newline at end of file
diff --git a/Engine/Plugins/Experimental/GeometryScripting/Source/GeometryScriptingCore/Public/GeometryScript/MeshQueryFunctions.h b/Engine/Plugins/Experimental/GeometryScripting/Source/GeometryScriptingCore/Public/GeometryScript/MeshQueryFunctions.h
--- a/Engine/Plugins/Experimental/GeometryScripting/Source/GeometryScriptingCore/Public/GeometryScript/MeshQueryFunctions.h
+++ b/Engine/Plugins/Experimental/GeometryScripting/Source/GeometryScriptingCore/Public/GeometryScript/MeshQueryFunctions.h
@@ -116,10 +116,10 @@ public:
 	static UPARAM(DisplayName = "Enabled") bool GetHasMaterialIDs( UDynamicMesh* TargetMesh );
 
 
-	UFUNCTION(BlueprintPure, Category = "GeometryScript|MeshQueries", meta=(ScriptMethod))
+	UFUNCTION(BlueprintPure, Category = "GeometryScript|MeshQueries", meta=(ScriptMethod, DisplayName = "GetHasPolyGroups"))
 	static UPARAM(DisplayName = "Enabled") bool GetHasPolygroups( UDynamicMesh* TargetMesh );
 
-	UFUNCTION(BlueprintPure, Category = "GeometryScript|MeshQueries", meta=(ScriptMethod))
+	UFUNCTION(BlueprintPure, Category = "GeometryScript|MeshQueries", meta=(ScriptMethod, DisplayName = "GetNumExtendedPolyGroupLayers"))
 	static UPARAM(DisplayName = "Num Layers") int GetNumExtendedPolygroupLayers( UDynamicMesh* TargetMesh );
 
 
diff --git a/Engine/Plugins/Experimental/GeometryScripting/Source/GeometryScriptingCore/Public/GeometryScript/MeshSelectionFunctions.h b/Engine/Plugins/Experimental/GeometryScripting/Source/GeometryScriptingCore/Public/GeometryScript/MeshSelectionFunctions.h
--- a/Engine/Plugins/Experimental/GeometryScripting/Source/GeometryScriptingCore/Public/GeometryScript/MeshSelectionFunctions.h
+++ b/Engine/Plugins/Experimental/GeometryScripting/Source/GeometryScriptingCore/Public/GeometryScript/MeshSelectionFunctions.h
@@ -16,7 +16,7 @@ UENUM(BlueprintType)
 enum class EGeometryScriptTopologyConnectionType : uint8
 {
 	Geometric,
-	Polygroup,
+	Polygroup UMETA(DisplayName = "PolyGroup"),
 	MaterialID
 };
 
@@ -66,7 +66,7 @@ public:
 	 * If bAllowPartialInclusion is disabled, then more restrictive conversions are performed, as follows:
 	 *   For To-Vertices, only include vertices where all one-ring triangles are included in FromSelection.
 	 *   For To-Triangles, only include triangles where all tri vertices are included in FromSelection.
-	 *   For To-Polygroups, only include groups where all group triangles are included in FromSelection
+	 *   For To-PolyGroups, only include groups where all group triangles are included in FromSelection
 	 * @param bAllowPartialInclusion if false, perform more limited selection conversion as described above
 	 */
 	UFUNCTION(BlueprintCallable, Category = "GeometryScript|MeshSelection", meta=(ScriptMethod))
@@ -209,7 +209,7 @@ public:
 	/**
 	* Create a new Mesh Selection of the SelectionType for the TargetMesh by finding all elements that have a normal
 	* vector that is within an angular deviation threshold from the given Normal. 
-	* For Triangle and Polygroup selections the triangle facet normal is used, for Vertex selections the per-vertex averaged normal is used.
+	* For Triangle and PolyGroup selections the triangle facet normal is used, for Vertex selections the per-vertex averaged normal is used.
 	* @param Normal normal/direction vector to measure against
 	* @param MaxAngleDeg maximum angular deviation from Normal, in degrees
 	* @param bInvert return a selection of all elements not within the given deviation
@@ -229,7 +229,7 @@ public:
 
 	/**
 	* Create a new Mesh Selection of the SelectionType for the TargetMesh by finding all elements inside a second SelectionMesh
-	* For Triangle and Polygroup selections the triangle facet normal is used, for Vertex selections the per-vertex averaged normal is used.
+	* For Triangle and PolyGroup selections the triangle facet normal is used, for Vertex selections the per-vertex averaged normal is used.
 	* @param SelectionMeshTransform Transform applied to SelectionMesh for inside/outside testing
 	* @param bInvert return a selection of all elements not within the given deviation
 	* @param ShellDistance If > 0, points within this distance from SelectionMesh will also be considered "inside"
@@ -280,10 +280,10 @@ public:
 	 * Grow or Shrink the Selection on the TargetMesh to connected neighbours.
 	 * For Vertex selections, Expand includes vertices in one-ring of selected vertices, and Contract removes any vertices with a one-ring neighbour that is not selected
 	 * For Triangle selections, Add/Remove Triangles connected to selected Triangles
-	 * For Polygroup selections, Add/Remove Polygroups connected to selected Polygroups
+	 * For PolyGroup selections, Add/Remove PolyGroups connected to selected PolyGroups
 	 * @param Iterations number of times to Expand/Contract the Selection. Valid range is [0,100] where 0 is a no-op.
 	 * @param bContract if true selection contracts instead of growing
-	 * @param bOnlyExpandToFaceNeighbours if true, only adjacent Triangles/Polygroups directly connected by an edge are added, vs connected to any selected vertex
+	 * @param bOnlyExpandToFaceNeighbours if true, only adjacent Triangles/PolyGroups directly connected by an edge are added, vs connected to any selected vertex
 	 */
 	UFUNCTION(BlueprintCallable, Category = "GeometryScript|MeshSelection", meta=(ScriptMethod))
 	static UPARAM(DisplayName = "Target Mesh") UDynamicMesh* 
diff --git a/Engine/Plugins/Experimental/GeometryScripting/Source/GeometryScriptingCore/Public/GeometryScript/MeshSimplifyFunctions.h b/Engine/Plugins/Experimental/GeometryScripting/Source/GeometryScriptingCore/Public/GeometryScript/MeshSimplifyFunctions.h
--- a/Engine/Plugins/Experimental/GeometryScripting/Source/GeometryScriptingCore/Public/GeometryScript/MeshSimplifyFunctions.h
+++ b/Engine/Plugins/Experimental/GeometryScripting/Source/GeometryScriptingCore/Public/GeometryScript/MeshSimplifyFunctions.h
@@ -92,7 +92,7 @@ public:
 		FGeometryScriptPlanarSimplifyOptions Options,
 		UGeometryScriptDebug* Debug = nullptr);
 
-	UFUNCTION(BlueprintCallable, Category = "GeometryScript|Simplification", meta=(ScriptMethod))
+	UFUNCTION(BlueprintCallable, Category = "GeometryScript|Simplification", meta=(ScriptMethod, DisplayName = "ApplySimplifyToPolyGroupTopology"))
 	static UPARAM(DisplayName = "Target Mesh") UDynamicMesh* 
 	ApplySimplifyToPolygroupTopology(  
 		UDynamicMesh* TargetMesh, 
diff --git a/Engine/Plugins/Experimental/GeometryScripting/Source/GeometryScriptingCore/Public/GeometryScript/MeshUVFunctions.h b/Engine/Plugins/Experimental/GeometryScripting/Source/GeometryScriptingCore/Public/GeometryScript/MeshUVFunctions.h
--- a/Engine/Plugins/Experimental/GeometryScripting/Source/GeometryScriptingCore/Public/GeometryScript/MeshUVFunctions.h
+++ b/Engine/Plugins/Experimental/GeometryScripting/Source/GeometryScriptingCore/Public/GeometryScript/MeshUVFunctions.h
@@ -325,8 +325,8 @@ public:
 
 	/**
 	 * Copy the 2D UVs from the given UVSetIndex in CopyFromMesh to the 3D vertex positions in CopyToUVMesh,
-	 * with the triangle mesh topolgoy defined by the UV Set. Generally this "UV Mesh" topolgoy will not
-	 * be the same as the 3D mesh topology. Polygroup IDs and Material IDs are preserved in the UVMesh.
+	 * with the triangle mesh topology defined by the UV Set. Generally this "UV Mesh" topology will not
+	 * be the same as the 3D mesh topology. PolyGroup IDs and Material IDs are preserved in the UVMesh.
 	 * 
 	 * 2D UV Positions are copied to 3D as (X, Y, 0) 
 	 * 
diff --git a/Engine/Plugins/Experimental/GeometryScripting/Source/GeometryScriptingCore/Public/GeometryScript/MeshVertexColorFunctions.h b/Engine/Plugins/Experimental/GeometryScripting/Source/GeometryScriptingCore/Public/GeometryScript/MeshVertexColorFunctions.h
--- a/Engine/Plugins/Experimental/GeometryScripting/Source/GeometryScriptingCore/Public/GeometryScript/MeshVertexColorFunctions.h
+++ b/Engine/Plugins/Experimental/GeometryScripting/Source/GeometryScriptingCore/Public/GeometryScript/MeshVertexColorFunctions.h
@@ -36,7 +36,7 @@ public:
 	/**
 	 * Set the colors in the TargetMesh VertexColor Overlay identified by the Selection to a constant value.
 	 * For a Vertex Selection, each existing VertexColor Overlay Element for the vertex is updated.
-	 * For a Triangle or Polygroup Selection, all Overlay Elements in the identified Triangles are updated.
+	 * For a Triangle or PolyGroup Selection, all Overlay Elements in the identified Triangles are updated.
 	 * @param Color the constant color to set
 	 * @param Flags specify which RGBA channels to set (default all channels)
 	 * @param bCreateColorSeam if true, a "hard edge" in the vertex colors is created, by creating new Elements for all the triangles in the selection. If enabled, Vertex selections are converted to Triangle selections, and Flags is ignored. 
diff --git a/Engine/Plugins/Experimental/GeometryScripting/Source/GeometryScriptingEditor/Public/GeometryScript/OpenSubdivUtilityFunctions.h b/Engine/Plugins/Experimental/GeometryScripting/Source/GeometryScriptingEditor/Public/GeometryScript/OpenSubdivUtilityFunctions.h
--- a/Engine/Plugins/Experimental/GeometryScripting/Source/GeometryScriptingEditor/Public/GeometryScript/OpenSubdivUtilityFunctions.h
+++ b/Engine/Plugins/Experimental/GeometryScripting/Source/GeometryScriptingEditor/Public/GeometryScript/OpenSubdivUtilityFunctions.h
@@ -14,7 +14,7 @@ class GEOMETRYSCRIPTINGEDITOR_API UGeometryScriptLibrary_OpenSubdivFunctions : p
 public:
 
 
-	UFUNCTION(BlueprintCallable, Category = "GeometryScript|OpenSubdiv", meta=(ScriptMethod))
+	UFUNCTION(BlueprintCallable, Category = "GeometryScript|OpenSubdiv", meta=(ScriptMethod, DisplayName = "ApplyPolyGroupCatmullClarkSubD"))
 	static UPARAM(DisplayName = "Target Mesh") UDynamicMesh* 
 	ApplyPolygroupCatmullClarkSubD(
 		UDynamicMesh* FromDynamicMesh, 
