GeometryScript:
- Add new EditorDynamicMeshFunctions library, with BeginTrackedMeshChange() and EmitTrackedMeshChange(), allows for undoable edits on dynamic mesh
- Add new TransformFunctions library with MakeTransformFromZAxis, MakeTransformFromAxes, GetTransformAxisVector, GetTransformAxisRay, GetTransformAxisPlane
- Add GetRayBoxIntersection, GetRayPlaneIntersection, GetRayLineClosestPoint, GetRaySegmentClosestPoint
- remove const& from arguments to many of the low-level functions as this prevents inline values in nodes

diff --git a/Engine/Plugins/Experimental/GeometryScripting/Source/GeometryScriptingCore/Private/ShapeFunctions.cpp b/Engine/Plugins/Experimental/GeometryScripting/Source/GeometryScriptingCore/Private/ShapeFunctions.cpp
--- a/Engine/Plugins/Experimental/GeometryScripting/Source/GeometryScriptingCore/Private/ShapeFunctions.cpp
+++ b/Engine/Plugins/Experimental/GeometryScripting/Source/GeometryScriptingCore/Private/ShapeFunctions.cpp
@@ -2,8 +2,11 @@
 
 #include "GeometryScript/ShapeFunctions.h"
 #include "VectorTypes.h"
+#include "FrameTypes.h"
 #include "Intersection/IntersectionUtil.h"
 #include "Intersection/IntrRay3AxisAlignedBox3.h"
+#include "Distance/DistLine3Ray3.h"
+#include "Distance/DistRay3Segment3.h"
 
 #include UE_INLINE_GENERATED_CPP_BY_NAME(ShapeFunctions)
 
@@ -13,7 +16,63 @@ using namespace UE::Geometry;
 
 
 
-FRay UGeometryScriptLibrary_RayFunctions::MakeRayFromPoints(const FVector& A, const FVector& B)
+FTransform UGeometryScriptLibrary_TransformFunctions::MakeTransformFromZAxis(FVector Location, FVector ZAxis)
+{
+	return FFrame3d(Location, ZAxis).ToFTransform();
+}
+
+
+FTransform UGeometryScriptLibrary_TransformFunctions::MakeTransformFromAxes(
+	FVector Location, 
+	FVector ZAxis, 
+	FVector TangentAxis, 
+	bool bTangentIsX)
+{
+	FFrame3d Frame(Location, ZAxis);
+	// if axes are parallel we can't do the alignment so just ignore it?
+	if (FMathd::Abs(ZAxis.Dot(TangentAxis)) < 0.999)
+	{
+		Frame.ConstrainedAlignAxis(
+			(bTangentIsX) ? 0 : 1, TangentAxis, Frame.Z());
+	}
+	return Frame.ToFTransform();
+}
+
+FVector UGeometryScriptLibrary_TransformFunctions::GetTransformAxisVector(FTransform Transform, EGeometryScriptAxis Axis)
+{
+	switch(Axis)
+	{
+		case EGeometryScriptAxis::X: return ((FQuaterniond)Transform.GetRotation()).AxisX();
+		case EGeometryScriptAxis::Y: return ((FQuaterniond)Transform.GetRotation()).AxisY();
+		default:
+		case EGeometryScriptAxis::Z: return ((FQuaterniond)Transform.GetRotation()).AxisZ();
+	}
+}
+
+FRay UGeometryScriptLibrary_TransformFunctions::GetTransformAxisRay(FTransform Transform, EGeometryScriptAxis Axis)
+{
+	switch(Axis)
+	{
+		case EGeometryScriptAxis::X: return FRay( Transform.GetLocation(), ((FQuaterniond)Transform.GetRotation()).AxisX() );
+		case EGeometryScriptAxis::Y: return FRay( Transform.GetLocation(), ((FQuaterniond)Transform.GetRotation()).AxisY() );
+		default:
+		case EGeometryScriptAxis::Z: return FRay( Transform.GetLocation(), ((FQuaterniond)Transform.GetRotation()).AxisZ() );
+	}
+}
+
+FPlane UGeometryScriptLibrary_TransformFunctions::GetTransformAxisPlane(FTransform Transform, EGeometryScriptAxis Axis)
+{
+	switch(Axis)
+	{
+		case EGeometryScriptAxis::X: return FPlane( Transform.GetLocation(), ((FQuaterniond)Transform.GetRotation()).AxisX() );
+		case EGeometryScriptAxis::Y: return FPlane( Transform.GetLocation(), ((FQuaterniond)Transform.GetRotation()).AxisY() );
+		default:
+		case EGeometryScriptAxis::Z: return FPlane( Transform.GetLocation(), ((FQuaterniond)Transform.GetRotation()).AxisZ() );
+	}
+}
+
+
+FRay UGeometryScriptLibrary_RayFunctions::MakeRayFromPoints(FVector A, FVector B)
 {
 	FVector Direction(B-A);
 	if (Normalize(Direction) < FMathf::Epsilon)
@@ -23,44 +82,44 @@ FRay UGeometryScriptLibrary_RayFunctions::MakeRayFromPoints(const FVector& A, co
 	return FRay(A, Direction, true);
 }
 
-FRay UGeometryScriptLibrary_RayFunctions::MakeRayFromPointDirection(const FVector& Origin, const FVector& Direction, bool bDirectionIsNormalized)
+FRay UGeometryScriptLibrary_RayFunctions::MakeRayFromPointDirection(FVector Origin, FVector Direction, bool bDirectionIsNormalized)
 {
 	return FRay(Origin, Direction, bDirectionIsNormalized);
 }
 
-FRay UGeometryScriptLibrary_RayFunctions::GetTransformedRay(const FRay& Ray, const FTransform& TransformIn, bool bInvert)
+FRay UGeometryScriptLibrary_RayFunctions::GetTransformedRay(FRay Ray, FTransform TransformIn, bool bInvert)
 {
 	FTransformSRT3d Transform(TransformIn);
 	return (bInvert) ? Transform.InverseTransformRay(Ray) : Transform.TransformRay(Ray);
 }
 
-FVector UGeometryScriptLibrary_RayFunctions::GetRayPoint(const FRay& Ray, double Distance)
+FVector UGeometryScriptLibrary_RayFunctions::GetRayPoint(FRay Ray, double Distance)
 {
 	return Ray.PointAt(Distance);
 }
 
-void UGeometryScriptLibrary_RayFunctions::GetRayStartEnd(const FRay& Ray, double StartDistance, double EndDistance, FVector& StartPoint, FVector& EndPoint)
+void UGeometryScriptLibrary_RayFunctions::GetRayStartEnd(FRay Ray, double StartDistance, double EndDistance, FVector& StartPoint, FVector& EndPoint)
 {
 	StartPoint = Ray.PointAt(StartDistance);
 	EndPoint = Ray.PointAt(EndDistance == 0 ? TNumericLimits<float>::Max() : EndDistance );
 }
 
-double UGeometryScriptLibrary_RayFunctions::GetRayParameter(const FRay& Ray, const FVector& Point)
+double UGeometryScriptLibrary_RayFunctions::GetRayParameter(FRay Ray, FVector Point)
 {
 	return Ray.GetParameter(Point);
 }
 
-double UGeometryScriptLibrary_RayFunctions::GetRayPointDistance(const FRay& Ray, const FVector& Point)
+double UGeometryScriptLibrary_RayFunctions::GetRayPointDistance(FRay Ray, FVector Point)
 {
 	return Ray.Dist(Point);
 }
 
-FVector UGeometryScriptLibrary_RayFunctions::GetRayClosestPoint(const FRay& Ray, const FVector& Point)
+FVector UGeometryScriptLibrary_RayFunctions::GetRayClosestPoint(FRay Ray, FVector Point)
 {
 	return Ray.ClosestPoint(Point);
 }
 
-bool UGeometryScriptLibrary_RayFunctions::GetRaySphereIntersection(const FRay& Ray, const FVector& SphereCenter, double SphereRadius, double& Distance1, double& Distance2)
+bool UGeometryScriptLibrary_RayFunctions::GetRaySphereIntersection(FRay Ray, FVector SphereCenter, double SphereRadius, double& Distance1, double& Distance2)
 {
 	Distance1 = Distance2 = (double)TNumericLimits<float>::Max();
 	FLinearIntersection Intersection;
@@ -73,36 +132,69 @@ bool UGeometryScriptLibrary_RayFunctions::GetRaySphereIntersection(const FRay& R
 	return bIntersects;
 }
 
-bool UGeometryScriptLibrary_RayFunctions::GetRayBoxIntersection(const FRay& Ray, const FBox& Box, double& HitDistance)
+bool UGeometryScriptLibrary_RayFunctions::GetRayBoxIntersection(FRay Ray, FBox Box, double& HitDistance)
 {
 	HitDistance = (double)TNumericLimits<float>::Max();
 	return FIntrRay3AxisAlignedBox3d::FindIntersection(Ray, FAxisAlignedBox3d(Box), HitDistance);
 }
 
+bool UGeometryScriptLibrary_RayFunctions::GetRayPlaneIntersection(FRay Ray, FPlane Plane, double& HitDistance)
+{
+	const FVector PlaneNormal = FVector(Plane.X, Plane.Y, Plane.Z);
+	const FVector PlaneOrigin = PlaneNormal * Plane.W;
+	if (FMathd::Abs(FVector::DotProduct(Ray.Direction, PlaneNormal)) > FMathd::ZeroTolerance)
+	{
+		HitDistance = FVector::DotProduct((PlaneOrigin - Ray.Origin), PlaneNormal);
+		return true;
+	}
+	HitDistance = 0;
+	return false;
+}
+
+
+double UGeometryScriptLibrary_RayFunctions::GetRayLineClosestPoint(FRay Ray, FVector LineOrigin, FVector LineDirection, double& RayParameter, FVector& RayPoint, double& LineParameter, FVector& LinePoint)
+{
+	FDistLine3Ray3d Distance( FLine3d(LineOrigin, LineDirection), Ray );
+	double Dist = Distance.Get();
+	RayParameter = Distance.RayParameter;
+	RayPoint = Distance.RayClosestPoint;
+	LineParameter = Distance.LineParameter;
+	LinePoint = Distance.LineClosestPoint;
+	return Dist;
+}
 
+double UGeometryScriptLibrary_RayFunctions::GetRaySegmentClosestPoint(FRay Ray, FVector SegStartPoint, FVector SegEndPoint, double& RayParameter, FVector& RayPoint, FVector& SegPoint)
+{
+	FDistRay3Segment3d Distance( Ray, FSegment3d(SegStartPoint, SegEndPoint) );
+	double Dist = Distance.Get();
+	RayParameter = Distance.RayParameter;
+	RayPoint = Distance.RayClosestPoint;
+	SegPoint = Distance.SegmentClosestPoint;
+	return Dist;
+}
 
 
 
-FBox UGeometryScriptLibrary_BoxFunctions::MakeBoxFromCenterSize(const FVector& Center, const FVector& Dimensions)
+FBox UGeometryScriptLibrary_BoxFunctions::MakeBoxFromCenterSize(FVector Center, FVector Dimensions)
 {
 	FVector Extents( FMathd::Max(0, Dimensions.X * 0.5), FMathd::Max(0, Dimensions.Y * 0.5), FMathd::Max(0, Dimensions.Z * 0.5) );
 	return FBox(Center - Extents, Center + Extents);
 }
 
-FBox UGeometryScriptLibrary_BoxFunctions::MakeBoxFromCenterExtents(const FVector& Center, const FVector& Extents)
+FBox UGeometryScriptLibrary_BoxFunctions::MakeBoxFromCenterExtents(FVector Center, FVector Extents)
 {
 	return FBox(Center - Extents, Center + Extents);
 }
 
 
-void UGeometryScriptLibrary_BoxFunctions::GetBoxCenterSize(const FBox& Box, FVector& Center, FVector& Dimensions)
+void UGeometryScriptLibrary_BoxFunctions::GetBoxCenterSize(FBox Box, FVector& Center, FVector& Dimensions)
 {
 	FVector Extents;
 	Box.GetCenterAndExtents(Center, Extents);
 	Dimensions = 2.0 * Extents;
 }
 
-FVector UGeometryScriptLibrary_BoxFunctions::GetBoxCorner(const FBox& Box, int CornerIndex)
+FVector UGeometryScriptLibrary_BoxFunctions::GetBoxCorner(FBox Box, int CornerIndex)
 {
 	CornerIndex = FMath::Clamp(CornerIndex, 0, 7);
 	switch (CornerIndex)		// matches Box.GetVertices()
@@ -119,7 +211,7 @@ FVector UGeometryScriptLibrary_BoxFunctions::GetBoxCorner(const FBox& Box, int C
 	}
 }
 
-FVector UGeometryScriptLibrary_BoxFunctions::GetBoxFaceCenter(const FBox& Box, int FaceIndex, FVector& FaceNormal)
+FVector UGeometryScriptLibrary_BoxFunctions::GetBoxFaceCenter(FBox Box, int FaceIndex, FVector& FaceNormal)
 {
 	FaceIndex = FMath::Clamp(FaceIndex, 0, 5);
 	FVector Center = 0.5 * (Box.Min + Box.Max);
@@ -135,7 +227,7 @@ FVector UGeometryScriptLibrary_BoxFunctions::GetBoxFaceCenter(const FBox& Box, i
 	}
 }
 
-void UGeometryScriptLibrary_BoxFunctions::GetBoxVolumeArea(const FBox& Box, double& Volume, double& SurfaceArea)
+void UGeometryScriptLibrary_BoxFunctions::GetBoxVolumeArea(FBox Box, double& Volume, double& SurfaceArea)
 {
 	FVector Dimensions = Box.GetSize();
 	double AreaXY = Dimensions.X * Dimensions.Y;
@@ -145,7 +237,7 @@ void UGeometryScriptLibrary_BoxFunctions::GetBoxVolumeArea(const FBox& Box, doub
 	Volume = Dimensions.X * Dimensions.Y * Dimensions.Z;
 }
 
-FBox UGeometryScriptLibrary_BoxFunctions::GetExpandedBox(const FBox& Box, const FVector& ExpandBy)
+FBox UGeometryScriptLibrary_BoxFunctions::GetExpandedBox(FBox Box, FVector ExpandBy)
 {
 	FBox Result = Box.ExpandBy(ExpandBy);
 	// ExpandBy with negative expansion factor does not clamp to original box center
@@ -159,48 +251,48 @@ FBox UGeometryScriptLibrary_BoxFunctions::GetExpandedBox(const FBox& Box, const
 	return Result;
 }
 
-FBox UGeometryScriptLibrary_BoxFunctions::GetTransformedBox(const FBox& Box, const FTransform& Transform)
+FBox UGeometryScriptLibrary_BoxFunctions::GetTransformedBox(FBox Box, FTransform Transform)
 {
 	return Box.TransformBy(Transform);
 }
 
 
-bool UGeometryScriptLibrary_BoxFunctions::TestBoxBoxIntersection(const FBox& Box1, const FBox& Box2)
+bool UGeometryScriptLibrary_BoxFunctions::TestBoxBoxIntersection(FBox Box1, FBox Box2)
 {
 	return Box1.Intersect(Box2);
 }
 
 
-FBox UGeometryScriptLibrary_BoxFunctions::FindBoxBoxIntersection(const FBox& Box1, const FBox& Box2, bool& bIsIntersecting)
+FBox UGeometryScriptLibrary_BoxFunctions::FindBoxBoxIntersection(FBox Box1, FBox Box2, bool& bIsIntersecting)
 {
 	bIsIntersecting = Box1.Intersect(Box2);
 	return Box1.Overlap(Box2);
 }
 
-double UGeometryScriptLibrary_BoxFunctions::GetBoxBoxDistance(const FBox& Box1, const FBox& Box2)
+double UGeometryScriptLibrary_BoxFunctions::GetBoxBoxDistance(FBox Box1, FBox Box2)
 {
 	double DistSqr = FMathd::Max(0.0, Box1.ComputeSquaredDistanceToBox(Box2));
 	return FMathd::Sqrt(DistSqr);
 }
 
-bool UGeometryScriptLibrary_BoxFunctions::TestPointInsideBox(const FBox& Box, const FVector& Point, bool bConsiderOnBoxAsInside)
+bool UGeometryScriptLibrary_BoxFunctions::TestPointInsideBox(FBox Box, FVector Point, bool bConsiderOnBoxAsInside)
 {
 	return (bConsiderOnBoxAsInside) ? Box.IsInsideOrOn(Point) : Box.IsInside(Point);
 }
 
-FVector UGeometryScriptLibrary_BoxFunctions::FindClosestPointOnBox(const FBox& Box, const FVector& Point, bool& bIsInside)
+FVector UGeometryScriptLibrary_BoxFunctions::FindClosestPointOnBox(FBox Box, FVector Point, bool& bIsInside)
 {
 	bIsInside = Box.IsInside(Point);
 	return Box.GetClosestPointTo(Point);
 }
 
-double UGeometryScriptLibrary_BoxFunctions::GetBoxPointDistance(const FBox& Box, const FVector& Point)
+double UGeometryScriptLibrary_BoxFunctions::GetBoxPointDistance(FBox Box, FVector Point)
 {
 	double DistSqr = FMathd::Max(0.0, Box.ComputeSquaredDistanceToPoint(Point));
 	return FMathd::Sqrt(DistSqr);
 }
 
-bool UGeometryScriptLibrary_BoxFunctions::TestBoxSphereIntersection(const FBox& Box, const FVector& SphereCenter, double SphereRadius)
+bool UGeometryScriptLibrary_BoxFunctions::TestBoxSphereIntersection(FBox Box, FVector SphereCenter, double SphereRadius)
 {
 	return FMath::SphereAABBIntersection(SphereCenter, SphereRadius*SphereRadius, Box);
 }
diff --git a/Engine/Plugins/Experimental/GeometryScripting/Source/GeometryScriptingCore/Public/GeometryScript/ShapeFunctions.h b/Engine/Plugins/Experimental/GeometryScripting/Source/GeometryScriptingCore/Public/GeometryScript/ShapeFunctions.h
--- a/Engine/Plugins/Experimental/GeometryScripting/Source/GeometryScriptingCore/Public/GeometryScript/ShapeFunctions.h
+++ b/Engine/Plugins/Experimental/GeometryScripting/Source/GeometryScriptingCore/Public/GeometryScript/ShapeFunctions.h
@@ -4,10 +4,63 @@
 
 #include "CoreMinimal.h"
 #include "Kismet/BlueprintFunctionLibrary.h"
+#include "GeometryScriptTypes.h"
 #include "Math/Box.h"
 #include "ShapeFunctions.generated.h"
 
 
+
+UCLASS(meta = (ScriptName = "GeometryScript_Transform"))
+class GEOMETRYSCRIPTINGCORE_API UGeometryScriptLibrary_TransformFunctions : public UBlueprintFunctionLibrary
+{
+	GENERATED_BODY()
+public:
+
+	/**
+	 * Create a Transform at the given Location, with the ZAxis vector as the Z axis
+	 * of the Transform, and the X or Y axis oriented to the Tangent vector, based on
+	 * the bTangentIsX parameter.
+	 */
+	UFUNCTION(BlueprintPure, Category = "GeometryScript|Shapes|Transform")
+	static UPARAM(DisplayName="Transform") FTransform
+	MakeTransformFromZAxis(FVector Location, FVector ZAxis);
+
+	/**
+	 * Create a Transform at the given Location, with the ZAxis vector as the Z axis
+	 * of the Transform, and the X or Y axis oriented to the Tangent vector, based on
+	 * the bTangentIsX parameter.
+	 */
+	UFUNCTION(BlueprintPure, Category = "GeometryScript|Shapes|Transform")
+	static UPARAM(DisplayName="Transform") FTransform
+	MakeTransformFromAxes(FVector Location, FVector ZAxis, FVector TangentAxis, bool bTangentIsX = true);
+
+	/**
+	 * Get the Vector for the direction of the X/Y/Z axis of the Transform, ie the 
+	 * Vector resulting from transforming the unit direction vectors (1,0,0) / etc  
+	 */
+	UFUNCTION(BlueprintPure, Category = "GeometryScript|Shapes|Transform")
+	static UPARAM(DisplayName="Direction") FVector
+	GetTransformAxisVector(FTransform Transform, EGeometryScriptAxis Axis = EGeometryScriptAxis::X);
+
+	/**
+	 * Get the Ray at the Transform Location aligned with the direction of the X/Y/Z axis of the Transform, 
+	 * ie the Direction Vector resulting from transforming the unit direction vectors (1,0,0) / etc  
+	 */
+	UFUNCTION(BlueprintPure, Category = "GeometryScript|Shapes|Transform")
+	static UPARAM(DisplayName="Ray") FRay
+	GetTransformAxisRay(FTransform Transform, EGeometryScriptAxis Axis = EGeometryScriptAxis::X);
+
+	/**
+	 * Get the Plane at the Transform Location with the Plane Normal aligned with the direction of the X/Y/Z axis of the Transform, 
+	 * ie the Direction Vector resulting from transforming the unit direction vectors (1,0,0) / etc  
+	 */
+	UFUNCTION(BlueprintPure, Category = "GeometryScript|Shapes|Transform")
+	static UPARAM(DisplayName="Ray") FPlane
+	GetTransformAxisPlane(FTransform Transform, EGeometryScriptAxis Axis = EGeometryScriptAxis::X);
+
+};
+
+
 UCLASS(meta = (ScriptName = "GeometryScript_Ray"))
 class GEOMETRYSCRIPTINGCORE_API UGeometryScriptLibrary_RayFunctions : public UBlueprintFunctionLibrary
 {
@@ -19,28 +72,28 @@ public:
 	 */
 	UFUNCTION(BlueprintPure, Category = "GeometryScript|Shapes|Ray")
 	static UPARAM(DisplayName="Ray") FRay
-	MakeRayFromPoints(const FVector& A, const FVector& B);
+	MakeRayFromPoints(FVector A, FVector B);
 
 	/**
 	 * Create a Ray from an Origin and Direction, with optionally non-normalized Direction
 	 */
 	UFUNCTION(BlueprintPure, Category = "GeometryScript|Shapes|Ray")
 	static UPARAM(DisplayName="Ray") FRay 
-	MakeRayFromPointDirection(const FVector& Origin, const FVector& Direction, bool bDirectionIsNormalized = true);
+	MakeRayFromPointDirection(FVector Origin, FVector Direction, bool bDirectionIsNormalized = true);
 
 	/**
 	 * Apply the given Transform to the given Ray, or optionally the Transform Inverse, and return the new transformed Ray
 	 */
 	UFUNCTION(BlueprintPure, Category = "GeometryScript|Shapes|Ray", meta=(ScriptMethod))
 	static UPARAM(DisplayName="Transformed Ray") FRay 
-	GetTransformedRay(const FRay& Ray, const FTransform& Transform, bool bInvert = false);
+	GetTransformedRay(FRay Ray, FTransform Transform, bool bInvert = false);
 
 	/**
 	 * Get a Point at the given Distance along the Ray (Origin + Distance*Direction)
 	 */
 	UFUNCTION(BlueprintPure, Category = "GeometryScript|Shapes|Ray", meta=(ScriptMethod))
 	static UPARAM(DisplayName="Point") FVector 
-	GetRayPoint(const FRay& Ray, double Distance);
+	GetRayPoint(FRay Ray, double Distance);
 
 	/**
 	 * Get two points along the ray. 
@@ -49,28 +102,28 @@ public:
 	 */
 	UFUNCTION(BlueprintPure, Category = "GeometryScript|Shapes|Ray", meta=(ScriptMethod, StartDistance="0", EndDistance = "0"))
 	static void 
-	GetRayStartEnd(const FRay& Ray, double StartDistance, double EndDistance, FVector& StartPoint, FVector& EndPoint);
+	GetRayStartEnd(FRay Ray, double StartDistance, double EndDistance, FVector& StartPoint, FVector& EndPoint);
 
 	/**
 	 * Project the given Point onto the closest point along the Ray, and return the Ray Parameter/Distance at that Point
 	 */
 	UFUNCTION(BlueprintPure, Category = "GeometryScript|Shapes|Ray", meta=(ScriptMethod))
 	static UPARAM(DisplayName="Ray Paramater") double 
-	GetRayParameter(const FRay& Ray, const FVector& Point);
+	GetRayParameter(FRay Ray, FVector Point);
 
 	/**
 	 * Get the distance from Point to the closest point on the Ray
 	 */
 	UFUNCTION(BlueprintPure, Category = "GeometryScript|Shapes|Ray", meta=(ScriptMethod))
 	static UPARAM(DisplayName="Distance") double 
-	GetRayPointDistance(const FRay& Ray, const FVector& Point);
+	GetRayPointDistance(FRay Ray, FVector Point);
 
 	/**
 	 * Get the closest point on the Ray to the given Point
 	 */
 	UFUNCTION(BlueprintPure, Category = "GeometryScript|Shapes|Ray", meta=(ScriptMethod))
 	static UPARAM(DisplayName="Closest Point") FVector 
-	GetRayClosestPoint(const FRay& Ray, const FVector& Point);
+	GetRayClosestPoint(FRay Ray, FVector Point);
 
 	/**
 	 * Check if the Ray intersects a Sphere defined by the SphereCenter and SphereRadius.
@@ -84,15 +137,50 @@ public:
 	 */
 	UFUNCTION(BlueprintPure, Category = "GeometryScript|Shapes|Ray", meta=(ScriptMethod))
 	static UPARAM(DisplayName="Intersects") bool 
-	GetRaySphereIntersection(const FRay& Ray, const FVector& SphereCenter, double SphereRadius, double& Distance1, double& Distance2);
+	GetRaySphereIntersection(FRay Ray, FVector SphereCenter, double SphereRadius, double& Distance1, double& Distance2);
 
 	/**
 	 * Check if the Ray intersects a Sphere defined by the SphereCenter and SphereRadius.
 	 * @param HitDistance Distance along the ray (Ray Parameter) to first intersection point with the Box
+	 * @return true if the ray hits the box, and false otherwise
+	 */
+	UFUNCTION(BlueprintPure, Category = "GeometryScript|Shapes|Ray", meta=(ScriptMethod))
+	static UPARAM(DisplayName="Intersects") bool 
+	GetRayBoxIntersection(FRay Ray, FBox Box, double& HitDistance);
+
+	/**
+	 * Find the intersection of a Ray and a Plane
+	 * @param HitDistance the returned Distance along the ray (Ray Parameter) to intersection point with the Plane
+	 * @return true if the ray hits the plane (only false if ray is parallel with plane)
 	 */
 	UFUNCTION(BlueprintPure, Category = "GeometryScript|Shapes|Ray", meta=(ScriptMethod))
 	static UPARAM(DisplayName="Intersects") bool 
-	GetRayBoxIntersection(const FRay& Ray, const FBox& Box, double& HitDistance);
+	GetRayPlaneIntersection(FRay Ray, FPlane Plane, double& HitDistance);
+
+	/**
+	 * Compute the pair of closest points on a 3D Ray and Line. 
+	 * The Line is defined by an Origin and Direction (ie same as a Ray) but extends infinitely in both directions.
+	 * @param RayParameter the Ray Parameter of the closest point on the Ray (range 0, inf)
+	 * @param RayPoint the point on the Ray corresponding to RayParameter
+	 * @param LineParameter the Line parameter of the closest point on the Line (range -inf, inf)
+	 * @param LinePoint the point on the Line corresponding to LineParameter
+	 * @return the minimum distance between the Ray and Line, ie between RayPoint and LinePoint
+	 */
+	UFUNCTION(BlueprintPure, Category = "GeometryScript|Shapes|Ray", meta=(ScriptMethod))
+	static UPARAM(DisplayName="Distance") double
+	GetRayLineClosestPoint(FRay Ray, FVector LineOrigin, FVector LineDirection, double& RayParameter, FVector& RayPoint, double& LineParameter, FVector& LinePoint);
+
+	/**
+	 * Compute the pair of closest points on a 3D Ray and Line Segment
+	 * The Line Segment is defined by its two Endpoints.
+	 * @param RayParameter the Ray Parameter of the closest point on the Ray (range 0, inf)
+	 * @param RayPoint the point on the Ray corresponding to RayParameter
+	 * @param SegPoint the point on the Segment
+	 * @return the minimum distance between the Ray and Segment, ie between RayPoint and SegPoint
+	 */
+	UFUNCTION(BlueprintPure, Category = "GeometryScript|Shapes|Ray", meta=(ScriptMethod))
+	static UPARAM(DisplayName="Distance") double
+	GetRaySegmentClosestPoint(FRay Ray, FVector SegStartPoint, FVector SegEndPoint, double& RayParameter, FVector& RayPoint, FVector& SegPoint);
 
 };
 
@@ -108,20 +196,20 @@ public:
 	 */
 	UFUNCTION(BlueprintPure, Category = "GeometryScript|Shapes|Box")
 	static UPARAM(DisplayName="Box") FBox 
-	MakeBoxFromCenterSize(const FVector& Center, const FVector& Dimensions);
+	MakeBoxFromCenterSize(FVector Center, FVector Dimensions);
 
 	/**
 	 * Create a Box from a Center point and X/Y/Z Extents (Extents are half-dimensions)
 	 */
 	UFUNCTION(BlueprintPure, Category = "GeometryScript|Shapes|Box")
 	static UPARAM(DisplayName="Box") FBox 
-	MakeBoxFromCenterExtents(const FVector& Center, const FVector& Extents);
+	MakeBoxFromCenterExtents(FVector Center, FVector Extents);
 
 	/**
 	 * Get the Center point and X/Y/Z Dimensions of a Box (full dimensions, not Extents)
 	 */
 	UFUNCTION(BlueprintPure, Category = "GeometryScript|Shapes|Box", meta=(ScriptMethod))
-	static void GetBoxCenterSize(const FBox& Box, FVector& Center, FVector& Dimensions);
+	static void GetBoxCenterSize(FBox Box, FVector& Center, FVector& Dimensions);
 
 
 	/**
@@ -131,7 +219,7 @@ public:
 	 */
 	UFUNCTION(BlueprintPure, Category = "GeometryScript|Shapes|Box", meta=(ScriptMethod))
 	static UPARAM(DisplayName="Corner Point") FVector 
-	GetBoxCorner(const FBox& Box, int CornerIndex);
+	GetBoxCorner(FBox Box, int CornerIndex);
 
 	/**
 	 * Get the position of the center of a face of the Box. Faces are indexed from 0 to 5,
@@ -140,13 +228,13 @@ public:
 	 */
 	UFUNCTION(BlueprintPure, Category = "GeometryScript|Shapes|Box", meta=(ScriptMethod))
 	static UPARAM(DisplayName="Center Point") FVector 
-	GetBoxFaceCenter(const FBox& Box, int FaceIndex, FVector& FaceNormal);
+	GetBoxFaceCenter(FBox Box, int FaceIndex, FVector& FaceNormal);
 
 	/**
 	 * Get the Volume and Surface Area of a Box
 	 */
 	UFUNCTION(BlueprintPure, Category = "GeometryScript|Shapes|Box", meta=(ScriptMethod))
-	static void GetBoxVolumeArea(const FBox& Box, double& Volume, double& SurfaceArea);
+	static void GetBoxVolumeArea(FBox Box, double& Volume, double& SurfaceArea);
 
 	/**
 	 * Get the input Box expanded by adding the ExpandBy parameter to both the Min and Max.
@@ -154,21 +242,21 @@ public:
 	 */
 	UFUNCTION(BlueprintPure, Category = "GeometryScript|Shapes|Box", meta=(ScriptMethod))
 	static UPARAM(DisplayName="Expanded Box") FBox 
-	GetExpandedBox(const FBox& Box, const FVector& ExpandBy);
+	GetExpandedBox(FBox Box, FVector ExpandBy);
 
 	/**
 	 * Apply the input Transform to the corners of the input Box, and return the new Box containing those points
 	 */
 	UFUNCTION(BlueprintPure, Category = "GeometryScript|Shapes|Box", meta=(ScriptMethod))
 	static UPARAM(DisplayName="Transformed Box") FBox 
-	GetTransformedBox(const FBox& Box, const FTransform& Transform);
+	GetTransformedBox(FBox Box, FTransform Transform);
 
 	/**
 	 * Test if Box1 and Box2 intersect
 	 */
 	UFUNCTION(BlueprintPure, Category = "GeometryScript|Shapes|Box", meta=(ScriptMethod))
 	static UPARAM(DisplayName="Intersects") bool 
-	TestBoxBoxIntersection(const FBox& Box1, const FBox& Box2);
+	TestBoxBoxIntersection(FBox Box1, FBox Box2);
 
 	/**
 	 * Find the Box formed by the intersection of Box1 and Box2
@@ -176,14 +264,14 @@ public:
 	 */
 	UFUNCTION(BlueprintPure, Category = "GeometryScript|Shapes|Box", meta=(ScriptMethod))
 	static UPARAM(DisplayName="Intersection Box") FBox 
-	FindBoxBoxIntersection(const FBox& Box1, const FBox& Box2, bool& bIsIntersecting);
+	FindBoxBoxIntersection(FBox Box1, FBox Box2, bool& bIsIntersecting);
 
 	/**
 	 * Calculate the minimum distance between Box1 and Box2
 	 */
 	UFUNCTION(BlueprintPure, Category = "GeometryScript|Shapes|Box", meta=(ScriptMethod))
 	static UPARAM(DisplayName="Distance") double 
-	GetBoxBoxDistance(const FBox& Box1, const FBox& Box2);
+	GetBoxBoxDistance(FBox Box1, FBox Box2);
 
 	/**
 	 * Test if a Point is inside the Box, returning true if so, otherwise false
@@ -191,7 +279,7 @@ public:
 	 */
 	UFUNCTION(BlueprintPure, Category = "GeometryScript|Shapes|Box", meta=(ScriptMethod))
 	static UPARAM(DisplayName="Is Inside") bool 
-	TestPointInsideBox(const FBox& Box, const FVector& Point, bool bConsiderOnBoxAsInside = true);
+	TestPointInsideBox(FBox Box, FVector Point, bool bConsiderOnBoxAsInside = true);
 
 	/**
 	 * Find the point on the faces of the Box that is closest to the input Point.
@@ -200,21 +288,21 @@ public:
 	 */
 	UFUNCTION(BlueprintPure, Category = "GeometryScript|Shapes|Box", meta=(ScriptMethod))
 	static UPARAM(DisplayName="Closest Point") FVector 
-	FindClosestPointOnBox(const FBox& Box, const FVector& Point, bool& bIsInside);
+	FindClosestPointOnBox(FBox Box, FVector Point, bool& bIsInside);
 
 	/**
 	 * Calculate the minimum distance between the Box and the Point
 	 */
 	UFUNCTION(BlueprintPure, Category = "GeometryScript|Shapes|Box", meta=(ScriptMethod))
 	static UPARAM(DisplayName="Distance") double 
-	GetBoxPointDistance(const FBox& Box, const FVector& Point);
+	GetBoxPointDistance(FBox Box, FVector Point);
 
 	/**
 	 * Check if the Box intersects a Sphere defined by the SphereCenter and SphereRadius
 	 */
 	UFUNCTION(BlueprintPure, Category = "GeometryScript|Shapes|Box", meta=(ScriptMethod))
 	static UPARAM(DisplayName="Intersects") bool 
-	TestBoxSphereIntersection(const FBox& Box, const FVector& SphereCenter, double SphereRadius);
+	TestBoxSphereIntersection(FBox Box, FVector SphereCenter, double SphereRadius);
 
 };
 
diff --git a/Engine/Plugins/Experimental/GeometryScripting/Source/GeometryScriptingEditor/Private/EditorDynamicMeshUtilityFunctions.cpp b/Engine/Plugins/Experimental/GeometryScripting/Source/GeometryScriptingEditor/Private/EditorDynamicMeshUtilityFunctions.cpp
new file mode 100644
--- /dev/null
+++ b/Engine/Plugins/Experimental/GeometryScripting/Source/GeometryScriptingEditor/Private/EditorDynamicMeshUtilityFunctions.cpp
@@ -0,0 +1,88 @@
+// Copyright Epic Games, Inc. All Rights Reserved.
+
+#include "GeometryScript/EditorDynamicMeshUtilityFunctions.h"
+
+#include "UDynamicMesh.h"
+#include "DynamicMesh/DynamicMesh3.h"
+
+#include "Misc/ITransaction.h"
+#include "ScopedTransaction.h"
+
+
+namespace UE
+{
+namespace Geometry
+{
+
+// hidden representation of dynamic mesh change used by FDynamicMeshChangeContainer
+struct FDynamicMeshChangeContainerInternals
+{
+	TSharedPtr<FDynamicMesh3> InitialMesh;
+};
+
+}
+}
+
+
+
+UDynamicMesh* UGeometryScriptLibrary_EditorDynamicMeshFunctions::BeginTrackedMeshChange(UDynamicMesh* TargetMesh, FDynamicMeshChangeContainer& ChangeContainer)
+{
+	if ( TargetMesh == nullptr )
+	{
+		UE_LOG(LogGeometry, Warning, TEXT("BeginTrackedMeshChange: TargetMesh was null"));
+		return TargetMesh;
+	}
+	if (ChangeContainer.bIsActive)
+	{
+		UE_LOG(LogGeometry, Warning, TEXT("BeginTrackedMeshChange: ChangeContainer was marked as active, discarding previous change"));
+	}
+
+	// todo: A FDynamicMesh3ChangeTracker would often be more efficient. However currently this seems to result
+	// in some problems...
+	ChangeContainer.ChangeInternals = MakeShared<UE::Geometry::FDynamicMeshChangeContainerInternals>();
+	ChangeContainer.ChangeInternals->InitialMesh = MakeShared<UE::Geometry::FDynamicMesh3>(TargetMesh->GetMeshRef());
+
+	ChangeContainer.bIsActive = true;
+	ChangeContainer.TargetMesh = TargetMesh;
+
+	return TargetMesh;
+}
+
+UDynamicMesh* UGeometryScriptLibrary_EditorDynamicMeshFunctions::EmitTrackedMeshChange(UDynamicMesh* TargetMesh, FDynamicMeshChangeContainer& ChangeContainer)
+{
+	if ( TargetMesh == nullptr )
+	{
+		UE_LOG(LogGeometry, Warning, TEXT("BeginTrackedMeshChange: TargetMesh was null"));
+		return TargetMesh;
+	}
+
+	if (GUndo == nullptr)
+	{
+		UE_LOG(LogGeometry, Warning, TEXT("EmitTrackedMeshChange: No Transaction is actively open, ignoring."));
+		return TargetMesh;
+	}
+
+	if (ChangeContainer.bIsActive == false || ChangeContainer.ChangeInternals.IsValid() == false)
+	{
+		UE_LOG(LogGeometry, Warning, TEXT("EmitTrackedMeshChange: ChangeContainer was not active, ignoring"));
+		return TargetMesh;
+	}
+	if (ChangeContainer.TargetMesh != TargetMesh)
+	{
+		UE_LOG(LogGeometry, Warning, TEXT("EmitTrackedMeshChange: ChangeContainer was created for a different UDynamicMesh, ignoring"));
+		return TargetMesh;
+	}
+
+	TUniquePtr<FMeshReplacementChange> Change = MakeUnique<FMeshReplacementChange>(
+		ChangeContainer.ChangeInternals->InitialMesh, MakeShared<FDynamicMesh3>(TargetMesh->GetMeshRef()) );
+
+	// container has been used and now needs to be reset
+	ChangeContainer.bIsActive = false;
+	ChangeContainer.TargetMesh = nullptr;
+	ChangeContainer.ChangeInternals.Reset();
+
+	GUndo->StoreUndo(TargetMesh, MoveTemp(Change));
+
+	return TargetMesh;
+}
+
diff --git a/Engine/Plugins/Experimental/GeometryScripting/Source/GeometryScriptingEditor/Public/GeometryScript/EditorDynamicMeshUtilityFunctions.h b/Engine/Plugins/Experimental/GeometryScripting/Source/GeometryScriptingEditor/Public/GeometryScript/EditorDynamicMeshUtilityFunctions.h
new file mode 100644
--- /dev/null
+++ b/Engine/Plugins/Experimental/GeometryScripting/Source/GeometryScriptingEditor/Public/GeometryScript/EditorDynamicMeshUtilityFunctions.h
@@ -0,0 +1,65 @@
+// Copyright Epic Games, Inc. All Rights Reserved.
+
+#pragma once
+
+#include "CoreMinimal.h"
+#include "Kismet/BlueprintFunctionLibrary.h"
+#include "GeometryScript/GeometryScriptTypes.h"
+#include "EditorDynamicMeshUtilityFunctions.generated.h"
+
+namespace UE::Geometry
+{ 
+	struct FDynamicMeshChangeContainerInternals;
+}
+
+
+/**
+ * FDynamicMeshChangeContainer is a temporary struct usable in Blueprints to
+ * emit "change" transactions for a UDynamicMesh. The internals are hidden
+ * and this type is only intended to be used with the BeginTrackedMeshChange()
+ * and EmitTrackedMeshChange() functions
+ */
+USTRUCT(BlueprintType, meta = (DisplayName = "Dynamic Mesh ChangeTracker"))
+struct GEOMETRYSCRIPTINGEDITOR_API FDynamicMeshChangeContainer
+{
+	GENERATED_BODY()
+public:
+	bool bIsActive = false;
+
+	// This is only used as a token to ensure that the ChangeContainer is not applied to a different UDynamicMesh,
+	// so it is not a uproperty/etc
+	UDynamicMesh* TargetMesh;
+
+	// internals are only available inside the cpp file and are likely to change
+	TSharedPtr<UE::Geometry::FDynamicMeshChangeContainerInternals> ChangeInternals;
+};
+
+
+
+
+
+UCLASS(meta = (ScriptName = "GeometryScript_EditorDynamicMeshUtil"))
+class GEOMETRYSCRIPTINGEDITOR_API UGeometryScriptLibrary_EditorDynamicMeshFunctions : public UBlueprintFunctionLibrary
+{
+	GENERATED_BODY()
+public:
+
+	/**
+	 * Save current state of TargetMesh so that an undoable/redoable Change can be emitted 
+	 * after TargetMesh is modified, using EmitTrackedMeshChange().
+	 * @param ChangeTracker output structure containing initial TargetMesh state
+	 */
+	UFUNCTION(BlueprintCallable, Category = "DynamicMesh|Changes")
+	static UPARAM(DisplayName = "Dynamic Mesh") UDynamicMesh* 
+	BeginTrackedMeshChange(UDynamicMesh* TargetMesh, FDynamicMeshChangeContainer& ChangeTracker);
+
+	/**
+	 * Emit an undo/redo Change for a modified TargetMesh, based on the ChangeTracker information that was 
+	 * saved (via call to BeginTrackedMeshChange) before TargetMesh was modified. This function must be
+	 * called in the context of a Transaction (ie BeginTransaction / EndTransaction pair)
+	 */
+	UFUNCTION(BlueprintCallable, Category = "DynamicMesh|Changes")
+	static UPARAM(DisplayName = "Dynamic Mesh") UDynamicMesh* 
+	EmitTrackedMeshChange(UDynamicMesh* TargetMesh, UPARAM(ref) FDynamicMeshChangeContainer& ChangeTracker);
+
+};
\ No newline at end of file
