GeometryScript: add MeshCreateBoneWeights function. Fix bug in CreateNewSkeletalMeshAssetFromMesh function where Skeleton was never being passed on to creation function.

ModelingOperators:
Change FSkinBindingOp to use GetOrCreateSkinWeightsAttribute intsead of just directly accessing the SkinWeightAttribute (as it could be nullptr and this was not handled)
Fix various issues in CreateSkeletalMeshAsset(). Register MeshDescription attributes before trying to convert from DynamicMesh. Make sure there is always at least one Material being passed (otherwise ensures are hit in Material UVData check). Actually return created SkeletalMesh (!)

diff --git a/Engine/Plugins/Experimental/GeometryScripting/Source/GeometryScriptingCore/Private/MeshBoneWeightFunctions.cpp b/Engine/Plugins/Experimental/GeometryScripting/Source/GeometryScriptingCore/Private/MeshBoneWeightFunctions.cpp
--- a/Engine/Plugins/Experimental/GeometryScripting/Source/GeometryScriptingCore/Private/MeshBoneWeightFunctions.cpp
+++ b/Engine/Plugins/Experimental/GeometryScripting/Source/GeometryScriptingCore/Private/MeshBoneWeightFunctions.cpp
@@ -54,14 +54,14 @@ ReturnType SimpleMeshBoneWeightEdit(
 	ReturnType RetVal = DefaultValue;
 	if (Mesh)
 	{
-		Mesh->EditMesh([&](FDynamicMesh3& ReadMesh)
+		Mesh->EditMesh([&](FDynamicMesh3& EditMesh)
 		{
-			if (ReadMesh.HasAttributes())
+			if (EditMesh.HasAttributes())
 			{
-				if ( FDynamicMeshVertexSkinWeightsAttribute* BoneWeights = ReadMesh.Attributes()->GetSkinWeightsAttribute(Profile.GetProfileName()) )
+				if ( FDynamicMeshVertexSkinWeightsAttribute* BoneWeights = EditMesh.Attributes()->GetSkinWeightsAttribute(Profile.GetProfileName()) )
 				{
 					bIsValidBoneWeights = true;
-					RetVal = EditFunc(ReadMesh, *BoneWeights);
+					RetVal = EditFunc(EditMesh, *BoneWeights);
 				}
 			}
 		}, EDynamicMeshChangeType::GeneralEdit, EDynamicMeshAttributeChangeFlags::Unknown, false);
@@ -85,6 +85,41 @@ UDynamicMesh* UGeometryScriptLibrary_MeshBoneWeightFunctions::MeshHasBoneWeights
 }
 
 
+UDynamicMesh* UGeometryScriptLibrary_MeshBoneWeightFunctions::MeshCreateBoneWeights(
+	UDynamicMesh* TargetMesh,
+	bool& bProfileExisted,
+	bool bReplaceExistingProfile,
+	FGeometryScriptBoneWeightProfile Profile)
+{
+	bProfileExisted = false;
+	if (TargetMesh)
+	{
+		TargetMesh->EditMesh([&](FDynamicMesh3& EditMesh)
+		{
+			if (EditMesh.HasAttributes() == false)
+			{
+				EditMesh.EnableAttributes();
+			}
+
+			FDynamicMeshVertexSkinWeightsAttribute *Attribute = EditMesh.Attributes()->GetSkinWeightsAttribute(Profile.GetProfileName());
+			bProfileExisted = (Attribute != nullptr);
+			if ( Attribute == nullptr || bReplaceExistingProfile)
+			{
+				if ( bReplaceExistingProfile && bProfileExisted )
+				{
+					EditMesh.Attributes()->RemoveSkinWeightsAttribute(Profile.GetProfileName());
+				}
+
+				Attribute = new FDynamicMeshVertexSkinWeightsAttribute(&EditMesh);
+				EditMesh.Attributes()->AttachSkinWeightsAttribute(Profile.GetProfileName(), Attribute);
+			}			
+		}, EDynamicMeshChangeType::GeneralEdit, EDynamicMeshAttributeChangeFlags::Unknown, false);
+	}
+	return TargetMesh;
+}
+
+
+
 UDynamicMesh* UGeometryScriptLibrary_MeshBoneWeightFunctions::GetMaxBoneWeightIndex(
 	UDynamicMesh* TargetMesh,
 	bool& bHasBoneWeights,
diff --git a/Engine/Plugins/Experimental/GeometryScripting/Source/GeometryScriptingCore/Public/GeometryScript/MeshBoneWeightFunctions.h b/Engine/Plugins/Experimental/GeometryScripting/Source/GeometryScriptingCore/Public/GeometryScript/MeshBoneWeightFunctions.h
--- a/Engine/Plugins/Experimental/GeometryScripting/Source/GeometryScriptingCore/Public/GeometryScript/MeshBoneWeightFunctions.h
+++ b/Engine/Plugins/Experimental/GeometryScripting/Source/GeometryScriptingCore/Public/GeometryScript/MeshBoneWeightFunctions.h
@@ -93,6 +93,21 @@ public:
 		UDynamicMesh* TargetMesh,
 		bool& bHasBoneWeights,
 		FGeometryScriptBoneWeightProfile Profile = FGeometryScriptBoneWeightProfile() );
+
+	/**
+	 * Create a new BoneWeights attribute on the TargetMesh, if it does not already exist. If it does exist, 
+	 * and bReplaceExistingProfile is passed as true, the attribute will be removed and re-added, to reset it. 
+	 * @param bProfileExisted will be returned true if the requested bone weight profile already existed
+	 * @param bReplaceExistingProfile if true, if the Profile already exists, it is reset
+	 * @param Profile identifier for the bone/skin weight profile
+	 */
+	UFUNCTION(BlueprintPure, Category = "GeometryScript|MeshQueries|BoneWeights", meta=(ScriptMethod))
+	static UPARAM(DisplayName = "Target Mesh") UDynamicMesh* 
+	MeshCreateBoneWeights( 
+		UDynamicMesh* TargetMesh,
+		bool& bProfileExisted,
+		bool bReplaceExistingProfile = false,
+		FGeometryScriptBoneWeightProfile Profile = FGeometryScriptBoneWeightProfile() );
 	
 	/**
 	 * Determine the largest bone weight index that exists on the Mesh
diff --git a/Engine/Plugins/Experimental/GeometryScripting/Source/GeometryScriptingEditor/Private/CreateNewAssetUtilityFunctions.cpp b/Engine/Plugins/Experimental/GeometryScripting/Source/GeometryScriptingEditor/Private/CreateNewAssetUtilityFunctions.cpp
--- a/Engine/Plugins/Experimental/GeometryScripting/Source/GeometryScriptingEditor/Private/CreateNewAssetUtilityFunctions.cpp
+++ b/Engine/Plugins/Experimental/GeometryScripting/Source/GeometryScriptingEditor/Private/CreateNewAssetUtilityFunctions.cpp
@@ -239,6 +239,11 @@ USkeletalMesh* UGeometryScriptLibrary_CreateNewAssetFunctions::CreateNewSkeletal
 		AppendError(Debug, EGeometryScriptErrorType::InvalidInputs, LOCTEXT("CreateNewSkeletalMeshAssetFromMesh_InvalidInput3", "CreateNewSkeletalMeshAssetFromMesh: FromDynamicMesh has no skin weight attributes"));
 		return nullptr;
 	}
+	if (InSkeleton == nullptr)
+	{
+		AppendError(Debug, EGeometryScriptErrorType::InvalidInputs, LOCTEXT("CreateNewSkeletalMeshAssetFromMesh_InvalidSkeleton", "CreateNewSkeletalMeshAssetFromMesh: Skeleton is Null"));
+		return nullptr;
+	}
 	
 	// todo: other safety checks
 
@@ -246,6 +251,7 @@ USkeletalMesh* UGeometryScriptLibrary_CreateNewAssetFunctions::CreateNewSkeletal
 
 	FSkeletalMeshAssetOptions AssetOptions;
 	AssetOptions.NewAssetPath = AssetPathAndName;
+	AssetOptions.Skeleton = InSkeleton;
 
 	AssetOptions.NumSourceModels = 1;
 
diff --git a/Engine/Plugins/Runtime/MeshModelingToolset/Source/ModelingComponentsEditorOnly/Private/AssetUtils/CreateSkeletalMeshUtil.cpp b/Engine/Plugins/Runtime/MeshModelingToolset/Source/ModelingComponentsEditorOnly/Private/AssetUtils/CreateSkeletalMeshUtil.cpp
--- a/Engine/Plugins/Runtime/MeshModelingToolset/Source/ModelingComponentsEditorOnly/Private/AssetUtils/CreateSkeletalMeshUtil.cpp
+++ b/Engine/Plugins/Runtime/MeshModelingToolset/Source/ModelingComponentsEditorOnly/Private/AssetUtils/CreateSkeletalMeshUtil.cpp
@@ -10,7 +10,7 @@
 #include "StaticToSkeletalMeshConverter.h"
 #include "UObject/Package.h"
 #include "UObject/UObjectGlobals.h"
-
+#include "SkeletalMeshAttributes.h"
 
 UE::AssetUtils::ECreateSkeletalMeshResult UE::AssetUtils::CreateSkeletalMeshAsset(
 	const FSkeletalMeshAssetOptions& Options,
@@ -71,6 +71,8 @@ UE::AssetUtils::ECreateSkeletalMeshResult UE::AssetUtils::CreateSkeletalMeshAsse
 			{
 				ConstructedMeshDescriptions.AddDefaulted();
 				FDynamicMeshToMeshDescription Converter;
+				FSkeletalMeshAttributes Attributes(ConstructedMeshDescriptions.Last());
+				Attributes.Register();
 				Converter.Convert(DynamicMesh, ConstructedMeshDescriptions.Last(), !Options.bEnableRecomputeTangents);
 				MeshDescriptions.Add(&ConstructedMeshDescriptions.Last());
 			}
@@ -89,6 +91,14 @@ UE::AssetUtils::ECreateSkeletalMeshResult UE::AssetUtils::CreateSkeletalMeshAsse
 		{
 			Materials.Add(FSkeletalMaterial(MaterialInterface));
 		}
+		MaterialView = Materials;
+	}
+
+	// ensure there is at least one material
+	if (Materials.Num() == 0)
+	{
+		Materials.Add(FSkeletalMaterial());
+		MaterialView = Materials;
 	}
 	
 	if (!FStaticToSkeletalMeshConverter::InitializeSkeletalMeshFromMeshDescriptions(
@@ -107,5 +117,6 @@ UE::AssetUtils::ECreateSkeletalMeshResult UE::AssetUtils::CreateSkeletalMeshAsse
 		Options.Skeleton->SetPreviewMesh(NewSkeletalMesh);
 	}
 	
+	ResultsOut.SkeletalMesh = NewSkeletalMesh;
 	return ECreateSkeletalMeshResult::Ok;
 }
diff --git a/Engine/Plugins/Runtime/MeshModelingToolset/Source/ModelingOperators/Private/SkinningOps/SkinBindingOp.cpp b/Engine/Plugins/Runtime/MeshModelingToolset/Source/ModelingOperators/Private/SkinningOps/SkinBindingOp.cpp
--- a/Engine/Plugins/Runtime/MeshModelingToolset/Source/ModelingOperators/Private/SkinningOps/SkinBindingOp.cpp
+++ b/Engine/Plugins/Runtime/MeshModelingToolset/Source/ModelingOperators/Private/SkinningOps/SkinBindingOp.cpp
@@ -380,7 +380,7 @@ void FSkinBindingOp::CreateSkinWeights_DirectDistance(
 {
 	using namespace AnimationCore;
 
-	FDynamicMeshVertexSkinWeightsAttribute *SkinWeights = InMesh.Attributes()->GetSkinWeightsAttribute(ProfileName);
+	FDynamicMeshVertexSkinWeightsAttribute *SkinWeights = GetOrCreateSkinWeightsAttribute(InMesh, ProfileName);
 		
 	const int32 NumVertices = InMesh.VertexCount();
 
@@ -445,7 +445,7 @@ void FSkinBindingOp::CreateSkinWeights_GeodesicVoxel(
 {
 	using namespace AnimationCore;
 
-	FDynamicMeshVertexSkinWeightsAttribute *SkinWeights = InMesh.Attributes()->GetSkinWeightsAttribute(ProfileName);
+	FDynamicMeshVertexSkinWeightsAttribute *SkinWeights = GetOrCreateSkinWeightsAttribute(InMesh, ProfileName);
 		
 	const int32 NumVertices = InMesh.VertexCount();
 
