Geometry Script.  Use DisplayName to update UV nodes with standardized terminology of "channels"

diff --git a/Engine/Plugins/Experimental/GeometryScripting/Source/GeometryScriptingCore/Public/GeometryScript/MeshQueryFunctions.h b/Engine/Plugins/Experimental/GeometryScripting/Source/GeometryScriptingCore/Public/GeometryScript/MeshQueryFunctions.h
--- a/Engine/Plugins/Experimental/GeometryScripting/Source/GeometryScriptingCore/Public/GeometryScript/MeshQueryFunctions.h
+++ b/Engine/Plugins/Experimental/GeometryScripting/Source/GeometryScriptingCore/Public/GeometryScript/MeshQueryFunctions.h
@@ -101,14 +101,17 @@ public:
 
 
 
-	UFUNCTION(BlueprintPure, Category = "GeometryScript|MeshQueries", meta=(ScriptMethod))
-	static UPARAM(DisplayName = "Num UV Sets") int32 GetNumUVSets( UDynamicMesh* TargetMesh );
+	UFUNCTION(BlueprintPure, Category = "GeometryScript|MeshQueries", meta=(ScriptMethod), DisplayName = "Get Num UV Channels")
+	static UPARAM(DisplayName = "Num UV Channels") int32 GetNumUVSets( UDynamicMesh* TargetMesh );
 
 	UFUNCTION(BlueprintPure, Category = "GeometryScript|MeshQueries", meta=(ScriptMethod))
-	static UPARAM(DisplayName = "Bounding Box") FBox2D GetUVSetBoundingBox( UDynamicMesh* TargetMesh, int UVSetIndex, bool& bIsValidUVSet, bool& bUVSetIsEmpty );
+	static UPARAM(DisplayName = "Bounding Box") FBox2D GetUVSetBoundingBox( UDynamicMesh* TargetMesh, 
+																			UPARAM(DisplayName = "UV Channel") int UVSetIndex, 
+																			UPARAM(DisplayName = "Is Valid UV Channel") bool& bIsValidUVSet, 
+																			UPARAM(DisplayName = "UV Channel Is Empty") bool& bUVSetIsEmpty );
 
 	UFUNCTION(BlueprintPure, Category = "GeometryScript|MeshQueries", meta=(ScriptMethod))
-	static void GetTriangleUVs( UDynamicMesh* TargetMesh, int32 UVSetIndex, int32 TriangleID, FVector2D& UV1, FVector2D& UV2, FVector2D& UV3, bool& bHaveValidUVs );
+	static void GetTriangleUVs( UDynamicMesh* TargetMesh, UPARAM(DisplayName = "UV Channel") int32 UVSetIndex, int32 TriangleID, FVector2D& UV1, FVector2D& UV2, FVector2D& UV3, bool& bHaveValidUVs );
 
 
 
diff --git a/Engine/Plugins/Experimental/GeometryScripting/Source/GeometryScriptingCore/Public/GeometryScript/MeshUVFunctions.h b/Engine/Plugins/Experimental/GeometryScripting/Source/GeometryScriptingCore/Public/GeometryScript/MeshUVFunctions.h
--- a/Engine/Plugins/Experimental/GeometryScripting/Source/GeometryScriptingCore/Public/GeometryScript/MeshUVFunctions.h
+++ b/Engine/Plugins/Experimental/GeometryScripting/Source/GeometryScriptingCore/Public/GeometryScript/MeshUVFunctions.h
@@ -146,19 +146,19 @@ class GEOMETRYSCRIPTINGCORE_API UGeometryScriptLibrary_MeshUVFunctions : public
 	GENERATED_BODY()
 public:
 
-	UFUNCTION(BlueprintCallable, Category = "GeometryScript|UVs", meta=(ScriptMethod))
+	UFUNCTION(BlueprintCallable, Category = "GeometryScript|UVs", meta=(ScriptMethod), DisplayName = "Set Num UV Channels")
 	static UPARAM(DisplayName = "Target Mesh") UDynamicMesh* 
 	SetNumUVSets( 
 		UDynamicMesh* TargetMesh, 
-		int NumUVSets,
+		UPARAM(DisplayName = "Num UV Channels") int NumUVSets,
 		UGeometryScriptDebug* Debug = nullptr );
 
-	UFUNCTION(BlueprintCallable, Category = "GeometryScript|UVs", meta=(ScriptMethod))
+	UFUNCTION(BlueprintCallable, Category = "GeometryScript|UVs", meta=(ScriptMethod), DisplayName = "Copy UV Channel")
 	static UPARAM(DisplayName = "Target Mesh") UDynamicMesh* 
 	CopyUVSet( 
 		UDynamicMesh* TargetMesh, 
-		int FromUVSet,
-		int ToUVSet,
+		UPARAM(DisplayName = "From UV Channel") int FromUVSet,
+		UPARAM(DisplayName = "To UV Channel")   int ToUVSet,
 		UGeometryScriptDebug* Debug = nullptr );
 
 
@@ -166,7 +166,7 @@ public:
 	static UPARAM(DisplayName = "Target Mesh") UDynamicMesh* 
 	SetMeshTriangleUVs( 
 		UDynamicMesh* TargetMesh, 
-		int UVSetIndex,
+		UPARAM(DisplayName = "UV Channel") int UVSetIndex,
 		int TriangleID, 
 		FGeometryScriptUVTriangle UVs,
 		bool& bIsValidTriangle, 
@@ -177,7 +177,7 @@ public:
 	static UPARAM(DisplayName = "Target Mesh") UDynamicMesh* 
 	TranslateMeshUVs( 
 		UDynamicMesh* TargetMesh, 
-		int UVSetIndex,
+		UPARAM(DisplayName = "UV Channel") int UVSetIndex,
 		FVector2D Translation,
 		FGeometryScriptMeshSelection Selection,
 		UGeometryScriptDebug* Debug = nullptr );
@@ -186,7 +186,7 @@ public:
 	static UPARAM(DisplayName = "Target Mesh") UDynamicMesh* 
 	ScaleMeshUVs( 
 		UDynamicMesh* TargetMesh, 
-		int UVSetIndex,
+		UPARAM(DisplayName = "UV Channel") int UVSetIndex,
 		FVector2D Scale,
 		FVector2D ScaleOrigin,
 		FGeometryScriptMeshSelection Selection,
@@ -196,7 +196,7 @@ public:
 	static UPARAM(DisplayName = "Target Mesh") UDynamicMesh* 
 	RotateMeshUVs( 
 		UDynamicMesh* TargetMesh, 
-		int UVSetIndex,
+		UPARAM(DisplayName = "UV Channel") int UVSetIndex,
 		float RotationAngle,
 		FVector2D RotationOrigin,
 		FGeometryScriptMeshSelection Selection,
@@ -210,7 +210,7 @@ public:
 	static UPARAM(DisplayName = "Target Mesh") UDynamicMesh* 
 	SetMeshUVsFromPlanarProjection( 
 		UDynamicMesh* TargetMesh, 
-		int UVSetIndex,
+		UPARAM(DisplayName = "UV Channel") int UVSetIndex,
 		FTransform PlaneTransform,
 		FGeometryScriptMeshSelection Selection,
 		UGeometryScriptDebug* Debug = nullptr );
@@ -220,7 +220,7 @@ public:
 	static UPARAM(DisplayName = "Target Mesh") UDynamicMesh* 
 	SetMeshUVsFromBoxProjection( 
 		UDynamicMesh* TargetMesh, 
-		int UVSetIndex,
+		UPARAM(DisplayName = "UV Channel") int UVSetIndex,
 		FTransform BoxTransform,
 		FGeometryScriptMeshSelection Selection,
 		int MinIslandTriCount = 2,
@@ -231,7 +231,7 @@ public:
 	static UPARAM(DisplayName = "Target Mesh") UDynamicMesh* 
 	SetMeshUVsFromCylinderProjection( 
 		UDynamicMesh* TargetMesh, 
-		int UVSetIndex,
+		UPARAM(DisplayName = "UV Channel") int UVSetIndex,
 		FTransform CylinderTransform,
 		FGeometryScriptMeshSelection Selection,
 		float SplitAngle = 45.0,
@@ -243,7 +243,7 @@ public:
 	static UPARAM(DisplayName = "Target Mesh") UDynamicMesh* 
 	RecomputeMeshUVs( 
 		UDynamicMesh* TargetMesh, 
-		int UVSetIndex,
+		UPARAM(DisplayName = "UV Channel") int UVSetIndex,
 		FGeometryScriptRecomputeUVsOptions Options,
 		FGeometryScriptMeshSelection Selection,
 		UGeometryScriptDebug* Debug = nullptr );
@@ -253,7 +253,7 @@ public:
 	static UPARAM(DisplayName = "Target Mesh") UDynamicMesh* 
 	RepackMeshUVs( 
 		UDynamicMesh* TargetMesh, 
-		int UVSetIndex,
+		UPARAM(DisplayName = "UV Channel") int UVSetIndex,
 		FGeometryScriptRepackUVsOptions RepackOptions,
 		UGeometryScriptDebug* Debug = nullptr );
 
@@ -261,7 +261,7 @@ public:
 	static UPARAM(DisplayName = "Target Mesh") UDynamicMesh* 
 	AutoGeneratePatchBuilderMeshUVs( 
 		UDynamicMesh* TargetMesh, 
-		int UVSetIndex,
+		UPARAM(DisplayName = "UV Channel") int UVSetIndex,
 		FGeometryScriptPatchBuilderOptions Options,
 		UGeometryScriptDebug* Debug = nullptr );
 
@@ -269,19 +269,19 @@ public:
 	static UPARAM(DisplayName = "Target Mesh") UDynamicMesh* 
 	AutoGenerateXAtlasMeshUVs( 
 		UDynamicMesh* TargetMesh, 
-		int UVSetIndex,
+		UPARAM(DisplayName = "UV Channel") int UVSetIndex,
 		FGeometryScriptXAtlasOptions Options,
 		UGeometryScriptDebug* Debug = nullptr );
 
 	/**
 	 * Compute information about dimensions and areas for a UV Set of a Mesh, with an optional Mesh Selection
-	 * @param UVSetIndex index of UV Set to query
+	 * @param UVSetIndex index of UV Channel to query
 	 * @param Selection subset of triangles to process, whole mesh is used if selection is not provided
 	 * @param MeshArea output 3D area of queried triangles
 	 * @param UVArea output 2D UV-space area of queried triangles
 	 * @param MeshBounds output 3D bounding box of queried triangles
 	 * @param UVBounds output 2D UV-space bounding box of queried triangles
-	 * @param bIsValidUVSet output flag set to false if UVSetIndex does not exist on the target mesh. In this case Areas and Bounds are not initialized.
+	 * @param bIsValidUVSet output flag set to false if UV Channel does not exist on the target mesh. In this case Areas and Bounds are not initialized.
 	 * @param bFoundUnsetUVs output flag set to true if any of the queried triangles do not have valid UVs set
 	 * @param bOnlyIncludeValidUVTris if true, only triangles with valid UVs are included in 3D Mesh Area/Bounds
 	 */
@@ -289,7 +289,7 @@ public:
 	static UPARAM(DisplayName = "Copy From Mesh") UDynamicMesh* 
 	GetMeshUVSizeInfo(  
 		UDynamicMesh* TargetMesh, 
-		int UVSetIndex,
+		UPARAM(DisplayName = "UV Channel") int UVSetIndex,
 		FGeometryScriptMeshSelection Selection,
 		double& MeshArea,
 		double& UVArea,
@@ -302,20 +302,20 @@ public:
 
 
 	/**
-	 * Get a list of single vertex UVs for each mesh vertex in the TargetMesh, derived from the specified UV Overlay.
-	 * The UV Overlay may store multiple UVs for a single vertex (along UV seams)
+	 * Get a list of single vertex UVs for each mesh vertex in the TargetMesh, derived from the specified UV Channel.
+	 * The UV Channel may store multiple UVs for a single vertex (along UV seams)
 	 * In such cases an arbitrary UV will be stored for that vertex, and bHasSplitUVs will be returned as true
-	 * @param UVSetIndex index of UV Set to read
+	 * @param UVSetIndex index of UV Channel to read
 	 * @param UVList output UV list will be stored here. Size will be equal to the MaxVertexID of TargetMesh  (not the VertexCount!)
-	 * @param bIsValidUVSet will be set to true if the UV Overlay was valid
+	 * @param bIsValidUVSet will be set to true if the UV Channel was valid
 	 * @param bHasVertexIDGaps will be set to true if some vertex indices in TargetMesh were invalid, ie MaxVertexID > VertexCount 
-	 * @param bHasSplitUVs will be set to true if there were split UVs in the UV overlay
+	 * @param bHasSplitUVs will be set to true if there were split UVs in the UV Channel
 	 */
 	UFUNCTION(BlueprintCallable, Category = "GeometryScript|UVs", meta=(ScriptMethod))
 	static UPARAM(DisplayName = "Target Mesh") UDynamicMesh* 
 	GetMeshPerVertexUVs( 
 		UDynamicMesh* TargetMesh, 
-		int UVSetIndex,
+		UPARAM(DisplayName = "UV Channel") int UVSetIndex,
 		FGeometryScriptUVList& UVList, 
 		bool& bIsValidUVSet,
 		bool& bHasVertexIDGaps,
@@ -324,23 +324,23 @@ public:
 
 
 	/**
-	 * Copy the 2D UVs from the given UVSetIndex in CopyFromMesh to the 3D vertex positions in CopyToUVMesh,
-	 * with the triangle mesh topology defined by the UV Set. Generally this "UV Mesh" topology will not
+	 * Copy the 2D UVs from the given UV Channel in CopyFromMesh to the 3D vertex positions in CopyToUVMesh,
+	 * with the triangle mesh topology defined by the UV Channel. Generally this "UV Mesh" topology will not
 	 * be the same as the 3D mesh topology. PolyGroup IDs and Material IDs are preserved in the UVMesh.
 	 * 
 	 * 2D UV Positions are copied to 3D as (X, Y, 0) 
 	 * 
-	 * CopyMeshToMeshUVLayer will copy the 3D UV Mesh back to the UV Set. This pair of functions can
+	 * CopyMeshToMeshUVChannel will copy the 3D UV Mesh back to the UV Channel. This pair of functions can
 	 * then be used to implement UV generation/editing via other mesh functions.
 	 * 
 	 * @param bInvalidTopology will be returned true if any topological issues were found
 	 * @param bIsValidUVSet will be returned false if UVSetIndex is not available
 	 */
-	UFUNCTION(BlueprintCallable, Category = "GeometryScript|UVs", meta=(ScriptMethod))
+	UFUNCTION(BlueprintCallable, Category = "GeometryScript|UVs", meta=(ScriptMethod), DisplayName="Copy Mesh UV Channel To Mesh")
 	static UPARAM(DisplayName = "Copy From Mesh") UDynamicMesh* 
 	CopyMeshUVLayerToMesh(  
 		UDynamicMesh* CopyFromMesh, 
-		int UVSetIndex,
+		UPARAM(DisplayName = "UV Channel") int UVSetIndex,
 		UPARAM(DisplayName = "Copy To UV Mesh", ref) UDynamicMesh* CopyToUVMesh, 
 		UPARAM(DisplayName = "Copy To UV Mesh") UDynamicMesh*& CopyToUVMeshOut,
 		bool& bInvalidTopology,
@@ -348,26 +348,26 @@ public:
 		UGeometryScriptDebug* Debug = nullptr);
 
 	/**
-	 * Transfer the 3D vertex positions and triangles of CopyFromUVMesh to the given UV Layer identified by ToUVSetIndex of CopyToMesh.
+	 * Transfer the 3D vertex positions and triangles of CopyFromUVMesh to the given UV Channel identified by ToUVChannel of CopyToMesh.
 	 * 3D positions (X,Y,Z) will be copied as UV positions (X,Y), ie Z is ignored.
 	 * 
 	 * bOnlyUVPositions controls whether only UV positions will be updated, or if the UV topology will be fully replaced.
-	 * When false, CopyFromUVMesh must currently have a MaxVertexID <= that of the UV Layer MaxElementID
+	 * When false, CopyFromUVMesh must currently have a MaxVertexID <= that of the UV Channel MaxElementID
 	 * When true, CopyFromUVMesh must currently have a MaxTriangleID <= that of CopyToMesh
 	 * 
 	 * @param bInvalidTopology will be returned true if any topological inconsistencies are found (but the operation will generally continue)
-	 * @param bIsValidUVSet will be returned false if ToUVSetIndex is not available
+	 * @param bIsValidUVSet will be returned false if To UV Channel is not available
 	 * @param bOnlyUVPositions if true, only (valid, matching) UV positions are updated, a full new UV topology is created
 	 */
-	UFUNCTION(BlueprintCallable, Category = "GeometryScript|UVs", meta=(ScriptMethod))
+	UFUNCTION(BlueprintCallable, Category = "GeometryScript|UVs", meta=(ScriptMethod), DisplayName="Copy Mesh To Mesh UV Channel")
 	static UPARAM(DisplayName = "Copy From Mesh") UDynamicMesh* 
 	CopyMeshToMeshUVLayer(  
 		UDynamicMesh* CopyFromUVMesh, 
-		int ToUVSetIndex,
+		UPARAM(DisplayName = "To UV Channel")  int ToUVSetIndex,
 		UPARAM(DisplayName = "Copy To Mesh", ref) UDynamicMesh* CopyToMesh, 
 		UPARAM(DisplayName = "Copy To Mesh") UDynamicMesh*& CopyToMeshOut,
 		bool& bFoundTopologyErrors,
-		bool& bIsValidUVSet,
+		UPARAM(DisplayName = "Is Valid UV Channel") bool& bIsValidUVSet,
 		bool bOnlyUVPositions = true,
 		UGeometryScriptDebug* Debug = nullptr);
 
