GeometryScripting: Fix missing UPROPERTY macro on FGeometryScriptFillHolesOptions::bDeleteIsolatedTriangles member

diff --git a/Engine/Plugins/Experimental/GeometryScripting/Source/GeometryScriptingCore/Public/GeometryScript/MeshRepairFunctions.h b/Engine/Plugins/Experimental/GeometryScripting/Source/GeometryScriptingCore/Public/GeometryScript/MeshRepairFunctions.h
--- a/Engine/Plugins/Experimental/GeometryScripting/Source/GeometryScriptingCore/Public/GeometryScript/MeshRepairFunctions.h
+++ b/Engine/Plugins/Experimental/GeometryScripting/Source/GeometryScriptingCore/Public/GeometryScript/MeshRepairFunctions.h
@@ -55,6 +55,7 @@ public:
 	EGeometryScriptFillHolesMethod FillMethod = EGeometryScriptFillHolesMethod::Automatic;
 
 	/** Delete floating, disconnected triangles, as they produce a "hole" that cannot be filled */
+	UPROPERTY(EditAnywhere, BlueprintReadWrite, Category = Options)
 	bool bDeleteIsolatedTriangles = true;
 };
 
