GeometryScript: enable vertex colors in target SkeletalMesh in CopyMeshToSkeletalMesh if they exist in the FSkeletalMeshImportData

diff --git a/Engine/Plugins/Experimental/GeometryScripting/Source/GeometryScriptingCore/Private/MeshAssetFunctions.cpp b/Engine/Plugins/Experimental/GeometryScripting/Source/GeometryScriptingCore/Private/MeshAssetFunctions.cpp
--- a/Engine/Plugins/Experimental/GeometryScripting/Source/GeometryScriptingCore/Private/MeshAssetFunctions.cpp
+++ b/Engine/Plugins/Experimental/GeometryScripting/Source/GeometryScriptingCore/Private/MeshAssetFunctions.cpp
@@ -8,6 +8,7 @@
 
 #include "Engine/EngineTypes.h"
 #include "Engine/StaticMesh.h"
+#include "Engine/SkinnedAssetCommon.h"
 #include "Engine/SkeletalMesh.h"
 #include "Engine/StaticMeshActor.h"
 #include "StaticMeshAttributes.h"
@@ -488,6 +489,56 @@ UDynamicMesh* UGeometryScriptLibrary_StaticMeshFunctions::CopyMeshFromSkeletalMe
 	return ToDynamicMesh;
 }
 
+#if WITH_EDITOR
+namespace UELocal 
+{
+
+// this is identical to UE::AssetUtils::GenerateNewMaterialSlotName except it takes a TArray<FSkeletalMaterial>
+// instead of a TArray<FStaticMaterial>. It seems likely that we will need SkeletalMeshMaterialUtil.h soon,
+// at that point this function can be moved there
+static FName GenerateNewMaterialSlotName(
+	const TArray<FSkeletalMaterial>& ExistingMaterials,
+	UMaterialInterface* SlotMaterial,
+	int32 NewSlotIndex)
+{
+	FString MaterialName = (SlotMaterial) ? SlotMaterial->GetName() : TEXT("Material");
+	FName BaseName(MaterialName);
+
+	bool bFound = false;
+	for (const FSkeletalMaterial& Mat : ExistingMaterials)
+	{
+		if (Mat.MaterialSlotName == BaseName || Mat.ImportedMaterialSlotName == BaseName)
+		{
+			bFound = true;
+			break;
+		}
+	}
+	if (bFound == false && SlotMaterial != nullptr)
+	{
+		return BaseName;
+	}
+
+	bFound = true;
+	while (bFound)
+	{
+		bFound = false;
+
+		BaseName = FName(FString::Printf(TEXT("%s_%d"), *MaterialName, NewSlotIndex++));
+		for (const FSkeletalMaterial& Mat : ExistingMaterials)
+		{
+			if (Mat.MaterialSlotName == BaseName || Mat.ImportedMaterialSlotName == BaseName)
+			{
+				bFound = true;
+				break;
+			}
+		}
+	}
+
+	return BaseName;
+}
+}
+#endif
+
 
 UDynamicMesh* UGeometryScriptLibrary_StaticMeshFunctions::CopyMeshToSkeletalMesh(
 		UDynamicMesh* FromDynamicMesh, 
@@ -550,15 +601,55 @@ UDynamicMesh* UGeometryScriptLibrary_StaticMeshFunctions::CopyMeshToSkeletalMesh
 
 	FSkeletalMeshImportData SkeletalMeshImportData = 
 		FSkeletalMeshImportData::CreateFromMeshDescription(MeshDescription);
+
+	// if we are replacing materials, construct a FSkeletalMaterial list with unique slot names
+	TArray<FSkeletalMaterial> NewMaterials;
+	if (Options.bReplaceMaterials)
+	{
+		bool bHaveSlotNames = (Options.NewMaterialSlotNames.Num() == Options.NewMaterials.Num());
+
+		for (int32 k = 0; k < Options.NewMaterials.Num(); ++k)
+		{
+			FSkeletalMaterial NewMaterial;
+			NewMaterial.MaterialInterface = Options.NewMaterials[k];
+			FName UseSlotName = (bHaveSlotNames && Options.NewMaterialSlotNames[k] != NAME_None) ? Options.NewMaterialSlotNames[k] :
+				UELocal::GenerateNewMaterialSlotName(NewMaterials, NewMaterial.MaterialInterface, k);
+
+			NewMaterial.MaterialSlotName = UseSlotName;
+			NewMaterial.ImportedMaterialSlotName = UseSlotName;
+			NewMaterial.UVChannelData = FMeshUVChannelInfo(1.f);		// this avoids an ensure in  UStaticMesh::GetUVChannelData
+			NewMaterials.Add(NewMaterial);
+		}
+
+		// copy the material set to the SkeletalMeshImportData
+		SkeletalMeshImportData.Materials.SetNum(NewMaterials.Num());
+		for (int32 k = 0; k < NewMaterials.Num(); ++k)
+		{
+			SkeletalMeshImportData.Materials[k].Material = NewMaterials[k].MaterialInterface;
+			SkeletalMeshImportData.Materials[k].MaterialImportName = NewMaterials[k].MaterialSlotName.ToString();
+		}
+	}
+
+
 	ToSkeletalMeshAsset->SaveLODImportedData(TargetLOD.LODIndex, SkeletalMeshImportData);
 
+	// configure vertex color setup in the Asset
+    ToSkeletalMeshAsset->SetHasVertexColors(SkeletalMeshImportData.bHasVertexColors);
+#if WITH_EDITORONLY_DATA
+    ToSkeletalMeshAsset->SetVertexColorGuid(SkeletalMeshImportData.bHasVertexColors ? FGuid::NewGuid() : FGuid() );
+#endif
+
+	// update materials on the Asset
+	if (Options.bReplaceMaterials)
+	{
+		ToSkeletalMeshAsset->SetMaterials(NewMaterials);
+	}
+
 	// Make sure the mesh builder knows it's the latest variety, so that the render data gets
 	// properly rebuilt.
 	ToSkeletalMeshAsset->SetLODImportedDataVersions(TargetLOD.LODIndex, ESkeletalMeshGeoImportVersions::LatestVersion, ESkeletalMeshSkinningImportVersions::LatestVersion);
 	ToSkeletalMeshAsset->SetUseLegacyMeshDerivedDataKey(false);
 
-	// TODO: Options.bReplaceMaterials support?
-
 	if (Options.bDeferMeshPostEditChange == false)
 	{
 		ToSkeletalMeshAsset->PostEditChange();
