Removed deprecated ApplyRecursivePNTessellation function.

diff --git a/Engine/Plugins/Experimental/GeometryScripting/Source/GeometryScriptingCore/Private/MeshSubdivideFunctions.cpp b/Engine/Plugins/Experimental/GeometryScripting/Source/GeometryScriptingCore/Private/MeshSubdivideFunctions.cpp
--- a/Engine/Plugins/Experimental/GeometryScripting/Source/GeometryScriptingCore/Private/MeshSubdivideFunctions.cpp
+++ b/Engine/Plugins/Experimental/GeometryScripting/Source/GeometryScriptingCore/Private/MeshSubdivideFunctions.cpp
@@ -17,45 +17,6 @@ using namespace UE::Geometry;
 
 #define LOCTEXT_NAMESPACE "UGeometryScriptLibrary_MeshSubdivideFunctions"
 
-
-UDynamicMesh* UGeometryScriptLibrary_MeshSubdivideFunctions::ApplyRecursivePNTessellation(
-	UDynamicMesh* TargetMesh,
-	FGeometryScriptPNTessellateOptions Options,
-	int NumIterations,
-	UGeometryScriptDebug* Debug)
-{
-	if (TargetMesh == nullptr)
-	{
-		UE::Geometry::AppendError(Debug, EGeometryScriptErrorType::InvalidInputs, LOCTEXT("ApplyPNTessellation_InvalidInput", "ApplyPNTessellation: TargetMesh is Null"));
-		return TargetMesh;
-	}
-	if (NumIterations <= 0)
-	{
-		return TargetMesh;
-	}
-
-	TargetMesh->EditMesh([&](FDynamicMesh3& EditMesh) 
-	{
-		FPNTriangles Tessellator(&EditMesh);
-		
-		// Convert the number of recursive subdivisions to the equivalent tessellation level to make sure that we 
-		// produce the same topology. 
-		int TessellationLevel = static_cast<int>(FMath::RoundHalfFromZero(FMath::Exp2((double)NumIterations))) - 1;
-		Tessellator.TessellationLevel = TessellationLevel;
-		Tessellator.Compute();
-
-		if (Options.bRecomputeNormals && EditMesh.HasAttributes() && EditMesh.Attributes()->PrimaryNormals() != nullptr)
-		{
-			FMeshNormals MeshNormals(&EditMesh);
-			MeshNormals.RecomputeOverlayNormals(EditMesh.Attributes()->PrimaryNormals(), true, true);
-			MeshNormals.CopyToOverlay(EditMesh.Attributes()->PrimaryNormals(), false);
-		}
-
-	}, EDynamicMeshChangeType::GeneralEdit, EDynamicMeshAttributeChangeFlags::Unknown, false);
-
-	return TargetMesh;
-}
-
 UDynamicMesh* UGeometryScriptLibrary_MeshSubdivideFunctions::ApplyPNTessellation(
 	UDynamicMesh* TargetMesh,
 	FGeometryScriptPNTessellateOptions Options,
diff --git a/Engine/Plugins/Experimental/GeometryScripting/Source/GeometryScriptingCore/Public/GeometryScript/MeshSubdivideFunctions.h b/Engine/Plugins/Experimental/GeometryScripting/Source/GeometryScriptingCore/Public/GeometryScript/MeshSubdivideFunctions.h
--- a/Engine/Plugins/Experimental/GeometryScripting/Source/GeometryScriptingCore/Public/GeometryScript/MeshSubdivideFunctions.h
+++ b/Engine/Plugins/Experimental/GeometryScripting/Source/GeometryScriptingCore/Public/GeometryScript/MeshSubdivideFunctions.h
@@ -53,15 +53,6 @@ class GEOMETRYSCRIPTINGCORE_API UGeometryScriptLibrary_MeshSubdivideFunctions :
 	GENERATED_BODY()
 public:
 
-	UE_DEPRECATED(5.0, "Use 'Apply PN Tessellation' instead; this deprecated version recursively subdivided the mesh NumIterations times while the new version splits every triangle into (TessellationLevel+1)^2 smaller triangles.")
-	UFUNCTION(BlueprintCallable, Category = "GeometryScript|Subdivide", meta = (ScriptMethod, DeprecatedFunction, DeprecationMessage = "Use 'Apply PN Tessellation' instead; this deprecated version recursively subdivided the mesh NumIterations times while the new version splits every triangle into (TessellationLevel+1)^2 smaller triangles."))
-	static UPARAM(DisplayName = "Target Mesh") UDynamicMesh* 
-	ApplyRecursivePNTessellation(
-		UDynamicMesh* TargetMesh,
-		FGeometryScriptPNTessellateOptions Options,
-		int NumIterations = 3,
-		UGeometryScriptDebug* Debug = nullptr );
-
 	UFUNCTION(BlueprintCallable, Category = "GeometryScript|Subdivide", meta=(ScriptMethod))
 	static UPARAM(DisplayName = "Target Mesh") UDynamicMesh* 
 	ApplyPNTessellation(
