GeometryScript: add barycentric-related basic query functions ComputeTriangleBarycentricCoords, GetInterpolatedTrianglePosition, GetInterpolatedTriangleUV, GetInterpolatedTriangleNormal, GetInterpolatedTriangleNormalTangents, GetInterpolatedTriangleVertexColor.
Also add functions GetHasTriangleNormals, GetTriangleNormals, GetTriangleNormalTangents, GetHasVertexColors, and GetTriangleVertexColors for completeness/parity w/ other functions

diff --git a/Engine/Plugins/Experimental/GeometryScripting/Source/GeometryScriptingCore/Private/MeshQueryFunctions.cpp b/Engine/Plugins/Experimental/GeometryScripting/Source/GeometryScriptingCore/Private/MeshQueryFunctions.cpp
--- a/Engine/Plugins/Experimental/GeometryScripting/Source/GeometryScriptingCore/Private/MeshQueryFunctions.cpp
+++ b/Engine/Plugins/Experimental/GeometryScripting/Source/GeometryScriptingCore/Private/MeshQueryFunctions.cpp
@@ -225,6 +225,25 @@ void UGeometryScriptLibrary_MeshQueryFunctions::GetTrianglePositions(UDynamicMes
 	});
 }
 
+UDynamicMesh* UGeometryScriptLibrary_MeshQueryFunctions::GetInterpolatedTrianglePosition(UDynamicMesh* TargetMesh, int32 TriangleID, FVector BarycentricCoords, bool& bIsValidTriangle, FVector& InterpolatedPosition)
+{
+	InterpolatedPosition = FVector::Zero();
+	bIsValidTriangle = SimpleMeshQuery<bool>(TargetMesh, false, [&](const FDynamicMesh3& Mesh) 
+	{
+		if (Mesh.IsTriangle(TriangleID))
+		{
+			FVector A,B,C;
+			Mesh.GetTriVertices<FVector>(TriangleID, A, B, C);
+			InterpolatedPosition = BarycentricCoords.X*A + BarycentricCoords.Y*B + BarycentricCoords.Z*C;
+			return true;
+		}
+		return false;
+	});
+	return TargetMesh;
+}
+
+
+
 FVector UGeometryScriptLibrary_MeshQueryFunctions::GetTriangleFaceNormal(UDynamicMesh* TargetMesh, int32 TriangleID, bool& bIsValidTriangle)
 {
 	return SimpleMeshQuery<FVector>(TargetMesh, FVector::ZeroVector, [&](const FDynamicMesh3& Mesh) {
@@ -234,6 +253,28 @@ FVector UGeometryScriptLibrary_MeshQueryFunctions::GetTriangleFaceNormal(UDynami
 }
 
 
+UDynamicMesh* UGeometryScriptLibrary_MeshQueryFunctions::ComputeTriangleBarycentricCoords(UDynamicMesh* TargetMesh, int32 TriangleID, bool& bIsValidTriangle, FVector Point, FVector& Vertex1, FVector& Vertex2, FVector& Vertex3, FVector& BarycentricCoords)
+{
+	BarycentricCoords = SimpleMeshQuery<FVector>(TargetMesh, FVector::Zero(), [&](const FDynamicMesh3& Mesh) {
+		bIsValidTriangle = Mesh.IsTriangle(TriangleID);
+		if (bIsValidTriangle)
+		{
+			Mesh.GetTriVertices(TriangleID, Vertex1, Vertex2, Vertex3);
+			if ( VectorUtil::NormalDirection(Vertex1, Vertex2, Vertex3).SquaredLength() > FMathd::ZeroTolerance )
+			{
+				FVector BaryCoords = VectorUtil::BarycentricCoords(Point, Vertex1, Vertex2, Vertex3);
+				// Point may not be inside the triangle in which case the coords are invalid, but we want to return something that won't cause later interpolations to explode...
+				BaryCoords.X = FMathd::Clamp(BaryCoords.X, 0, 1);
+				BaryCoords.Y = FMathd::Clamp(BaryCoords.Y, 0, 1);
+				BaryCoords.Z = FMathd::Clamp(BaryCoords.Z, 0, 1);
+				return BaryCoords;
+			}
+		}
+		return FVector::Zero();
+	});
+	return TargetMesh;
+}
+
 
 int32 UGeometryScriptLibrary_MeshQueryFunctions::GetVertexCount( UDynamicMesh* TargetMesh )
 {
@@ -368,8 +409,195 @@ void UGeometryScriptLibrary_MeshQueryFunctions::GetTriangleUVs(UDynamicMesh* Tar
 	});
 }
 
+UDynamicMesh* UGeometryScriptLibrary_MeshQueryFunctions::GetInterpolatedTriangleUV(UDynamicMesh* TargetMesh, UPARAM(DisplayName = "UV Channel") int32 UVSetIndex, int32 TriangleID, FVector BarycentricCoords, bool& bHaveValidUVs, FVector2D& InterpolatedUV)
+{
+	InterpolatedUV = FVector2D::Zero();
+	bHaveValidUVs = SimpleMeshUVSetQuery<bool>(TargetMesh, UVSetIndex, bHaveValidUVs, false, 
+		[&](const FDynamicMesh3& Mesh, const FDynamicMeshUVOverlay& UVSet) 
+	{
+		if (Mesh.IsTriangle(TriangleID) && UVSet.IsSetTriangle(TriangleID))
+		{
+			FVector2f A, B, C;
+			UVSet.GetTriElements(TriangleID, A, B, C);
+			InterpolatedUV = BarycentricCoords.X*(FVector2D)A + BarycentricCoords.Y*(FVector2D)B + BarycentricCoords.Z*(FVector2D)C;
+			return true;
+		}
+		return false;
+	});
+	return TargetMesh;
+}
+
+
+bool UGeometryScriptLibrary_MeshQueryFunctions::GetHasTriangleNormals( UDynamicMesh* TargetMesh )
+{
+	return SimpleMeshQuery<bool>(TargetMesh, false, [&](const FDynamicMesh3& Mesh) {
+		return (Mesh.HasAttributes() && Mesh.Attributes()->PrimaryNormals() != nullptr);
+	});
+}
+
+
+UDynamicMesh* UGeometryScriptLibrary_MeshQueryFunctions::GetTriangleNormals(UDynamicMesh* TargetMesh, int32 TriangleID, FVector& Normal1, FVector& Normal2, FVector& Normal3, bool& bTriHasValidNormals)
+{
+	bTriHasValidNormals = SimpleMeshQuery<bool>(TargetMesh, false, [&](const FDynamicMesh3& Mesh) 
+	{
+		if (Mesh.HasAttributes())
+		{
+			if (const FDynamicMeshNormalOverlay* Normals = Mesh.Attributes()->PrimaryNormals())
+			{
+				if ( Normals->IsSetTriangle(TriangleID) )
+				{
+					FVector3f A,B,C;
+					Normals->GetTriElements(TriangleID, A,B,C);
+					Normal1 = (FVector)A;
+					Normal2 = (FVector)B;
+					Normal3 = (FVector)C;
+					return true;
+				}
+			}
+		}
+		return false;
+	});
+	return TargetMesh;
+}
 
 
+UDynamicMesh* UGeometryScriptLibrary_MeshQueryFunctions::GetInterpolatedTriangleNormal(UDynamicMesh* TargetMesh, int32 TriangleID, FVector BarycentricCoords, bool& bTriHasValidNormals, FVector& InterpolatedNormal)
+{
+	bTriHasValidNormals = false;
+	InterpolatedNormal = SimpleMeshQuery<FVector>(TargetMesh, FVector::UnitZ(), [&](const FDynamicMesh3& Mesh)
+	{
+		if (Mesh.HasAttributes())
+		{
+			if (const FDynamicMeshNormalOverlay* Normals = Mesh.Attributes()->PrimaryNormals())
+			{
+				if ( Normals->IsSetTriangle(TriangleID) )
+				{
+					bTriHasValidNormals = true;
+					FVector3f A,B,C;
+					Normals->GetTriElements(TriangleID, A,B,C);
+					return BarycentricCoords.X*(FVector)A + BarycentricCoords.Y*(FVector)B + BarycentricCoords.Z*(FVector)C;
+				}
+			}
+		}
+		return FVector::UnitZ();
+	});
+	return TargetMesh;
+}
+
+
+UDynamicMesh* UGeometryScriptLibrary_MeshQueryFunctions::GetTriangleNormalTangents(UDynamicMesh* TargetMesh, int32 TriangleID, 
+	bool& bTriHasValidElements, FGeometryScriptTriangle& NormalsTri, FGeometryScriptTriangle& TangentsTri, FGeometryScriptTriangle& BiTangentsTri)
+{
+	bTriHasValidElements = SimpleMeshQuery<bool>(TargetMesh, false, [&](const FDynamicMesh3& Mesh) 
+	{
+		if (Mesh.HasAttributes() && Mesh.Attributes()->HasTangentSpace() )
+		{
+			const FDynamicMeshNormalOverlay* Normals = Mesh.Attributes()->PrimaryNormals();
+			const FDynamicMeshNormalOverlay* Tangents = Mesh.Attributes()->PrimaryTangents();
+			const FDynamicMeshNormalOverlay* BiTangents = Mesh.Attributes()->PrimaryBiTangents();
+			if ( Normals->IsSetTriangle(TriangleID) && Tangents->IsSetTriangle(TriangleID) && BiTangents->IsSetTriangle(TriangleID) )
+			{
+				FVector3f A,B,C;
+				Normals->GetTriElements(TriangleID, A,B,C);
+				NormalsTri.Vector0 = (FVector)A; NormalsTri.Vector1 = (FVector)B; NormalsTri.Vector2 = (FVector)C;
+				Tangents->GetTriElements(TriangleID, A,B,C);
+				TangentsTri.Vector0 = (FVector)A; TangentsTri.Vector1 = (FVector)B; TangentsTri.Vector2 = (FVector)C;
+				BiTangents->GetTriElements(TriangleID, A,B,C);
+				BiTangentsTri.Vector0 = (FVector)A; BiTangentsTri.Vector1 = (FVector)B; BiTangentsTri.Vector2 = (FVector)C;
+				return true;
+			}
+		}
+		return false;
+	});
+	return TargetMesh;
+}
+
+
+UDynamicMesh* UGeometryScriptLibrary_MeshQueryFunctions::GetInterpolatedTriangleNormalTangents(UDynamicMesh* TargetMesh, int32 TriangleID, FVector BarycentricCoords, 
+	bool& bTriHasValidElements, FVector& InterpolatedNormal, FVector& InterpolatedTangent, FVector& InterpolatedBiTangent)
+{
+	bTriHasValidElements = SimpleMeshQuery<bool>(TargetMesh, false, [&](const FDynamicMesh3& Mesh)
+	{
+		if (Mesh.HasAttributes() && Mesh.Attributes()->HasTangentSpace() )
+		{
+			const FDynamicMeshNormalOverlay* Normals = Mesh.Attributes()->PrimaryNormals();
+			const FDynamicMeshNormalOverlay* Tangents = Mesh.Attributes()->PrimaryTangents();
+			const FDynamicMeshNormalOverlay* BiTangents = Mesh.Attributes()->PrimaryBiTangents();
+			if ( Normals->IsSetTriangle(TriangleID) && Tangents->IsSetTriangle(TriangleID) && BiTangents->IsSetTriangle(TriangleID) )
+			{
+				FVector3f A,B,C;
+				Normals->GetTriElements(TriangleID, A,B,C);
+				InterpolatedNormal = BarycentricCoords.X*(FVector)A + BarycentricCoords.Y*(FVector)B + BarycentricCoords.Z*(FVector)C;
+				Tangents->GetTriElements(TriangleID, A,B,C);
+				InterpolatedTangent = BarycentricCoords.X*(FVector)A + BarycentricCoords.Y*(FVector)B + BarycentricCoords.Z*(FVector)C;
+				BiTangents->GetTriElements(TriangleID, A,B,C);
+				InterpolatedBiTangent = BarycentricCoords.X*(FVector)A + BarycentricCoords.Y*(FVector)B + BarycentricCoords.Z*(FVector)C;
+				return true;
+			}
+		}
+		return false;
+	});
+	return TargetMesh;
+}
+
+
+
+bool UGeometryScriptLibrary_MeshQueryFunctions::GetHasVertexColors( UDynamicMesh* TargetMesh )
+{
+	return SimpleMeshQuery<bool>(TargetMesh, false, [&](const FDynamicMesh3& Mesh) {
+		return (Mesh.HasAttributes() && Mesh.Attributes()->PrimaryColors() != nullptr);
+	});
+}
+
+
+UDynamicMesh* UGeometryScriptLibrary_MeshQueryFunctions::GetTriangleVertexColors(UDynamicMesh* TargetMesh, int32 TriangleID, FLinearColor& Color1, FLinearColor& Color2, FLinearColor& Color3, bool& bHaveValidVertexColors)
+{
+	bHaveValidVertexColors = SimpleMeshQuery<bool>(TargetMesh, false, [&](const FDynamicMesh3& Mesh) 
+	{
+		if (Mesh.HasAttributes())
+		{
+			if (const FDynamicMeshColorOverlay* Colors = Mesh.Attributes()->PrimaryColors())
+			{
+				if ( Colors->IsSetTriangle(TriangleID) )
+				{
+					FVector4f A,B,C;
+					Colors->GetTriElements(TriangleID, A,B,C);
+					Color1 = (FLinearColor)A;
+					Color2 = (FLinearColor)B;
+					Color3 = (FLinearColor)C;
+					return true;
+				}
+			}
+		}
+		return false;
+	});
+	return TargetMesh;
+}
+
+
+UDynamicMesh* UGeometryScriptLibrary_MeshQueryFunctions::GetInterpolatedTriangleVertexColor(UDynamicMesh* TargetMesh, int32 TriangleID, FVector BarycentricCoords, FLinearColor DefaultColor, bool& bHaveValidVertexColors, FLinearColor& InterpolatedColor)
+{
+	bHaveValidVertexColors = false;
+	InterpolatedColor = SimpleMeshQuery<FLinearColor>(TargetMesh, DefaultColor, [&](const FDynamicMesh3& Mesh) 
+	{
+		if (Mesh.HasAttributes())
+		{
+			if (const FDynamicMeshColorOverlay* Colors = Mesh.Attributes()->PrimaryColors())
+			{
+				if ( Colors->IsSetTriangle(TriangleID) )
+				{
+					FVector4f A,B,C;
+					Colors->GetTriElements(TriangleID, A,B,C);
+					FVector4f Interpolated = BarycentricCoords.X*A + BarycentricCoords.Y*B + BarycentricCoords.Z*C;
+					bHaveValidVertexColors = true;
+					return (FLinearColor)Interpolated;
+				}
+			}
+		}
+		return DefaultColor;
+	});
+	return TargetMesh;
+}
 
 
 bool UGeometryScriptLibrary_MeshQueryFunctions::GetHasMaterialIDs( UDynamicMesh* TargetMesh )
diff --git a/Engine/Plugins/Experimental/GeometryScripting/Source/GeometryScriptingCore/Public/GeometryScript/MeshQueryFunctions.h b/Engine/Plugins/Experimental/GeometryScripting/Source/GeometryScriptingCore/Public/GeometryScript/MeshQueryFunctions.h
--- a/Engine/Plugins/Experimental/GeometryScripting/Source/GeometryScriptingCore/Public/GeometryScript/MeshQueryFunctions.h
+++ b/Engine/Plugins/Experimental/GeometryScripting/Source/GeometryScriptingCore/Public/GeometryScript/MeshQueryFunctions.h
@@ -71,9 +71,44 @@ public:
 	UFUNCTION(BlueprintPure, Category = "GeometryScript|MeshQueries", meta=(ScriptMethod))
 	static void GetTrianglePositions( UDynamicMesh* TargetMesh, int32 TriangleID, bool& bIsValidTriangle, FVector& Vertex1, FVector& Vertex2, FVector& Vertex3 );
 
+	/**
+	 * Compute the interpolated Position (A*Vertex1 + B*Vertex2 + C*Vertex3), where (A,B,C)=BarycentricCoords and the Vertex positions are taken
+	 * from the specified TriangleID of the TargetMesh. 
+	 * @param bIsValidTriangle will be returned true if TriangleID exists in TargetMesh, and otherwise will be returned false
+	 */
+	UFUNCTION(BlueprintCallable, Category = "GeometryScript|MeshQueries", meta=(ScriptMethod))
+	static UPARAM(DisplayName = "Target Mesh") UDynamicMesh* 
+	GetInterpolatedTrianglePosition( 
+		UDynamicMesh* TargetMesh, 
+		int32 TriangleID, 
+		FVector BarycentricCoords, 
+		bool& bIsValidTriangle,
+		FVector& InterpolatedPosition );
+
+
 	UFUNCTION(BlueprintPure, Category = "GeometryScript|MeshQueries", meta=(ScriptMethod))
 	static UPARAM(DisplayName = "Normal") FVector GetTriangleFaceNormal( UDynamicMesh* TargetMesh, int32 TriangleID, bool& bIsValidTriangle );
 
+	/**
+	 * Compute the barycentric coordinates (A,B,C) of the Point relative to the specified TriangleID of the TargetMesh.
+	 * The properties of barycentric coordinates are such that A,B,C are all positive, A+B+C=1.0, and A*Vertex1 + B*Vertex2 + C*Vertex3 = Point.
+	 * So, the barycentric coordinates can be used to smoothly interpolate/blend any other per-triangle-vertex quantities.
+	 * The Point must lie in the plane of the Triangle, otherwise the coordinates are somewhat meaningless (but clamped to 0-1 range to avoid catastrophic errors)
+	 * The Positions of the Triangle Vertices are also returned for convenience (similar to GetTrianglePositions)
+	 * @param bIsValidTriangle will be returned true if TriangleID exists in TargetMesh, and otherwise will be returned false
+	 */
+	UFUNCTION(BlueprintCallable, Category = "GeometryScript|MeshQueries", meta=(ScriptMethod))
+	static UPARAM(DisplayName = "Target Mesh") UDynamicMesh* 
+	ComputeTriangleBarycentricCoords( 
+		UDynamicMesh* TargetMesh, 
+		int32 TriangleID, 
+		bool& bIsValidTriangle, 
+		FVector Point, 
+		FVector& Vertex1, 
+		FVector& Vertex2, 
+		FVector& Vertex3, 
+		FVector& BarycentricCoords );
+
 
 
 	UFUNCTION(BlueprintPure, Category = "GeometryScript|MeshQueries", meta=(ScriptMethod))
@@ -100,6 +135,10 @@ public:
 	GetAllVertexPositions( UDynamicMesh* TargetMesh, FGeometryScriptVectorList& PositionList, bool bSkipGaps, bool& bHasVertexIDGaps );
 
 
+	//
+	// UV Queries
+	//
+
 
 	UFUNCTION(BlueprintPure, Category = "GeometryScript|MeshQueries", meta=(ScriptMethod), DisplayName = "Get Num UV Channels")
 	static UPARAM(DisplayName = "Num UV Channels") int32 GetNumUVSets( UDynamicMesh* TargetMesh );
@@ -113,12 +152,146 @@ public:
 	UFUNCTION(BlueprintPure, Category = "GeometryScript|MeshQueries", meta=(ScriptMethod))
 	static void GetTriangleUVs( UDynamicMesh* TargetMesh, UPARAM(DisplayName = "UV Channel") int32 UVSetIndex, int32 TriangleID, FVector2D& UV1, FVector2D& UV2, FVector2D& UV3, bool& bHaveValidUVs );
 
+	/**
+	 * Compute the interpolated UV (A*UV1+ B*UV2+ C*UV3), where (A,B,C)=BarycentricCoords and the UV positions are taken
+	 * from the specified TriangleID in the specified UVSet of the TargetMesh. 
+	 * @param bIsValidTriangle bTriHasValidUVs be returned true if TriangleID exists in TargetMesh and is set to valid UVs in the UV Set, and otherwise will be returned false
+	 */
+	UFUNCTION(BlueprintCallable, Category = "GeometryScript|MeshQueries", meta=(ScriptMethod))
+	static UPARAM(DisplayName = "Target Mesh") UDynamicMesh* 
+	GetInterpolatedTriangleUV( 
+		UDynamicMesh* TargetMesh, 
+		UPARAM(DisplayName = "UV Channel") int32 UVSetIndex, 
+		int32 TriangleID, 
+		FVector BarycentricCoords, 
+		bool& bTriHasValidUVs, 
+		FVector2D& InterpolatedUV );
+
+
+	//
+	// Normal queries
+	//
+
+	/**
+	 * @return true if the TargetMesh has the Normals Attribute enabled (which allows for storing split normals)
+	 */
+	UFUNCTION(BlueprintPure, Category = "GeometryScript|MeshQueries", meta=(ScriptMethod))
+	static UPARAM(DisplayName = "Enabled") bool GetHasTriangleNormals( UDynamicMesh* TargetMesh );
+
+	/**
+	 * For the specified TriangleID of the TargetMesh, get the Normal vectors at each vertex of the Triangle.
+	 * These Normals will be taken from the Normal Overlay, ie they will potentially be split-normals.
+	 * @param bTriHasValidNormals will be returned true if TriangleID exists in TargetMesh and has Normals set
+	 */
+	UFUNCTION(BlueprintCallable, Category = "GeometryScript|MeshQueries", meta=(ScriptMethod))
+	static UPARAM(DisplayName = "Target Mesh") UDynamicMesh* 
+	GetTriangleNormals( 
+		UDynamicMesh* TargetMesh, 
+		int32 TriangleID, 
+		FVector& Normal1, 
+		FVector& Normal2, 
+		FVector& Normal3, 
+		bool& bTriHasValidNormals );
+
+	/**
+	 * Compute the interpolated Normal (A*Normal1 + B*Normal2 + C*Normal3), where (A,B,C)=BarycentricCoords and the Normals are taken
+	 * from the specified TriangleID in the Normal layer of the TargetMesh. 
+	 * @param bTriHasValidNormals will be returned true if TriangleID exists in TargetMesh and has Normals set, and otherwise will be returned false
+	 */
+	UFUNCTION(BlueprintCallable, Category = "GeometryScript|MeshQueries", meta=(ScriptMethod))
+	static UPARAM(DisplayName = "Target Mesh") UDynamicMesh* 
+	GetInterpolatedTriangleNormal( 
+		UDynamicMesh* TargetMesh, 
+		int32 TriangleID, 
+		FVector BarycentricCoords, 
+		bool& bTriHasValidNormals,
+		FVector& InterpolatedNormal );
+
+	/**
+	 * For the specified TriangleID of the TargetMesh, get the Normal and Tangent vectors at each vertex of the Triangle.
+	 * These Normals/Tangents will be taken from the Normal and Tangents Overlays, ie they will potentially be split-normals.
+	 * @param bTriHasValidElements will be returned true if TriangleID exists in TargetMesh and has Normals and Tangents set
+	 */
+	UFUNCTION(BlueprintCallable, Category = "GeometryScript|MeshQueries", meta=(ScriptMethod))
+	static UPARAM(DisplayName = "Target Mesh") UDynamicMesh* 
+	GetTriangleNormalTangents( 
+		UDynamicMesh* TargetMesh, 
+		int32 TriangleID, 
+		bool& bTriHasValidElements,
+		FGeometryScriptTriangle& Normals,
+		FGeometryScriptTriangle& Tangents,
+		FGeometryScriptTriangle& BiTangents);
+
+	/**
+	 * Compute the interpolated Normal and Tangents for the specified specified TriangleID in the Normal and Tangent attributes of the TargetMesh. 
+	 * @param bTriHasValidElements will be returned true if TriangleID exists in TargetMesh and has Normals and Tangents set, and otherwise will be returned false
+	 */
+	UFUNCTION(BlueprintCallable, Category = "GeometryScript|MeshQueries", meta=(ScriptMethod))
+	static UPARAM(DisplayName = "Target Mesh") UDynamicMesh* 
+	GetInterpolatedTriangleNormalTangents( 
+		UDynamicMesh* TargetMesh, 
+		int32 TriangleID, 
+		FVector BarycentricCoords, 
+		bool& bTriHasValidElements,
+		FVector& InterpolatedNormal, 
+		FVector& InterpolatedTangent,
+		FVector& InterpolatedBiTangent);
+
+
+	//
+	// Vertex Color Queries
+	//
+
+
+	/**
+	 * @return true if the TargetMesh has the Vertex Colors attribute enabled
+	 */
+	UFUNCTION(BlueprintPure, Category = "GeometryScript|MeshQueries", meta=(ScriptMethod))
+	static UPARAM(DisplayName = "Enabled") bool GetHasVertexColors( UDynamicMesh* TargetMesh );
+
+	/**
+	 * For the specified TriangleID of the TargetMesh, get the Vertex Colors at each vertex of the Triangle.
+	 * These Colors will be taken from the Vertex Color Attribute, ie they will potentially be split-colors.
+	 * @param bTriHasValidVertexColors will be returned true if TriangleID exists in TargetMesh and has Vertex Colors set
+	 */
+	UFUNCTION(BlueprintCallable, Category = "GeometryScript|MeshQueries", meta=(ScriptMethod))
+	static UPARAM(DisplayName = "Target Mesh") UDynamicMesh* 
+	GetTriangleVertexColors( 
+		UDynamicMesh* TargetMesh, 
+		int32 TriangleID, 
+		FLinearColor& Color1, 
+		FLinearColor& Color2, 
+		FLinearColor& Color3, 
+		bool& bTriHasValidVertexColors );
+
+	/**
+	 * Compute the interpolated Vertex Color (A*Color1 + B*Color2 + C*Color3), where (A,B,C)=BarycentricCoords and the Colors are taken
+	 * from the specified TriangleID in the Vertex Color layer of the TargetMesh. 
+	 * @param bTriHasValidVertexColors will be returned true if TriangleID exists in TargetMesh and has Colors set, and otherwise will be returned false
+	 */
+	UFUNCTION(BlueprintCallable, Category = "GeometryScript|MeshQueries", meta=(ScriptMethod))
+	static UPARAM(DisplayName = "Target Mesh") UDynamicMesh* 
+	GetInterpolatedTriangleVertexColor( 
+		UDynamicMesh* TargetMesh, 
+		int32 TriangleID, 
+		FVector BarycentricCoords, 
+		FLinearColor DefaultColor,
+		bool& bTriHasValidVertexColors,
+		FLinearColor& InterpolatedColor );
 
 
+	//
+	// Material Queries
+	//
+
 	UFUNCTION(BlueprintPure, Category = "GeometryScript|MeshQueries", meta=(ScriptMethod))
 	static UPARAM(DisplayName = "Enabled") bool GetHasMaterialIDs( UDynamicMesh* TargetMesh );
 
 
+	//
+	// Polygroup Queries
+	//
+
 	UFUNCTION(BlueprintPure, Category = "GeometryScript|MeshQueries", meta=(ScriptMethod, DisplayName = "GetHasPolyGroups"))
 	static UPARAM(DisplayName = "Enabled") bool GetHasPolygroups( UDynamicMesh* TargetMesh );
 
