Sanitize package names generated and used by geometry script's new asset utility functions

diff --git a/Engine/Plugins/Experimental/GeometryScripting/Source/GeometryScriptingEditor/Private/CreateNewAssetUtilityFunctions.cpp b/Engine/Plugins/Experimental/GeometryScripting/Source/GeometryScriptingEditor/Private/CreateNewAssetUtilityFunctions.cpp
--- a/Engine/Plugins/Experimental/GeometryScripting/Source/GeometryScriptingEditor/Private/CreateNewAssetUtilityFunctions.cpp
+++ b/Engine/Plugins/Experimental/GeometryScripting/Source/GeometryScriptingEditor/Private/CreateNewAssetUtilityFunctions.cpp
@@ -21,6 +21,7 @@
 #include "Editor.h"
 #include "AssetRegistry/AssetRegistryModule.h"
 #include "Misc/Paths.h"
+#include "PackageTools.h"
 
 #include "ConversionUtils/DynamicMeshToVolume.h"
 #include "AssetUtils/CreateStaticMeshUtil.h"
@@ -57,7 +58,7 @@ void UGeometryScriptLibrary_CreateNewAssetFunctions::CreateUniqueNewAssetPathNam
 	{
 		FString UUID = UE::Modeling::GenerateRandomShortHexString(Options.UniqueIDDigits);
 		UniqueAssetName = FString::Printf(TEXT("%s_%s"), *BaseAssetName, *UUID);
-		UniqueAssetPathAndName = FPaths::Combine(AssetFolderPath, UniqueAssetName);
+		UniqueAssetPathAndName = UPackageTools::SanitizePackageName(FPaths::Combine(AssetFolderPath, UniqueAssetName));
 
 		// if asset does not exist at this path, we can use it
 		FAssetData AssetData = AssetRegistryModule.Get().GetAssetByObjectPath(FSoftObjectPath(UniqueAssetPathAndName));
@@ -167,6 +168,7 @@ UStaticMesh* UGeometryScriptLibrary_CreateNewAssetFunctions::CreateNewStaticMesh
 	GEditor->BeginTransaction(LOCTEXT("CreateNewStaticMeshAssetFromMesh_Transaction", "Create StaticMesh"));
 
 	UE::AssetUtils::FStaticMeshAssetOptions AssetOptions;
+	AssetPathAndName = UPackageTools::SanitizePackageName(AssetPathAndName);
 	AssetOptions.NewAssetPath = AssetPathAndName;
 
 	AssetOptions.NumSourceModels = 1;
@@ -250,6 +252,7 @@ USkeletalMesh* UGeometryScriptLibrary_CreateNewAssetFunctions::CreateNewSkeletal
 	GEditor->BeginTransaction(LOCTEXT("CreateNewSkeletalMeshAssetFromMesh_Transaction", "Create SkeletalMesh"));
 
 	FSkeletalMeshAssetOptions AssetOptions;
+	AssetPathAndName = UPackageTools::SanitizePackageName(AssetPathAndName);
 	AssetOptions.NewAssetPath = AssetPathAndName;
 	AssetOptions.Skeleton = InSkeleton;
 
@@ -315,6 +318,7 @@ UTexture2D* UGeometryScriptLibrary_CreateNewAssetFunctions::CreateNewTexture2DAs
 	}
 
 	UE::AssetUtils::FTexture2DAssetOptions AssetOptions;
+	AssetPathAndName = UPackageTools::SanitizePackageName(AssetPathAndName);
 	AssetOptions.NewAssetPath = AssetPathAndName;
 	AssetOptions.bOverwriteIfExists = Options.bOverwriteIfExists;
 
