Fix capitalization of PolyGroup in some rare user-visible places.

diff --git a/Engine/Plugins/Experimental/GeometryScripting/Source/GeometryScriptingCore/Public/GeometryScript/MeshPolygroupFunctions.h b/Engine/Plugins/Experimental/GeometryScripting/Source/GeometryScriptingCore/Public/GeometryScript/MeshPolygroupFunctions.h
--- a/Engine/Plugins/Experimental/GeometryScripting/Source/GeometryScriptingCore/Public/GeometryScript/MeshPolygroupFunctions.h
+++ b/Engine/Plugins/Experimental/GeometryScripting/Source/GeometryScriptingCore/Public/GeometryScript/MeshPolygroupFunctions.h
@@ -123,7 +123,7 @@ public:
 	/**
 	 * Create list of all triangles with the given PolyGroup ID in the given GroupLayer (not necessarily a single connected-component)
 	 */
-	UFUNCTION(BlueprintCallable, Category = "GeometryScript|Polygroups", meta = (ScriptMethod, DisplayName = "GetTrianglesInPolyGroup"))
+	UFUNCTION(BlueprintCallable, Category = "GeometryScript|PolyGroups", meta = (ScriptMethod, DisplayName = "GetTrianglesInPolyGroup"))
 	static UPARAM(DisplayName = "Target Mesh") UDynamicMesh* 
 	GetTrianglesInPolygroup( 
 		UDynamicMesh* TargetMesh, 
diff --git a/Engine/Plugins/Experimental/MeshModelingToolsetExp/Source/MeshModelingToolsExp/Private/ConvertToPolygonsTool.cpp b/Engine/Plugins/Experimental/MeshModelingToolsetExp/Source/MeshModelingToolsExp/Private/ConvertToPolygonsTool.cpp
--- a/Engine/Plugins/Experimental/MeshModelingToolsetExp/Source/MeshModelingToolsExp/Private/ConvertToPolygonsTool.cpp
+++ b/Engine/Plugins/Experimental/MeshModelingToolsetExp/Source/MeshModelingToolsExp/Private/ConvertToPolygonsTool.cpp
@@ -349,7 +349,7 @@ void UConvertToPolygonsTool::OnShutdown(EToolShutdownType ShutdownType)
 		FDynamicMeshOpResult Result = PreviewCompute->Shutdown();
 		if (ShutdownType == EToolShutdownType::Accept)
 		{
-			GetToolManager()->BeginUndoTransaction(LOCTEXT("ConvertToPolygonsToolTransactionName", "Find Polygroups"));
+			GetToolManager()->BeginUndoTransaction(LOCTEXT("ConvertToPolygonsToolTransactionName", "Find PolyGroups"));
 			FDynamicMesh3* DynamicMeshResult = Result.Mesh.Get();
 			if (ensure(DynamicMeshResult != nullptr))
 			{
diff --git a/Engine/Plugins/Runtime/MeshModelingToolset/Source/MeshModelingTools/Public/AddPrimitiveTool.h b/Engine/Plugins/Runtime/MeshModelingToolset/Source/MeshModelingTools/Public/AddPrimitiveTool.h
--- a/Engine/Plugins/Runtime/MeshModelingToolset/Source/MeshModelingTools/Public/AddPrimitiveTool.h
+++ b/Engine/Plugins/Runtime/MeshModelingToolset/Source/MeshModelingTools/Public/AddPrimitiveTool.h
@@ -62,15 +62,15 @@ enum class EMakeMeshPivotLocation : uint8
 	Top
 };
 
-/** Polygroup mode for shape */
+/** PolyGroup mode for shape */
 UENUM()
 enum class EMakeMeshPolygroupMode : uint8
 {
-	/** One Polygroup for the entire shape */
+	/** One PolyGroup for the entire shape */
 	PerShape,
-	/** One Polygroup for each geometric face */
+	/** One PolyGroup for each geometric face */
 	PerFace,
-	/** One Polygroup for each quad/triangle */
+	/** One PolyGroup for each quad/triangle */
 	PerQuad
 };
 
diff --git a/Engine/Plugins/Runtime/MeshModelingToolset/Source/MeshModelingTools/Public/Properties/RevolveProperties.h b/Engine/Plugins/Runtime/MeshModelingToolset/Source/MeshModelingTools/Public/Properties/RevolveProperties.h
--- a/Engine/Plugins/Runtime/MeshModelingToolset/Source/MeshModelingTools/Public/Properties/RevolveProperties.h
+++ b/Engine/Plugins/Runtime/MeshModelingToolset/Source/MeshModelingTools/Public/Properties/RevolveProperties.h
@@ -30,9 +30,9 @@ enum class ERevolvePropertiesCapFillMode : uint8
 UENUM()
 enum class ERevolvePropertiesPolygroupMode : uint8
 {
-	/** One Polygroup for the entire shape */
+	/** One PolyGroup for the entire shape */
 	PerShape,
-	/** One Polygroup for each geometric face */
+	/** One PolyGroup for each geometric face */
 	PerFace,
 	/** One PolyGroup along the path for each revolution step */
 	PerRevolveStep,
diff --git a/Engine/Plugins/Runtime/MeshModelingToolset/Source/ModelingOperatorsEditorOnly/Public/CleaningOps/SimplifyMeshOp.h b/Engine/Plugins/Runtime/MeshModelingToolset/Source/ModelingOperatorsEditorOnly/Public/CleaningOps/SimplifyMeshOp.h
--- a/Engine/Plugins/Runtime/MeshModelingToolset/Source/ModelingOperatorsEditorOnly/Public/CleaningOps/SimplifyMeshOp.h
+++ b/Engine/Plugins/Runtime/MeshModelingToolset/Source/ModelingOperatorsEditorOnly/Public/CleaningOps/SimplifyMeshOp.h
@@ -52,7 +52,7 @@ enum class ESimplifyType : uint8
 	MinimalPlanar = 4 UMETA(DisplayName = "Minimal Shape-Preserving"),
 
 	/** Only preserve polygroup boundaries; ignore all other shape features */
-	MinimalPolygroup = 5 UMETA(DisplayName = "Minimal Polygroup-Preserving"),
+	MinimalPolygroup = 5 UMETA(DisplayName = "Minimal PolyGroup-Preserving"),
 
 };
 
