GeometryScripting.  Scripting Core - Adding missing comments to Geometry Scrip nodes

diff --git a/Engine/Plugins/Experimental/GeometryScripting/Source/GeometryScriptingCore/Public/GeometryScript/CollisionFunctions.h b/Engine/Plugins/Experimental/GeometryScripting/Source/GeometryScriptingCore/Public/GeometryScript/CollisionFunctions.h
--- a/Engine/Plugins/Experimental/GeometryScripting/Source/GeometryScriptingCore/Public/GeometryScript/CollisionFunctions.h
+++ b/Engine/Plugins/Experimental/GeometryScripting/Source/GeometryScriptingCore/Public/GeometryScript/CollisionFunctions.h
@@ -110,7 +110,9 @@ class GEOMETRYSCRIPTINGCORE_API UGeometryScriptLibrary_CollisionFunctions : publ
 	GENERATED_BODY()
 public:
 
-
+	/** 
+	* Generates Simple Collision shapes for a Static Mesh Asset based on the input Dynamic Mesh.
+	*/
 	UFUNCTION(BlueprintCallable, Category = "GeometryScript|Collision")
 	static UPARAM(DisplayName = "Dynamic Mesh") UDynamicMesh* 
 	SetStaticMeshCollisionFromMesh(
@@ -120,7 +122,7 @@ public:
 		UGeometryScriptDebug* Debug = nullptr);
 
 	/**
-	 * Copy the Simple Collision Geometry from the SourceComponent to the StaticMeshAsset
+	 * Copy the Simple Collision Geometry from the Source Component to the Static Mesh Asset.
 	 */
 	UFUNCTION(BlueprintCallable, Category = "GeometryScript|Collision")
 	static void 
@@ -130,7 +132,9 @@ public:
 		FGeometryScriptSetSimpleCollisionOptions Options = FGeometryScriptSetSimpleCollisionOptions(),
 		UGeometryScriptDebug* Debug = nullptr);
 
-
+	/** 
+	* Generate Simple Collision shapes for a Dynamic Mesh Component based on the input Dynamic Mesh. 
+	*/
 	UFUNCTION(BlueprintCallable, Category = "GeometryScript|Collision")
 	static UPARAM(DisplayName = "Dynamic Mesh") UDynamicMesh* 
 	SetDynamicMeshCollisionFromMesh(
@@ -139,7 +143,9 @@ public:
 		FGeometryScriptCollisionFromMeshOptions Options,
 		UGeometryScriptDebug* Debug = nullptr);
 
-
+    /** 
+	* Clears Simple Collisions from the Dynamic Mesh Component. 
+	*/
 	UFUNCTION(BlueprintCallable, Category = "GeometryScript|Collision")
 	static void
 	ResetDynamicMeshCollision(
diff --git a/Engine/Plugins/Experimental/GeometryScripting/Source/GeometryScriptingCore/Public/GeometryScript/ListUtilityFunctions.h b/Engine/Plugins/Experimental/GeometryScripting/Source/GeometryScriptingCore/Public/GeometryScript/ListUtilityFunctions.h
--- a/Engine/Plugins/Experimental/GeometryScripting/Source/GeometryScriptingCore/Public/GeometryScript/ListUtilityFunctions.h
+++ b/Engine/Plugins/Experimental/GeometryScripting/Source/GeometryScriptingCore/Public/GeometryScript/ListUtilityFunctions.h
@@ -15,156 +15,308 @@ class GEOMETRYSCRIPTINGCORE_API UGeometryScriptLibrary_ListUtilityFunctions : pu
 {
 	GENERATED_BODY()
 public:
-
+	/**
+	* Returns the number of Items in Index List.
+	*/
 	UFUNCTION(BlueprintPure, Category = "GeometryScript|ListUtils", meta=(ScriptMethod))
 	static int GetIndexListLength(FGeometryScriptIndexList IndexList);
 
+	/**
+	* Returns the index of the last element in the Index List.
+	* Note, the value -1 will be returned if the list is empty or invalid. 
+	*/
 	UFUNCTION(BlueprintPure, Category = "GeometryScript|ListUtils", meta=(ScriptMethod))
 	static int GetIndexListLastIndex(FGeometryScriptIndexList IndexList);
 
+	/**
+	* Returns the item associated with Index in the Index List.
+	* If Index is not valid for this Index List the value -1 will be returned and bIsValidIndex will be set to false.
+	*/
 	UFUNCTION(BlueprintPure, Category = "GeometryScript|ListUtils", meta=(ScriptMethod))
 	static int GetIndexListItem(FGeometryScriptIndexList IndexList, int Index, bool& bIsValidIndex);
 
+	/**
+	* Updates the value associated with Index in the Index List.  
+	* If the Index is invalid, the operation will fail and in this case bValidIndex will be set to false on return.
+	*/
 	UFUNCTION(BlueprintCallable, Category = "GeometryScript|ListUtils", meta=(ScriptMethod))
 	static void SetIndexListItem(UPARAM(ref) FGeometryScriptIndexList& IndexList, int Index, int NewValue, bool& bIsValidIndex);
 
+	/**
+	* Populates Index Array with the integer values stored in the Index List.
+	*/
 	UFUNCTION(BlueprintCallable, Category = "GeometryScript|ListUtils", meta=(ScriptMethod))
 	static void ConvertIndexListToArray(FGeometryScriptIndexList IndexList, TArray<int>& IndexArray);
 
+	/**
+	* Populates Index List of the specified Index Type from the integer values stored in the Index Array.
+	*/
 	UFUNCTION(BlueprintCallable, Category = "GeometryScript|ListUtils")
 	static void ConvertArrayToIndexList(const TArray<int>& IndexArray, FGeometryScriptIndexList& IndexList, EGeometryScriptIndexType IndexType = EGeometryScriptIndexType::Any);
 
+	/**
+	* Updates Duplicate List to be identical to Index List.
+	*/
 	UFUNCTION(BlueprintCallable, Category = "GeometryScript|ListUtils", meta=(ScriptMethod))
 	static void DuplicateIndexList(FGeometryScriptIndexList IndexList, FGeometryScriptIndexList& DuplicateList);
 
+	/**
+	* Set each value in Index List to the given Clear Value.
+	*/
 	UFUNCTION(BlueprintCallable, Category = "GeometryScript|ListUtils", meta=(ScriptMethod))
 	static void ClearIndexList(UPARAM(ref) FGeometryScriptIndexList& IndexList, int ClearValue = 0);
 
-
+	/**
+	* Returns the number of Triangles in the  Triangle list.
+	*/
 	UFUNCTION(BlueprintPure, Category = "GeometryScript|ListUtils", meta=(ScriptMethod))
 	static int GetTriangleListLength(FGeometryScriptTriangleList TriangleList);
 
+	/**
+	* Returns the index of the last element in the Triangle List.  
+	* If the Triangle List is empty or invalid, the value 0 will be returned.
+	*/
 	UFUNCTION(BlueprintPure, Category = "GeometryScript|ListUtils", meta=(ScriptMethod))
 	static int GetTriangleListLastTriangle(FGeometryScriptTriangleList TriangleList);
 
+	/**
+	* Returns the integer triplet associated with the index Triangle in the Triangle  List.
+	* If Triangle is not valid for this Triangle List, the triplet (-1, -1, -1) will be returned and bIsValidIndex set to false.
+	*/
 	UFUNCTION(BlueprintPure, Category = "GeometryScript|ListUtils", meta=(ScriptMethod))
 	static FIntVector GetTriangleListItem(FGeometryScriptTriangleList TriangleList, int Triangle, bool& bIsValidTriangle);
 
+	/**
+	* Converts Triangle List to Triangle Array by populating with the appropriate integer triplets.
+	*/
 	UFUNCTION(BlueprintCallable, Category = "GeometryScript|ListUtils", meta=(ScriptMethod))
 	static void ConvertTriangleListToArray(FGeometryScriptTriangleList TriangleList, TArray<FIntVector>& TriangleArray);
 
+	/**
+	* Converts a Triangle Array of integer triplets to a Triangle List.
+	*/
 	UFUNCTION(BlueprintCallable, Category = "GeometryScript|ListUtils")
 	static void ConvertArrayToTriangleList(const TArray<FIntVector>& TriangleArray, FGeometryScriptTriangleList& TriangleList);
 
 
-
+	/**
+	* Returns the number of items in the Scalar List.
+	*/
 	UFUNCTION(BlueprintPure, Category = "GeometryScript|ListUtils", meta=(ScriptMethod))
 	static int GetScalarListLength(FGeometryScriptScalarList ScalarList);
 
+	/**
+	* Returns the index of the last Scalar in Scalar List.
+	* If Scalar List is empty or invalid, the value -1 will be returned 
+	*/
 	UFUNCTION(BlueprintPure, Category = "GeometryScript|ListUtils", meta=(ScriptMethod))
 	static int GetScalarListLastIndex(FGeometryScriptScalarList ScalarList);
 
+	/**
+	* Returns the Scalar value associated with Index in Scalar List.
+	* If the Index is not valid for this Scalar List, the value 0.0 will be returned and bIsValidIndex set to false.
+	*/
 	UFUNCTION(BlueprintPure, Category = "GeometryScript|ListUtils", meta=(ScriptMethod))
 	static double GetScalarListItem(FGeometryScriptScalarList ScalarList, int Index, bool& bIsValidIndex);
 
+	/**
+	* Updates the value associated with Index in the Scalar List.  
+	* If the Index is invalid, the operation will fail and bValidIndex will be set to false.
+	*/
 	UFUNCTION(BlueprintCallable, Category = "GeometryScript|ListUtils", meta=(ScriptMethod))
 	static void SetScalarListItem(UPARAM(ref) FGeometryScriptScalarList& ScalarList, int Index, double NewValue, bool& bIsValidIndex);
 
+	/**
+	* Converts a Scalar List to an Scalar Array (an array of doubles).
+	*/
 	UFUNCTION(BlueprintCallable, Category = "GeometryScript|ListUtils", meta=(ScriptMethod))
 	static void ConvertScalarListToArray(FGeometryScriptScalarList ScalarList, TArray<double>& ScalarArray);
 
+	/**
+	* Converts an array of doubles (Scalar Array) to Scalar List.
+	*/
 	UFUNCTION(BlueprintCallable, Category = "GeometryScript|ListUtils")
-	static void ConvertArrayToScalarList(const TArray<double>& VectorArray, FGeometryScriptScalarList& ScalarList);
+	static void ConvertArrayToScalarList(UPARAM(DisplayName = "Scalar Array") const TArray<double>& VectorArray, FGeometryScriptScalarList& ScalarList);
 
+	/**
+	* Copies the contents of Scalar List into Duplicate List.
+	*/
 	UFUNCTION(BlueprintCallable, Category = "GeometryScript|ListUtils", meta=(ScriptMethod))
 	static void DuplicateScalarList(FGeometryScriptScalarList ScalarList, FGeometryScriptScalarList& DuplicateList);
 
+	/**
+	* Resets all the items in the Scalar List to the Clear Value.
+	*/
 	UFUNCTION(BlueprintCallable, Category = "GeometryScript|ListUtils", meta=(ScriptMethod))
 	static void ClearScalarList(UPARAM(ref) FGeometryScriptScalarList& ScalarList, double ClearValue = 0.0);
 
 
 
-
+	/**
+	* Returns the number of items in the Vector List.
+	*/
 	UFUNCTION(BlueprintPure, Category = "GeometryScript|ListUtils", meta=(ScriptMethod))
 	static int GetVectorListLength(FGeometryScriptVectorList VectorList);
 
+	/**
+	* Returns the index of the last item in the Vector List.
+	* If Vector List is empty or invalid, the value -1 will be returned. 
+	*/
 	UFUNCTION(BlueprintPure, Category = "GeometryScript|ListUtils", meta=(ScriptMethod))
 	static int GetVectorListLastIndex(FGeometryScriptVectorList VectorList);
 
+	/**
+	* Returns the FVector stored in the VectorList at the specified location.
+	* if the Index is not valid for this Vector List, the Zero Vector will be returned and bIsValidIndex set to false. 
+	*/
 	UFUNCTION(BlueprintPure, Category = "GeometryScript|ListUtils", meta=(ScriptMethod))
 	static FVector GetVectorListItem(FGeometryScriptVectorList VectorList, int Index, bool& bIsValidIndex);
 
+	/**
+	* Updates the value of the FVector stored in the Vector List at the specified location.
+	* If the Index is invalid, the operation will fail and bValidIndex will be set to false.
+	*/
 	UFUNCTION(BlueprintCallable, Category = "GeometryScript|ListUtils", meta=(ScriptMethod))
 	static void SetVectorListItem(UPARAM(ref) FGeometryScriptVectorList& VectorList, int Index, FVector NewValue, bool& bIsValidIndex);
 
+	/**
+	* Converts Vector List to an array of FVectors (Vector Array).
+	*/
 	UFUNCTION(BlueprintCallable, Category = "GeometryScript|ListUtils", meta=(ScriptMethod))
 	static void ConvertVectorListToArray(FGeometryScriptVectorList VectorList, TArray<FVector>& VectorArray);
 
+	/**
+	* Converts an Array of FVectors (Vector Array) to Vector List.
+	*/
 	UFUNCTION(BlueprintCallable, Category = "GeometryScript|ListUtils")
 	static void ConvertArrayToVectorList(const TArray<FVector>& VectorArray, FGeometryScriptVectorList& VectorList);
 
+	/**
+	* Copies the contents of Vector List into Duplicate Vector List.
+	*/
 	UFUNCTION(BlueprintCallable, Category = "GeometryScript|ListUtils", meta=(ScriptMethod))
 	static void DuplicateVectorList(FGeometryScriptVectorList VectorList, FGeometryScriptVectorList& DuplicateList);
 
+	/**
+	* Resets all the items in the Vector List to the Clear Value.
+	*/
 	UFUNCTION(BlueprintCallable, Category = "GeometryScript|ListUtils", meta=(ScriptMethod))
 	static void ClearVectorList(UPARAM(ref) FGeometryScriptVectorList& VectorList, FVector ClearValue = FVector::ZeroVector);
 
 
 
-
+	/**
+	* Returns the number of items in the UV List.
+	*/
 	UFUNCTION(BlueprintPure, Category = "GeometryScript|ListUtils", meta=(ScriptMethod))
 	static int GetUVListLength(FGeometryScriptUVList UVList);
 
+	/**
+	* Returns the index of the last item in the UV List.
+	* If UV List is empty or invalid, the value -1 will be returned.
+	*/
 	UFUNCTION(BlueprintPure, Category = "GeometryScript|ListUtils", meta=(ScriptMethod))
 	static int GetUVListLastIndex(FGeometryScriptUVList UVList);
 
+	/**
+	* Returns the FVector2D stored in the UV List at the specified location.
+	* If the Index is not valid for this UV List, the Zero Vector will be returned and bIsValidIndex set to false.
+	*/
 	UFUNCTION(BlueprintPure, Category = "GeometryScript|ListUtils", meta=(ScriptMethod))
 	static FVector2D GetUVListItem(FGeometryScriptUVList UVList, int Index, bool& bIsValidIndex);
 
+	/**
+	* Updates the value of the FVector2D stored in the UV List at the specified location.
+	* If the Index is invalid, the operation will fail and bValidIndex will be set to false on return.
+	*/
 	UFUNCTION(BlueprintCallable, Category = "GeometryScript|ListUtils", meta=(ScriptMethod))
 	static void SetUVListItem(UPARAM(ref) FGeometryScriptUVList& UVList, int Index, FVector2D NewUV, bool& bIsValidIndex);
 
+	/**
+	* Converts a UV List to an array of FVector2Ds (UV Array).
+	*/
 	UFUNCTION(BlueprintCallable, Category = "GeometryScript|ListUtils", meta=(ScriptMethod))
 	static void ConvertUVListToArray(FGeometryScriptUVList UVList, TArray<FVector2D>& UVArray);
 
+	/**
+	* Converts an array of FVector2D (UV Array) to UV List.
+	*/
 	UFUNCTION(BlueprintCallable, Category = "GeometryScript|ListUtils")
 	static void ConvertArrayToUVList(const TArray<FVector2D>& UVArray, FGeometryScriptUVList& UVList);
 
+	/**
+	* Duplicates the contents of UV List into Duplicate List.
+	*/
 	UFUNCTION(BlueprintCallable, Category = "GeometryScript|ListUtils", meta=(ScriptMethod))
 	static void DuplicateUVList(FGeometryScriptUVList UVList, FGeometryScriptUVList& DuplicateList);
 
+	/**
+	* Resets all the items in the Vector List to the given Clear UV value.
+	*/
 	UFUNCTION(BlueprintCallable, Category = "GeometryScript|ListUtils", meta=(ScriptMethod))
 	static void ClearUVList(UPARAM(ref) FGeometryScriptUVList& UVList, FVector2D ClearUV);
 
 
-
+	/**
+	* Returns the number of items in the Color List.
+	*/
 	UFUNCTION(BlueprintPure, Category = "GeometryScript|ListUtils", meta=(ScriptMethod))
 	static int GetColorListLength(FGeometryScriptColorList ColorList);
 
+	/**
+	* Returns the index of the last item in the Color List.
+	* If Color List is empty or invalid, the value -1 will be returned.
+	*/
 	UFUNCTION(BlueprintPure, Category = "GeometryScript|ListUtils", meta=(ScriptMethod))
 	static int GetColorListLastIndex(FGeometryScriptColorList ColorList);
 
+	/**
+	* Returns the FLinearColor stored in the Color List at the specified location.
+	* If the Index is not valid for this Color List, FLinearColor::White will be returned and bIsValidIndex set to false.
+	*/
 	UFUNCTION(BlueprintPure, Category = "GeometryScript|ListUtils", meta=(ScriptMethod))
 	static FLinearColor GetColorListItem(FGeometryScriptColorList ColorList, int Index, bool& bIsValidIndex);
 
+	/**
+	* Updates the value of the FLinearColor stored in the Color List at the specified location.
+	* If the Index is invalid, the operation will fail and bValidIndex will be set to false.
+	*/
 	UFUNCTION(BlueprintCallable, Category = "GeometryScript|ListUtils", meta=(ScriptMethod))
 	static void SetColorListItem(UPARAM(ref) FGeometryScriptColorList& ColorList, int Index, FLinearColor NewColor, bool& bIsValidIndex);
 
+	/**
+	* Converts the Color List to an array of FLinearColor (Color Array).
+	*/
 	UFUNCTION(BlueprintCallable, Category = "GeometryScript|ListUtils", meta=(ScriptMethod))
 	static void ConvertColorListToArray(FGeometryScriptColorList ColorList, TArray<FLinearColor>& ColorArray);
 
+	/**
+	* Converts an array of FLinearColor (Color Array) to a Color List.
+	*/
 	UFUNCTION(BlueprintCallable, Category = "GeometryScript|ListUtils")
 	static void ConvertArrayToColorList(const TArray<FLinearColor>& ColorArray, FGeometryScriptColorList& ColorList);
 
+	/**
+	* Duplicates the contents of Color List into Duplicate List.
+	*/
 	UFUNCTION(BlueprintCallable, Category = "GeometryScript|ListUtils", meta=(ScriptMethod))
 	static void DuplicateColorList(FGeometryScriptColorList ColorList, FGeometryScriptColorList& DuplicateList);
 
+	/**
+	* Resets all the items in the Color List to the specified Clear Color.
+	*/
 	UFUNCTION(BlueprintCallable, Category = "GeometryScript|ListUtils", meta=(ScriptMethod))
 	static void ClearColorList(UPARAM(ref) FGeometryScriptColorList& ColorList, FLinearColor ClearColor);
 
+	/**
+	* Populates a Scalar List with values that corresponds to the 0, 1, 2, or 3 channel of a Color List.
+	*/
 	UFUNCTION(BlueprintCallable, Category = "GeometryScript|ListUtils", meta=(ScriptMethod))
 	static void ExtractColorListChannel(FGeometryScriptColorList ColorList, FGeometryScriptScalarList& ScalarList, int32 ChannelIndex = 0);
 
+	/**
+	* Populates a Vector List from a Color List. The channels in the Color List are mapped to vector components by means of X Channel Index, Y Channel Index, and Z Channel Index.
+	*/
 	UFUNCTION(BlueprintCallable, Category = "GeometryScript|ListUtils", meta=(ScriptMethod))
 	static void ExtractColorListChannels(FGeometryScriptColorList ColorList, FGeometryScriptVectorList& VectorList, int32 XChannelIndex = 0, int32 YChannelIndex = 1, int32 ZChannelIndex = 2);
 
diff --git a/Engine/Plugins/Experimental/GeometryScripting/Source/GeometryScriptingCore/Public/GeometryScript/MeshAssetFunctions.h b/Engine/Plugins/Experimental/GeometryScripting/Source/GeometryScriptingCore/Public/GeometryScript/MeshAssetFunctions.h
--- a/Engine/Plugins/Experimental/GeometryScripting/Source/GeometryScriptingCore/Public/GeometryScript/MeshAssetFunctions.h
+++ b/Engine/Plugins/Experimental/GeometryScripting/Source/GeometryScriptingCore/Public/GeometryScript/MeshAssetFunctions.h
@@ -104,7 +104,10 @@ class GEOMETRYSCRIPTINGCORE_API UGeometryScriptLibrary_StaticMeshFunctions : pub
 {
 	GENERATED_BODY()
 public:
-
+	
+	/** 
+	* Extracts a Dynamic Mesh from a Static Mesh Asset. 
+	*/
 	UFUNCTION(BlueprintCallable, Category = "GeometryScript|StaticMesh", meta = (ExpandEnumAsExecs = "Outcome"))
 	static UPARAM(DisplayName = "Dynamic Mesh") UDynamicMesh* 
 	CopyMeshFromStaticMesh(
@@ -115,7 +118,9 @@ public:
 		EGeometryScriptOutcomePins& Outcome,
 		UGeometryScriptDebug* Debug = nullptr);
 
-
+	/** 
+	* Updates a Static Mesh Asset with new geometry converted from a Dynamic Mesh.
+	*/
 	UFUNCTION(BlueprintCallable, Category = "GeometryScript|StaticMesh", meta = (ExpandEnumAsExecs = "Outcome"))
 	static UPARAM(DisplayName = "Dynamic Mesh") UDynamicMesh* 
 	CopyMeshToStaticMesh(
@@ -127,7 +132,9 @@ public:
 		UGeometryScriptDebug* Debug = nullptr);
 
 
-
+    /** 
+	* Extracts the Material List and corresponding Material Indices from the specified LOD of the Static Mesh Asset. 
+	*/
 	UFUNCTION(BlueprintCallable, Category = "GeometryScript|StaticMesh", meta = (ExpandEnumAsExecs = "Outcome"))
 	static void
 	GetSectionMaterialListFromStaticMesh(
@@ -138,7 +145,9 @@ public:
 		EGeometryScriptOutcomePins& Outcome,
 		UGeometryScriptDebug* Debug = nullptr);
 
-
+	/** 
+	* Extracts a Dynamic Mesh from a Skeletal Mesh Asset. 
+	*/
 	UFUNCTION(BlueprintCallable, Category = "GeometryScript|SkeletalMesh", meta = (ExpandEnumAsExecs = "Outcome"))
 	static UPARAM(DisplayName = "Dynamic Mesh") UDynamicMesh* 
 	CopyMeshFromSkeletalMesh(
@@ -149,7 +158,9 @@ public:
 		EGeometryScriptOutcomePins& Outcome,
 		UGeometryScriptDebug* Debug = nullptr);
 
-
+	/** 
+	* Updates a Skeletal Mesh Asset with new geometry and bone weights data from a Dynamic Mesh.
+	*/
 	UFUNCTION(BlueprintCallable, Category = "GeometryScript|SkeletalMesh", meta = (ExpandEnumAsExecs = "Outcome"))
 	static UPARAM(DisplayName = "Dynamic Mesh") UDynamicMesh* 
 	CopyMeshToSkeletalMesh(
diff --git a/Engine/Plugins/Experimental/GeometryScripting/Source/GeometryScriptingCore/Public/GeometryScript/MeshBasicEditFunctions.h b/Engine/Plugins/Experimental/GeometryScripting/Source/GeometryScriptingCore/Public/GeometryScript/MeshBasicEditFunctions.h
--- a/Engine/Plugins/Experimental/GeometryScripting/Source/GeometryScriptingCore/Public/GeometryScript/MeshBasicEditFunctions.h
+++ b/Engine/Plugins/Experimental/GeometryScripting/Source/GeometryScriptingCore/Public/GeometryScript/MeshBasicEditFunctions.h
@@ -110,8 +110,8 @@ public:
 		bool bDeferChangeNotifications = false );
 
 	/**
-	 * Set all vertex positions in the TargetMesh to the specified Positions
-	 * @param PositionList new vertex Positions. Size must be less than or equal to the MaxVertexID of TargetMesh  (ie gaps are supported)
+	 * Set all vertex positions in the TargetMesh to the specified Positions.
+	 * @param PositionList new vertex Positions. Size must be less than or equal to the MaxVertexID of TargetMesh  (ie gaps are supported).
 	 */
 	UFUNCTION(BlueprintCallable, Category = "GeometryScript|MeshEdits", meta=(ScriptMethod))
 	static UPARAM(DisplayName = "Target Mesh") UDynamicMesh*
@@ -120,7 +120,9 @@ public:
 		FGeometryScriptVectorList PositionList,
 		UGeometryScriptDebug* Debug = nullptr);
 
-
+	/**
+	* Adds a new vertex to the mesh and returns a new Vertex ID (NewVertexIndex).
+	*/
 	UFUNCTION(BlueprintCallable, Category = "GeometryScript|MeshEdits", meta=(ScriptMethod))
 	static UPARAM(DisplayName = "Target Mesh") UDynamicMesh* 
 	AddVertexToMesh( 
@@ -129,6 +131,9 @@ public:
 		int& NewVertexIndex,
 		bool bDeferChangeNotifications = false );
 
+	/**
+	* Adds a list of vertices to the mesh, and populates the NewIndicesList with the corresponding new Vertex IDs.
+	*/ 
 	UFUNCTION(BlueprintCallable, Category = "GeometryScript|MeshEdits", meta=(ScriptMethod))
 	static UPARAM(DisplayName = "Target Mesh") UDynamicMesh* 
 	AddVerticesToMesh( 
@@ -137,6 +142,10 @@ public:
 		FGeometryScriptIndexList& NewIndicesList,
 		bool bDeferChangeNotifications = false );
 
+	/**
+	* Removes a vertex from the mesh as indicated by the VertexID.  
+	* Should the delete fail, e.g. if the specified vertex was not a mesh element, the flag bWasVertexDeleted will be set to false. 
+	*/
 	UFUNCTION(BlueprintCallable, Category = "GeometryScript|MeshEdits", meta=(ScriptMethod))
 	static UPARAM(DisplayName = "Target Mesh") UDynamicMesh* 
 	DeleteVertexFromMesh( 
@@ -145,6 +154,10 @@ public:
 		bool& bWasVertexDeleted,
 		bool bDeferChangeNotifications = false );
 
+	/**
+	* Removes a list of vertices from the mesh.  
+	* On return, NumDeleted will contain the actual number of vertices removed.
+	*/
 	UFUNCTION(BlueprintCallable, Category = "GeometryScript|MeshEdits", meta=(ScriptMethod))
 	static UPARAM(DisplayName = "Target Mesh") UDynamicMesh* 
 	DeleteVerticesFromMesh( 
@@ -153,7 +166,9 @@ public:
 		int& NumDeleted,
 		bool bDeferChangeNotifications = false );
 
-
+	/**
+	* Adds a triangle (Vertex ID triplet) to the mesh and updates New Triangle Index with the resulting Triangle ID.
+	*/
 	UFUNCTION(BlueprintCallable, Category = "GeometryScript|MeshEdits", meta=(ScriptMethod))
 	static UPARAM(DisplayName = "Target Mesh") UDynamicMesh* 
 	AddTriangleToMesh( 
@@ -164,7 +179,9 @@ public:
 		bool bDeferChangeNotifications = false,
 		UGeometryScriptDebug* Debug = nullptr );
 
-
+	/**
+	* Adds a list of triangles to the mesh and populates the New Indices List with the corresponding new Triangle IDs.
+	*/
 	UFUNCTION(BlueprintCallable, Category = "GeometryScript|MeshEdits", meta=(ScriptMethod))
 	static UPARAM(DisplayName = "Target Mesh") UDynamicMesh* 
 	AddTrianglesToMesh( 
@@ -175,6 +192,10 @@ public:
 		bool bDeferChangeNotifications = false,
 		UGeometryScriptDebug* Debug = nullptr );
 
+	/**
+	* Removes a triangle from the mesh as indicated by the Triangle ID.
+	* Should the delete fail, e.g. if the specified triangle was not a mesh element, the flag bWasTriangleDelete will be set to false. 
+	*/
 	UFUNCTION(BlueprintCallable, Category = "GeometryScript|MeshEdits", meta=(ScriptMethod))
 	static UPARAM(DisplayName = "Target Mesh") UDynamicMesh* 
 	DeleteTriangleFromMesh( 
@@ -183,6 +204,10 @@ public:
 		bool& bWasTriangleDeleted,
 		bool bDeferChangeNotifications = false );
 
+	/**
+	* Removes a list of triangles from the mesh. 
+	* On return, NumDeleted will contain the actual number of triangles removed.
+	*/
 	UFUNCTION(BlueprintCallable, Category = "GeometryScript|MeshEdits", meta=(ScriptMethod))
 	static UPARAM(DisplayName = "Target Mesh") UDynamicMesh* 
 	DeleteTrianglesFromMesh( 
@@ -191,6 +216,10 @@ public:
 		int& NumDeleted,
 		bool bDeferChangeNotifications = false );
 
+	/**
+	 * Removes specified triangles, identified by mesh selection, from the mesh.
+	 * On return, NumDeleted will contain the actual number of triangles removed.
+	 */
 	UFUNCTION(BlueprintCallable, Category = "GeometryScript|MeshEdits", meta=(ScriptMethod))
 	static UPARAM(DisplayName = "Target Mesh") UDynamicMesh* 
 	DeleteSelectedTrianglesFromMesh( 
@@ -200,8 +229,8 @@ public:
 		bool bDeferChangeNotifications = false );
 
 	/**
-	 * Apply AppendTransform to AppendMesh and then add its geometry to the TargetMesh
-	 * @param AppendOptions Control how details like mesh attributes are handled when one mesh is appended to another
+	 * Apply Append Transform to Append Mesh and then add its geometry to the Target Mesh.
+	 * @param AppendOptions Control how details like mesh attributes are handled when one mesh is appended to another.
 	 */
 	UFUNCTION(BlueprintCallable, Category = "GeometryScript|MeshEdits", meta=(ScriptMethod))
 	static UPARAM(DisplayName = "Target Mesh") UDynamicMesh* 
@@ -250,7 +279,9 @@ public:
 		UGeometryScriptDebug* Debug = nullptr);
 
 
-
+	/**
+	* Adds a set of vertices/triangles to the mesh, with Normals, UVs, and Colors; returning the new triangles indices 
+	*/
 	UFUNCTION(BlueprintCallable, Category = "GeometryScript|MeshEdits", meta=(ScriptMethod))
 	static UPARAM(DisplayName = "Target Mesh") UDynamicMesh* 
 	AppendBuffersToMesh( 
diff --git a/Engine/Plugins/Experimental/GeometryScripting/Source/GeometryScriptingCore/Public/GeometryScript/MeshBooleanFunctions.h b/Engine/Plugins/Experimental/GeometryScripting/Source/GeometryScriptingCore/Public/GeometryScript/MeshBooleanFunctions.h
--- a/Engine/Plugins/Experimental/GeometryScripting/Source/GeometryScriptingCore/Public/GeometryScript/MeshBooleanFunctions.h
+++ b/Engine/Plugins/Experimental/GeometryScripting/Source/GeometryScriptingCore/Public/GeometryScript/MeshBooleanFunctions.h
@@ -119,6 +119,15 @@ class GEOMETRYSCRIPTINGCORE_API UGeometryScriptLibrary_MeshBooleanFunctions : pu
 	GENERATED_BODY()
 public:
 
+	/**
+	* Applies a Boolean operation (such as, Union, Intersect, and Subtract) to the Target Dynamic Mesh based on a Tool Dynamic Mesh.
+	* @param TargetMesh  Dynamic Mesh to be acted upon
+	* @param TargetTransform  used to position the TargetMesh
+	* @param ToolMesh  Dynamic Mesh that acts as the cutting tool
+	* @param ToolTransform used to position the ToolMesh
+	* @param Operation selects the specific boolean operation
+	* @param Options selects additional options that are applied to the result
+	*/
 	UFUNCTION(BlueprintCallable, Category = "GeometryScript|Booleans", meta=(ScriptMethod))
 	static UPARAM(DisplayName = "Target Mesh") UDynamicMesh* 
 	ApplyMeshBoolean( 
@@ -130,7 +139,9 @@ public:
 		FGeometryScriptMeshBooleanOptions Options,
 		UGeometryScriptDebug* Debug = nullptr);
 
-
+	/**
+	* Mesh-Boolean-Union an object with itself to repair self-intersections, remove floating geometry, etc.
+	*/
 	UFUNCTION(BlueprintCallable, Category = "GeometryScript|Booleans", meta=(ScriptMethod))
 	static UPARAM(DisplayName = "Target Mesh") UDynamicMesh* 
 	ApplyMeshSelfUnion( 
@@ -139,7 +150,12 @@ public:
 		UGeometryScriptDebug* Debug = nullptr);
 
 
-
+	/**
+	* Applies a plane cut to a mesh, optionally filling any holes created.
+	* @param TargetMesh Dynamic Mesh to updated with the cut
+	* @param CutFrame defines the plane used to cut the TargetMesh
+	* @param Options selects additional clean-up operations performed after the cut.
+	*/
 	UFUNCTION(BlueprintCallable, Category = "GeometryScript|Booleans", meta=(ScriptMethod))
 	static UPARAM(DisplayName = "Target Mesh") UDynamicMesh* 
 	ApplyMeshPlaneCut( 
@@ -147,7 +163,13 @@ public:
 		FTransform CutFrame,
 		FGeometryScriptMeshPlaneCutOptions Options,
 		UGeometryScriptDebug* Debug = nullptr);
-
+	
+	/**
+	 * Slices a mesh into two halves, with optional hole filling.
+	 * @param TargetMesh Dynamic Mesh to be updated by the slice.
+	 * @param CutFrame defines the plane used to slice the TargetMesh.
+	 * @param Options selects additional clean-up operations performed after the cut.
+	 */
 	UFUNCTION(BlueprintCallable, Category = "GeometryScript|Booleans", meta=(ScriptMethod))
 	static UPARAM(DisplayName = "Target Mesh") UDynamicMesh* 
 	ApplyMeshPlaneSlice( 
@@ -156,7 +178,12 @@ public:
 		FGeometryScriptMeshPlaneSliceOptions Options,
 		UGeometryScriptDebug* Debug = nullptr);
 
-
+	/**
+	* Mirrors a mesh across a plane, with optional cutting and welding of triangles.
+	* @param TargetMesh Dynamic Mesh to be updated by the mirror operation.
+	* @param MirrorFrame defines the plane used to mirror the TargetMesh.
+	* @param Options selects  additional clean-up operations performed after the mirror.
+	*/
 	UFUNCTION(BlueprintCallable, Category = "GeometryScript|Booleans", meta=(ScriptMethod))
 	static UPARAM(DisplayName = "Target Mesh") UDynamicMesh* 
 	ApplyMeshMirror( 
diff --git a/Engine/Plugins/Experimental/GeometryScripting/Source/GeometryScriptingCore/Public/GeometryScript/MeshComparisonFunctions.h b/Engine/Plugins/Experimental/GeometryScripting/Source/GeometryScriptingCore/Public/GeometryScript/MeshComparisonFunctions.h
--- a/Engine/Plugins/Experimental/GeometryScripting/Source/GeometryScriptingCore/Public/GeometryScript/MeshComparisonFunctions.h
+++ b/Engine/Plugins/Experimental/GeometryScripting/Source/GeometryScriptingCore/Public/GeometryScript/MeshComparisonFunctions.h
@@ -59,6 +59,9 @@ class GEOMETRYSCRIPTINGCORE_API UGeometryScriptLibrary_MeshComparisonFunctions :
 	GENERATED_BODY()
 public:
 
+	/**
+	 * Returns true if the two input meshes are equivalent under the comparisons defined by the input options.
+	 */
 	UFUNCTION(BlueprintCallable, Category = "GeometryScript|Comparison", meta = (ScriptMethod))
 	static UPARAM(DisplayName = "Target Mesh") UDynamicMesh*
 	IsSameMeshAs(
@@ -68,6 +71,9 @@ public:
 		bool &bIsSameMesh,
 		UGeometryScriptDebug* Debug = nullptr);
 
+	/**
+	* Measures the min/max and average closest-point distances between two meshes.
+	*/
 	UFUNCTION(BlueprintCallable, Category = "GeometryScript|Comparison", meta = (ScriptMethod))
 	static UPARAM(DisplayName = "Target Mesh") UDynamicMesh*
 	MeasureDistancesBetweenMeshes(
@@ -80,6 +86,9 @@ public:
 		double& RootMeanSqrDeviation,
 		UGeometryScriptDebug* Debug = nullptr);
 
+	/**
+	* Returns true if the two input meshes (with optional transforms) are geometrically intersecting.
+	*/
 	UFUNCTION(BlueprintCallable, Category = "GeometryScript|Comparison", meta = (ScriptMethod))
 	static UPARAM(DisplayName = "Target Mesh") UDynamicMesh*
 	IsIntersectingMesh(
diff --git a/Engine/Plugins/Experimental/GeometryScripting/Source/GeometryScriptingCore/Public/GeometryScript/MeshDeformFunctions.h b/Engine/Plugins/Experimental/GeometryScripting/Source/GeometryScriptingCore/Public/GeometryScript/MeshDeformFunctions.h
--- a/Engine/Plugins/Experimental/GeometryScripting/Source/GeometryScriptingCore/Public/GeometryScript/MeshDeformFunctions.h
+++ b/Engine/Plugins/Experimental/GeometryScripting/Source/GeometryScriptingCore/Public/GeometryScript/MeshDeformFunctions.h
@@ -200,6 +200,10 @@ class GEOMETRYSCRIPTINGCORE_API UGeometryScriptLibrary_MeshDeformFunctions : pub
 	GENERATED_BODY()
 public:
 
+	/**
+	* Applies a Bend Warp around an axis defined by the Bend Orientation transform.  
+	* The extents of the affected region can be controlled by the Options.
+	*/
 	UFUNCTION(BlueprintCallable, Category = "GeometryScript|Deformations", meta=(ScriptMethod))
 	static UPARAM(DisplayName = "Target Mesh") UDynamicMesh* 
 	ApplyBendWarpToMesh(  
@@ -210,6 +214,10 @@ public:
 		float BendExtent = 50,
 		UGeometryScriptDebug* Debug = nullptr);
 
+	/**
+	 * Applies a twist warp around an axis defined by the Twist Orientation transform.
+	 * The extents of the affected region can be controlled by the Options.
+	 */
 	UFUNCTION(BlueprintCallable, Category = "GeometryScript|Deformations", meta=(ScriptMethod))
 	static UPARAM(DisplayName = "Target Mesh") UDynamicMesh* 
 	ApplyTwistWarpToMesh(  
@@ -220,6 +228,11 @@ public:
 		float TwistExtent = 50,
 		UGeometryScriptDebug* Debug = nullptr);
 
+	/**
+	* Applies a Flare/Bulge warp around an axis defined by the Flare Orientation transform.
+	* The amount of flare in the perpendicular directions can be controlled by FlarePercentX and FlarePercentY
+	* and the extents of the affected region can be controlled by the Options.
+	*/
 	UFUNCTION(BlueprintCallable, Category = "GeometryScript|Deformations", meta=(ScriptMethod))
 	static UPARAM(DisplayName = "Target Mesh") UDynamicMesh* 
 	ApplyFlareWarpToMesh(  
@@ -232,6 +245,10 @@ public:
 		UGeometryScriptDebug* Debug = nullptr);
 
 
+	/**
+	* Applies various simple math-function-based warps around an axis defined by the Warp Orientation transform, 
+	* currently a 1D or 2D sine-wave with arbitrary orientation may be selected by WarpType.
+	*/
 	UFUNCTION(BlueprintCallable, Category = "GeometryScript|Deformations", meta=(ScriptMethod))
 	static UPARAM(DisplayName = "Target Mesh") UDynamicMesh* 
 	ApplyMathWarpToMesh(  
@@ -241,7 +258,11 @@ public:
 		FGeometryScriptMathWarpOptions Options,
 		UGeometryScriptDebug* Debug = nullptr);
 
-
+	/**
+	* Applies 3D Perlin noise displacement to the Target Mesh.
+	* @param Selection if non-empty, only the vertices identified by the selection will be displaced, otherwise
+	* the Options' EmptyBehavior will be followed.
+	*/
 	UFUNCTION(BlueprintCallable, Category = "GeometryScript|Deformations", meta=(ScriptMethod))
 	static UPARAM(DisplayName = "Target Mesh") UDynamicMesh* 
 	ApplyPerlinNoiseToMesh(  
@@ -250,6 +271,11 @@ public:
 		FGeometryScriptPerlinNoiseOptions Options,
 		UGeometryScriptDebug* Debug = nullptr);
 
+	/**
+	* Applies a number of iterations of mesh smoothing to a Dynamic Mesh.
+	* @param Selection if non-empty, only the vertices identified by the selection will be subject to smoothing,  otherwise
+	* the Options' EmptyBehavior will be followed.
+	*/
 	UFUNCTION(BlueprintCallable, Category = "GeometryScript|Deformations", meta=(ScriptMethod))
 	static UPARAM(DisplayName = "Target Mesh") UDynamicMesh* 
 	ApplyIterativeSmoothingToMesh(  
@@ -258,7 +284,11 @@ public:
 		FGeometryScriptIterativeMeshSmoothingOptions Options,
 		UGeometryScriptDebug* Debug = nullptr);
 
-
+	/**
+	 * Applies a displacement to a Dynamic Mesh based on a Texture2D and a UV Channel.
+	 * @param Selection if non-empty, only the vertices identified by the selection will be subject to displacement,  otherwise
+	 * the Options' EmptyBehavior will be followed.
+	 */
 	UFUNCTION(BlueprintCallable, Category = "GeometryScript|Deformations", meta=(ScriptMethod))
 	static UPARAM(DisplayName = "Target Mesh") UDynamicMesh* 
 	ApplyDisplaceFromTextureMap(  
@@ -266,12 +296,12 @@ public:
 		UTexture2D* Texture,
 		FGeometryScriptMeshSelection Selection,
 		FGeometryScriptDisplaceFromTextureOptions Options,
-		int32 UVLayer = 0,
+		UPARAM(DisplayName = "UV Channel") int32 UVLayer = 0,
 		UGeometryScriptDebug* Debug = nullptr);
 
 	/**
-	 * Add the vectors in VectorList, scaled by Magnitude, to the vertex positions in TargetMesh.
-	 * VectorList Length must be >= the as TargetMesh MaxVertexID.
+	 * Add the vectors in VectorList, scaled by Magnitude, to the vertex positions in Target Mesh.
+	 * VectorList Length must be >= the MaxVertexID of the Target Mesh.
 	 * @param Selection if non-empty, only the vertices identified by the selection will be displaced. The VectorList must still be the same size as the whole mesh, this is just a filter on which vertices are updated.
 	 */
 	UFUNCTION(BlueprintCallable, Category = "GeometryScript|Deformations", meta=(ScriptMethod))
diff --git a/Engine/Plugins/Experimental/GeometryScripting/Source/GeometryScriptingCore/Public/GeometryScript/MeshMaterialFunctions.h b/Engine/Plugins/Experimental/GeometryScripting/Source/GeometryScriptingCore/Public/GeometryScript/MeshMaterialFunctions.h
--- a/Engine/Plugins/Experimental/GeometryScripting/Source/GeometryScriptingCore/Public/GeometryScript/MeshMaterialFunctions.h
+++ b/Engine/Plugins/Experimental/GeometryScripting/Source/GeometryScriptingCore/Public/GeometryScript/MeshMaterialFunctions.h
@@ -20,14 +20,25 @@ public:
 	UFUNCTION(BlueprintPure, Category = "GeometryScript|MeshQueries", meta=(ScriptMethod))
 	static UPARAM(DisplayName = "Enabled") int GetMaxMaterialID( UDynamicMesh* TargetMesh, bool& bHasMaterialIDs );
 
+	/**
+	* Enables per-triangle Material IDs on a mesh and initializes the values to 0. 
+	* If Target Mesh already has Material IDs, this function will do nothing.
+	*/
 	UFUNCTION(BlueprintCallable, Category = "GeometryScript|Materials", meta=(ScriptMethod))
 	static UPARAM(DisplayName = "Target Mesh") UDynamicMesh* 
 	EnableMaterialIDs( UDynamicMesh* TargetMesh, UGeometryScriptDebug* Debug = nullptr );
 
+	/**
+	* Resets all Material IDs on a mesh to the given ClearValue, or 0 if no ClearValue is provided.
+	* If Material IDs are not already enabled on the Target Mesh, this function will first enable them and then set the value.
+	*/
 	UFUNCTION(BlueprintCallable, Category = "GeometryScript|Materials", meta=(ScriptMethod))
 	static UPARAM(DisplayName = "Target Mesh") UDynamicMesh* 
 	ClearMaterialIDs( UDynamicMesh* TargetMesh, int ClearValue = 0, UGeometryScriptDebug* Debug = nullptr );
 
+	/**
+	* For all triangles with a Material ID matching the given value (From Material ID), update the Material ID to the new value (To Material ID).
+	*/
 	UFUNCTION(BlueprintCallable, Category = "GeometryScript|Materials", meta=(ScriptMethod))
 	static UPARAM(DisplayName = "Target Mesh") UDynamicMesh* 
 	RemapMaterialIDs( 
@@ -36,10 +47,19 @@ public:
 		int ToMaterialID,
 		UGeometryScriptDebug* Debug = nullptr );
 
+	/**
+	 * Returns the current Material ID for a Triangle.  
+	 * If the mesh does not have Material IDs enabled or if the Triangle ID is not an element of the mesh, the value 0 will be returned and bIsValidTriangle will be false.
+	 */
 	UFUNCTION(BlueprintPure, Category = "GeometryScript|Materials", meta=(ScriptMethod, DisplayName="GetMaterialIDOfTriangle"))
 	static UPARAM(DisplayName = "Material ID") int32
 	GetTriangleMaterialID( UDynamicMesh* TargetMesh, int TriangleID, bool& bIsValidTriangle );
 
+	/**
+	* This populates the MaterialIDList with Material IDs for each triangle in the TriangleIDList.  
+	* If a triangle is not present in the Target Mesh the number -1 will be used for the corresponding Material ID.  
+	* If Material IDs are not enabled on the TargetMesh no Material IDs will be added to the result list.
+	*/
 	UFUNCTION(BlueprintPure, Category = "GeometryScript|Materials", meta=(ScriptMethod))
 	static UPARAM(DisplayName = "Target Mesh") UDynamicMesh* 
 	GetMaterialIDsOfTriangles(
@@ -47,11 +67,19 @@ public:
 		FGeometryScriptIndexList TriangleIDList,
 		FGeometryScriptIndexList& MaterialIDList,
 		UGeometryScriptDebug* Debug = nullptr);
-
+	
+	/**
+	* Returns an Index List of all triangle Material IDs, constructed with one entry for each consecutive Triangle ID.
+	* If Material IDs are not enabled on the mesh, bHasMaterialsIDs will be set to false on return and nothing will be added to the Material ID List.
+	* @warning if the mesh is not Triangle-Compact (eg GetHasTriangleIDGaps == false) then the returned list will also have the same gaps where the number -1 will be recorded for any missing Triangle IDs.
+	*/
 	UFUNCTION(BlueprintPure, Category = "GeometryScript|Materials", meta=(ScriptMethod, DisplayName="GetMaterialIDsOfAllTriangles"))
 	static UPARAM(DisplayName = "Target Mesh") UDynamicMesh* 
 	GetAllTriangleMaterialIDs( UDynamicMesh* TargetMesh, FGeometryScriptIndexList& MaterialIDList, bool& bHasMaterialIDs );
 
+	/**
+	 * Populates Triangle ID List with the Triangle IDs of triangles that share the specified Material ID in the Target Mesh.
+	 */
 	UFUNCTION(BlueprintPure, Category = "GeometryScript|Selection", meta=(ScriptMethod))
 	static UPARAM(DisplayName = "Target Mesh") UDynamicMesh* 
 	GetTrianglesByMaterialID( 
@@ -60,6 +88,10 @@ public:
 		FGeometryScriptIndexList& TriangleIDList,
 		UGeometryScriptDebug* Debug = nullptr);
 
+	/**
+	* Assigns the specified triangle the given Material ID.  
+	* If the Target Mesh does not have Material IDs enabled, or if the Triangle ID is not an element of the Target Mesh then bIsValidTriangle will be set to false.
+	*/
 	UFUNCTION(BlueprintCallable, Category = "GeometryScript|Materials", meta=(ScriptMethod, DisplayName="SetMaterialIDOnTriangle"))
 	static UPARAM(DisplayName = "Target Mesh") UDynamicMesh* 
 	SetTriangleMaterialID( 
@@ -69,6 +101,11 @@ public:
 		bool& bIsValidTriangle, 
 		bool bDeferChangeNotifications = false );
 
+	/**
+	* Sets the Material ID of all triangles in a mesh to the values in an input Index List.
+	* @param TriangleMateralIDList the list of Material IDs, the length of this list should be the same as the Max Triangle ID for this mesh.
+	* @param bDeferChangeNotifications if true, the UDynamicMesh does not emit a change event/signal for this modification.
+	*/
 	UFUNCTION(BlueprintCallable, Category = "GeometryScript|Materials", meta=(ScriptMethod, DisplayName="SetMaterialIDsOnAllTriangles"))
 	static UPARAM(DisplayName = "Target Mesh") UDynamicMesh*
 	SetAllTriangleMaterialIDs(
@@ -77,6 +114,12 @@ public:
 		bool bDeferChangeNotifications = false,
 		UGeometryScriptDebug* Debug = nullptr);
 
+	/**
+	* Assigns the Material ID to all the triangles specified by the Triangle ID List.
+	* @param TriangleIDList the triangles in the target mesh that will be updated with the new Material ID
+	* @param MaterialID the ID to be assigned to each triangle in the input list.
+	* @param bDeferChangeNotifications if true, the UDynamicMesh does not emit a change event/signal for this modification.
+	*/
 	UFUNCTION(BlueprintCallable, Category = "GeometryScript|Materials", meta=(ScriptMethod))
 	static UPARAM(DisplayName = "Target Mesh") UDynamicMesh* 
 	SetMaterialIDOnTriangles(
@@ -103,7 +146,7 @@ public:
 	/**
 	* Set a new MaterialID on all the triangles of TargetMesh with the given PolyGroup.
 	* @param GroupLayer PolyGroup Layer to use as basis for PolyGroups
-	* @param PolygroupID PolyGroup ID that specifies Triangles to set to new MaterialID
+	* @param PolyGroupID PolyGroup ID that specifies Triangles to set to new MaterialID
 	* @param MaterialID explicit new MaterialID to set
 	* @param bDeferChangeNotifications if true, the UDynamicMesh does not emit a change event/signal for this modification
 	*/
diff --git a/Engine/Plugins/Experimental/GeometryScripting/Source/GeometryScriptingCore/Public/GeometryScript/MeshModelingFunctions.h b/Engine/Plugins/Experimental/GeometryScripting/Source/GeometryScriptingCore/Public/GeometryScript/MeshModelingFunctions.h
--- a/Engine/Plugins/Experimental/GeometryScripting/Source/GeometryScriptingCore/Public/GeometryScript/MeshModelingFunctions.h
+++ b/Engine/Plugins/Experimental/GeometryScripting/Source/GeometryScriptingCore/Public/GeometryScript/MeshModelingFunctions.h
@@ -373,6 +373,9 @@ public:
 		UGeometryScriptDebug* Debug = nullptr );
 
 
+	/**
+	 * Apply a Mesh Bevel operation to all PolyGroup Edges.
+     */
 	UFUNCTION(BlueprintCallable, Category = "GeometryScript|Modeling", meta=(ScriptMethod, DisplayName = "ApplyMeshPolyGroupBevel"))
 	static UPARAM(DisplayName = "Target Mesh") UDynamicMesh* 
 	ApplyMeshPolygroupBevel(
diff --git a/Engine/Plugins/Experimental/GeometryScripting/Source/GeometryScriptingCore/Public/GeometryScript/MeshPolygroupFunctions.h b/Engine/Plugins/Experimental/GeometryScripting/Source/GeometryScriptingCore/Public/GeometryScript/MeshPolygroupFunctions.h
--- a/Engine/Plugins/Experimental/GeometryScripting/Source/GeometryScriptingCore/Public/GeometryScript/MeshPolygroupFunctions.h
+++ b/Engine/Plugins/Experimental/GeometryScripting/Source/GeometryScriptingCore/Public/GeometryScript/MeshPolygroupFunctions.h
@@ -16,11 +16,16 @@ class GEOMETRYSCRIPTINGCORE_API UGeometryScriptLibrary_MeshPolygroupFunctions :
 {
 	GENERATED_BODY()
 public:
-
+	/**
+	* Enables the standard PolyGroup Layer on the Target Mesh.
+	*/
 	UFUNCTION(BlueprintCallable, Category = "GeometryScript|PolyGroups", meta = (ScriptMethod, DisplayName = "EnablePolyGroups"))
 	static UPARAM(DisplayName = "Target Mesh") UDynamicMesh* 
 	EnablePolygroups( UDynamicMesh* TargetMesh, UGeometryScriptDebug* Debug = nullptr );
 
+	/**
+	* Sets the number of extended PolyGroup Layers on a Mesh.
+	*/
 	UFUNCTION(BlueprintCallable, Category = "GeometryScript|PolyGroups", meta = (ScriptMethod, DisplayName = "SetNumExtendedPolyGroupLayers"))
 	static UPARAM(DisplayName = "Target Mesh") UDynamicMesh* 
 	SetNumExtendedPolygroupLayers( 
@@ -28,6 +33,10 @@ public:
 		int NumLayers,
 		UGeometryScriptDebug* Debug = nullptr );
 
+	/**
+	* Resets the triangle PolyGroup assignments within a PolyGroup Layer to the given Clear Value (or 0 if no Clear Value is specified).
+	* Note, this will have no effect if PolyGroups have not been enabled on the mesh, or if the requested Group Layer does not exist. 
+	*/
 	UFUNCTION(BlueprintCallable, Category = "GeometryScript|PolyGroups", meta = (ScriptMethod, DisplayName = "ClearPolyGroups"))
 	static UPARAM(DisplayName = "Target Mesh") UDynamicMesh* 
 	ClearPolygroups( 
@@ -36,6 +45,10 @@ public:
 		int ClearValue = 0,
 		UGeometryScriptDebug* Debug = nullptr );
 
+	/**
+	* Copies the triangle PolyGroup assignments from one layer on the Target Mesh to another.
+	* Note, this will have no effect if PolyGroups have not been enabled on the mesh, or if one of the requested Group Layers does not exist.
+	*/
 	UFUNCTION(BlueprintCallable, Category = "GeometryScript|PolyGroups", meta = (ScriptMethod, DisplayName = "CopyPolyGroupsLayer"))
 	static UPARAM(DisplayName = "Target Mesh") UDynamicMesh* 
 	CopyPolygroupsLayer( 
@@ -44,7 +57,12 @@ public:
 		FGeometryScriptGroupLayer ToGroupLayer,
 		UGeometryScriptDebug* Debug = nullptr );
 
-
+	/**
+	* Creates and assigns a new PolyGroup for each disconnected UV island of a Mesh.
+	* Note, this will have no effect if either the requested UV Layer or Group Layer does not exist.
+	* @param GroupLayer indicates PolyGroup Layer that will be populated with unique values for each UV island.
+	* @param UVLayer specifies the UV Layer used to construct the PolyGroups.
+	*/
 	UFUNCTION(BlueprintCallable, Category = "GeometryScript|PolyGroups", meta = (ScriptMethod, DisplayName = "ConvertUVIslandsToPolyGroups"))
 	static UPARAM(DisplayName = "Target Mesh") UDynamicMesh* 
 	ConvertUVIslandsToPolygroups( 
@@ -53,6 +71,10 @@ public:
 		int UVLayer = 0,
 		UGeometryScriptDebug* Debug = nullptr );
 
+	/**
+	* Creates and assigns a new PolyGroup for each disconnected component of a Mesh. Regions of a mesh are disconnected they do not have a triangle in common.
+	* Note, this will have no effect if the Group Layer does not exist.
+	*/
 	UFUNCTION(BlueprintCallable, Category = "GeometryScript|PolyGroups", meta = (ScriptMethod, DisplayName = "ConvertComponentsToPolyGroups"))
 	static UPARAM(DisplayName = "Target Mesh") UDynamicMesh* 
 	ConvertComponentsToPolygroups( 
@@ -60,6 +82,13 @@ public:
 		FGeometryScriptGroupLayer GroupLayer,
 		UGeometryScriptDebug* Debug = nullptr );
 
+	/**
+	* Sets PolyGroups by partitioning the mesh based on an edge crease/opening-angle.
+	* Note, this will have no effect if the Group Layer does not exist.
+	* @param GroupLayer indicates the PolyGroup Layer that will be populated.
+	* @param CreaseAngle measured in degrees and used when comparing adjacent faces.
+	* @param MinGroupSize the minimum number of triangles in each PolyGroup.
+	*/
 	UFUNCTION(BlueprintCallable, Category = "GeometryScript|PolyGroups", meta = (ScriptMethod, DisplayName = "ComputePolyGroupsFromAngleThreshold"))
 	static UPARAM(DisplayName = "Target Mesh") UDynamicMesh* 
 	ComputePolygroupsFromAngleThreshold( 
@@ -69,6 +98,9 @@ public:
 		int MinGroupSize = 2,
 		UGeometryScriptDebug* Debug = nullptr );
 
+	/**
+	* Sets PolyGroups by identifying adjacent triangles that form reasonable quads. Note any triangles that do not neatly pair to form quads will receive their own PolyGroup.  
+	*/
 	UFUNCTION(BlueprintCallable, Category = "GeometryScript|PolyGroups", meta = (ScriptMethod, DisplayName = "ComputePolyGroupsFromPolygonDetection"))
 	static UPARAM(DisplayName = "Target Mesh") UDynamicMesh* 
 	ComputePolygroupsFromPolygonDetection( 
@@ -81,6 +113,13 @@ public:
 		int MaxSearchRounds = 1,
 		UGeometryScriptDebug* Debug = nullptr );
 
+	/**
+	* Gets the PolyGroup ID associated with the specified Triangle ID and stored in the Group Layer.
+	* If the Group Layer or the Triangle does not exist, the value 0 will be returned and bIsValidTriangle set to false.
+	* @param GroupLayer indicates the layer on the Target Mesh to query.
+	* @param TriangleID identifies a triangle in the Target Mesh.
+	* @param bIsValidTriangle will be populated on return with false if either the Group Layer or the Triangle does not exist.
+	*/
 	UFUNCTION(BlueprintPure, Category = "GeometryScript|PolyGroups", meta = (ScriptMethod, DisplayName = "GetTrianglePolyGroupID"))
 	static UPARAM(DisplayName = "PolyGroup ID") int32
 	GetTrianglePolygroupID( 
@@ -89,6 +128,12 @@ public:
 		int TriangleID, 
 		bool& bIsValidTriangle );
 
+	/**
+	 * Deletes all triangles from the Target Mesh that have a particular PolyGroup ID, in the specific Group Layer.
+	 * @param GroupLayer specifies the PolyGroup Layer to query.
+	 * @param PolyGroup ID identifies the triangles in the Target Mesh to delete.
+	 * @param NumDeleted on return will contain the number of triangles deleted from the Target Mesh.
+	 */
 	UFUNCTION(BlueprintCallable, Category = "GeometryScript|PolyGroups", meta = (ScriptMethod, DisplayName = "DeleteTrianglesInPolyGroup"))
 	static UPARAM(DisplayName = "Target Mesh") UDynamicMesh* 
 	DeleteTrianglesInPolygroup( 
diff --git a/Engine/Plugins/Experimental/GeometryScripting/Source/GeometryScriptingCore/Public/GeometryScript/MeshPrimitiveFunctions.h b/Engine/Plugins/Experimental/GeometryScripting/Source/GeometryScriptingCore/Public/GeometryScript/MeshPrimitiveFunctions.h
--- a/Engine/Plugins/Experimental/GeometryScripting/Source/GeometryScriptingCore/Public/GeometryScript/MeshPrimitiveFunctions.h
+++ b/Engine/Plugins/Experimental/GeometryScripting/Source/GeometryScriptingCore/Public/GeometryScript/MeshPrimitiveFunctions.h
@@ -101,7 +101,9 @@ class GEOMETRYSCRIPTINGCORE_API UGeometryScriptLibrary_MeshPrimitiveFunctions :
 	GENERATED_BODY()
 public:
 
-
+	/**
+	* Appends a 3D box to the Target Mesh.
+	*/
 	UFUNCTION(BlueprintCallable, Category = "GeometryScript|Primitives", meta=(ScriptMethod))
 	static UPARAM(DisplayName = "Target Mesh") UDynamicMesh* 
 	AppendBox( 
@@ -117,7 +119,9 @@ public:
 		EGeometryScriptPrimitiveOriginMode Origin = EGeometryScriptPrimitiveOriginMode::Base,
 		UGeometryScriptDebug* Debug = nullptr);
 
-
+	/**
+	* Appends a 3D Sphere triangulated using latitude/longitude topology to the Target Mesh.
+	*/ 
 	UFUNCTION(BlueprintCallable, Category = "GeometryScript|Primitives", meta=(ScriptMethod))
 	static UPARAM(DisplayName = "Target Mesh") UDynamicMesh* 
 	AppendSphereLatLong( 
@@ -130,7 +134,9 @@ public:
 		EGeometryScriptPrimitiveOriginMode Origin = EGeometryScriptPrimitiveOriginMode::Center,
 		UGeometryScriptDebug* Debug = nullptr);
 
-
+	/**
+	* Appends a 3D sphere triangulated using box topology to the Target Mesh.
+	*/ 
 	UFUNCTION(BlueprintCallable, Category = "GeometryScript|Primitives", meta=(ScriptMethod))
 	static UPARAM(DisplayName = "Target Mesh") UDynamicMesh* 
 	AppendSphereBox( 
@@ -144,7 +150,9 @@ public:
 		EGeometryScriptPrimitiveOriginMode Origin = EGeometryScriptPrimitiveOriginMode::Center,
 		UGeometryScriptDebug* Debug = nullptr);
 
-
+	/**
+	* Appends a 3D Capsule to the Target Mesh.
+	*/ 
 	UFUNCTION(BlueprintCallable, Category = "GeometryScript|Primitives", meta=(ScriptMethod))
 	static UPARAM(DisplayName = "Target Mesh") UDynamicMesh* 
 	AppendCapsule( 
@@ -158,7 +166,9 @@ public:
 		EGeometryScriptPrimitiveOriginMode Origin = EGeometryScriptPrimitiveOriginMode::Base,
 		UGeometryScriptDebug* Debug = nullptr);
 
-
+	/**
+	* Appends a 3D Cylinder (with optional end caps) to the Target Mesh.
+	*/
 	UFUNCTION(BlueprintCallable, Category = "GeometryScript|Primitives", meta=(ScriptMethod))
 	static UPARAM(DisplayName = "Target Mesh") UDynamicMesh* 
 	AppendCylinder( 
@@ -173,7 +183,9 @@ public:
 		EGeometryScriptPrimitiveOriginMode Origin = EGeometryScriptPrimitiveOriginMode::Base,
 		UGeometryScriptDebug* Debug = nullptr);
 
-
+	/**
+	* Appends a 3D cone to the Target Mesh.
+	*/ 
 	UFUNCTION(BlueprintCallable, Category = "GeometryScript|Primitives", meta=(ScriptMethod))
 	static UPARAM(DisplayName = "Target Mesh") UDynamicMesh* 
 	AppendCone( 
@@ -189,7 +201,9 @@ public:
 		EGeometryScriptPrimitiveOriginMode Origin = EGeometryScriptPrimitiveOriginMode::Base,
 		UGeometryScriptDebug* Debug = nullptr);
 
-
+	/**
+	* Appends a 3D torus (donut) or partial torus to the Target Mesh.
+	*/
 	UFUNCTION(BlueprintCallable, Category = "GeometryScript|Primitives", meta=(ScriptMethod))
 	static UPARAM(DisplayName = "Target Mesh") UDynamicMesh* 
 	AppendTorus( 
@@ -221,7 +235,9 @@ public:
 		int32 Steps = 8,
 		UGeometryScriptDebug* Debug = nullptr);
 
-
+    /**
+	* Revolves a 2D polygon on a helical path, like one used to create a vertical spiral, appending the result to the Target Mesh.
+	*/ 
 	UFUNCTION(BlueprintCallable, Category = "GeometryScript|Primitives", meta = (ScriptMethod))
 	static UPARAM(DisplayName = "Target Mesh") UDynamicMesh*
 	AppendSpiralRevolvePolygon(
@@ -235,7 +251,9 @@ public:
 		float RisePerRevolution = 50,
 		UGeometryScriptDebug* Debug = nullptr);
 
-
+	/**
+	* Revolves an open 2D path, with optional top and bottom end caps, appending the result to the Target Mesh.
+	*/ 
 	UFUNCTION(BlueprintCallable, Category = "GeometryScript|Primitives", meta=(ScriptMethod))
 	static UPARAM(DisplayName = "Target Mesh") UDynamicMesh* 
 	AppendRevolvePath( 
@@ -294,7 +312,9 @@ public:
 		EGeometryScriptPrimitiveOriginMode Origin = EGeometryScriptPrimitiveOriginMode::Base,
 		UGeometryScriptDebug* Debug = nullptr);
 
-
+	/**
+	* Sweeps a 2D polygon along an arbitrary 3D path, appending the result to the Target Mesh.
+	*/
 	UFUNCTION(BlueprintCallable, Category = "GeometryScript|Primitives", meta=(ScriptMethod))
 	static UPARAM(DisplayName = "Target Mesh") UDynamicMesh* 
 	AppendSimpleSweptPolygon( 
@@ -337,6 +357,9 @@ public:
 		UGeometryScriptDebug* Debug = nullptr);
 
 
+	/**
+	* Appends a planar Rectangle to a Dynamic Mesh.
+	*/ 
 	UFUNCTION(BlueprintCallable, Category = "GeometryScript|Primitives", meta=(ScriptMethod))
 	static UPARAM(DisplayName = "Target Mesh") UDynamicMesh* 
 	AppendRectangleXY( 
@@ -349,6 +372,9 @@ public:
 		int32 StepsHeight = 0,
 		UGeometryScriptDebug* Debug = nullptr);
 
+	/**
+	* Appends a planar Rectangle with Rounded Corners (RoundRect) to the Target Mesh.
+	*/
 	UFUNCTION(BlueprintCallable, Category = "GeometryScript|Primitives", meta=(ScriptMethod))
 	static UPARAM(DisplayName = "Target Mesh") UDynamicMesh* 
 	AppendRoundRectangleXY( 
@@ -363,6 +389,9 @@ public:
 		int32 StepsRound = 4,
 		UGeometryScriptDebug* Debug = nullptr);
 
+	/**
+	* Appends a planar disc to the Target Mesh.
+	*/ 
 	UFUNCTION(BlueprintCallable, Category = "GeometryScript|Primitives", meta=(ScriptMethod))
 	static UPARAM(DisplayName = "Target Mesh") UDynamicMesh* 
 	AppendDisc( 
@@ -378,6 +407,7 @@ public:
 		UGeometryScriptDebug* Debug = nullptr);
 
 	/**
+	* Appends a Triangulated Polygon to the Target Mesh.
 	* Polygon should be oriented counter-clockwise to produce a correctly-oriented shape, otherwise it will be inside-out
 	* Polygon endpoint is not repeated.
 	*/
@@ -392,7 +422,9 @@ public:
 		UGeometryScriptDebug* Debug = nullptr);
 
 
-
+	/**
+	* Appends a linear staircase to the Target Mesh.
+	*/
 	UFUNCTION(BlueprintCallable, Category = "GeometryScript|Primitives", meta=(ScriptMethod))
 	static UPARAM(DisplayName = "Target Mesh") UDynamicMesh* 
 	AppendLinearStairs( 
@@ -406,7 +438,9 @@ public:
 		bool bFloating = false,
 		UGeometryScriptDebug* Debug = nullptr);
 
-
+	/**
+	* Appends a rising circular staircase to the Target Mesh.
+	*/ 
 	UFUNCTION(BlueprintCallable, Category = "GeometryScript|Primitives", meta=(ScriptMethod))
 	static UPARAM(DisplayName = "Target Mesh") UDynamicMesh* 
 	AppendCurvedStairs( 
@@ -421,7 +455,9 @@ public:
 		bool bFloating = false,
 		UGeometryScriptDebug* Debug = nullptr);
 
-
+	/**
+	* Generates triangulated Voronoi Cells from the provided Voronoi Sites, identifying each with PolyGroups, and appends to the Target Mesh.
+	*/ 
 	UFUNCTION(BlueprintCallable, Category = "GeometryScript|Primitives", meta = (ScriptMethod))
 	static UPARAM(DisplayName = "Target Mesh") UDynamicMesh*
 	AppendVoronoiDiagram2D(
diff --git a/Engine/Plugins/Experimental/GeometryScripting/Source/GeometryScriptingCore/Public/GeometryScript/MeshQueryFunctions.h b/Engine/Plugins/Experimental/GeometryScripting/Source/GeometryScriptingCore/Public/GeometryScript/MeshQueryFunctions.h
--- a/Engine/Plugins/Experimental/GeometryScripting/Source/GeometryScriptingCore/Public/GeometryScript/MeshQueryFunctions.h
+++ b/Engine/Plugins/Experimental/GeometryScripting/Source/GeometryScriptingCore/Public/GeometryScript/MeshQueryFunctions.h
@@ -16,30 +16,58 @@ class GEOMETRYSCRIPTINGCORE_API UGeometryScriptLibrary_MeshQueryFunctions : publ
 	GENERATED_BODY()
 public:
 
+	/**
+	* Returns information about the Target Mesh, such as the vertex and triangle count as well as some attribute information.
+	*/
 	UFUNCTION(BlueprintPure, Category = "GeometryScript|MeshQueries", meta=(ScriptMethod))
 	static FString GetMeshInfoString( UDynamicMesh* TargetMesh );
 
+	/**
+	* Returns true if the mesh is dense. For example, no gaps in Vertex IDs or Triangle IDs.
+	* Note if a mesh is not dense, the Compact Mesh node can be used to removed the gaps.
+	*/
 	UFUNCTION(BlueprintPure, Category = "GeometryScript|MeshQueries", meta=(ScriptMethod))
 	static bool GetIsDenseMesh( UDynamicMesh* TargetMesh );
 
+	/**
+	* Returns true if the Target Mesh has attributes enabled.
+	*/
 	UFUNCTION(BlueprintPure, Category = "GeometryScript|MeshQueries", meta=(ScriptMethod))
 	static bool GetMeshHasAttributeSet( UDynamicMesh* TargetMesh );
 
+	/**
+	* Computes the bounding box of the mesh vertices in the local space of the mesh.
+	*/
 	UFUNCTION(BlueprintPure, Category = "GeometryScript|MeshQueries", meta=(ScriptMethod))
 	static UPARAM(DisplayName = "Bounding Box") FBox GetMeshBoundingBox( UDynamicMesh* TargetMesh );
 
+	/**
+	* Computes the volume and area of the mesh.
+	*/
 	UFUNCTION(BlueprintPure, Category = "GeometryScript|MeshQueries", meta=(ScriptMethod))
 	static void GetMeshVolumeArea( UDynamicMesh* TargetMesh, float& SurfaceArea, float& Volume );
 
+	/**
+	* Returns true if the mesh is closed, such as no topological boundary edges.
+	*/
 	UFUNCTION(BlueprintPure, Category = "GeometryScript|MeshQueries", meta=(ScriptMethod))
 	static bool GetIsClosedMesh( UDynamicMesh* TargetMesh );
 
+	/**
+	* Returns the number of open border loops, such as "holes" in the mesh.
+	*/
 	UFUNCTION(BlueprintPure, Category = "GeometryScript|MeshQueries", meta=(ScriptMethod))
 	static UPARAM(DisplayName = "Num Loops") int32 GetNumOpenBorderLoops( UDynamicMesh* TargetMesh, bool& bAmbiguousTopologyFound );
 
+	/**
+	* Returns the number of topological boundary edges in the mesh, i.e counts edges that only have one adjacent triangle.
+	*/
 	UFUNCTION(BlueprintPure, Category = "GeometryScript|MeshQueries", meta=(ScriptMethod))
 	static UPARAM(DisplayName = "Num Edges") int32 GetNumOpenBorderEdges( UDynamicMesh* TargetMesh );
 
+	/**
+	* Returns the number of separate connected components in the mesh, such as "triangle patches" internally connected by shared edges.
+	*/
 	UFUNCTION(BlueprintPure, Category = "GeometryScript|MeshQueries", meta=(ScriptMethod))
 	static UPARAM(DisplayName = "Num Edges") int32 GetNumConnectedComponents( UDynamicMesh* TargetMesh );
 
@@ -47,27 +75,53 @@ public:
 	//UFUNCTION(BlueprintPure, Category = "GeometryScript|MeshQueries", meta=(ScriptMethod))
 	//static UPARAM(DisplayName = "Triangle Count") int32 GetTriangleCount( UDynamicMesh* TargetMesh );
 
+	/**
+	* Gets the number of Triangle IDs in the mesh. This may be larger than the Triangle Count if the mesh is not dense, even after deleting triangles.
+	*/
 	UFUNCTION(BlueprintPure, Category = "GeometryScript|MeshQueries", meta=(ScriptMethod))
 	static UPARAM(DisplayName = "Num Triangle IDs") int32 GetNumTriangleIDs( UDynamicMesh* TargetMesh );
 
+	/**
+	* Returns true if there are gaps in the Triangle ID list, such that Get Num Triangle IDs is greater than Get Triangle Count.
+	*/
 	UFUNCTION(BlueprintPure, Category = "GeometryScript|MeshQueries", meta=(ScriptMethod))
 	static bool GetHasTriangleIDGaps( UDynamicMesh* TargetMesh );
 
+	/**
+	* Returns true if Triangle ID refers to a valid Triangle in the Target Mesh.
+	*/
 	UFUNCTION(BlueprintPure, Category = "GeometryScript|MeshQueries", meta=(ScriptMethod))
 	static bool IsValidTriangleID( UDynamicMesh* TargetMesh, int32 TriangleID );
 
+	/**
+	* Returns an Index List of all Triangle IDs in a mesh. 
+	* @param bHasTriangleIDGaps will be true on return if there are breaks in the sequential numeration of Triangle IDs, as would happen after deleting triangles.
+	*/
 	UFUNCTION(BlueprintCallable, Category = "GeometryScript|MeshQueries", meta=(ScriptMethod))
 	static UPARAM(DisplayName = "Target Mesh") UDynamicMesh* 
 	GetAllTriangleIDs( UDynamicMesh* TargetMesh, FGeometryScriptIndexList& TriangleIDList, bool& bHasTriangleIDGaps );
 
-
+	/**
+	* Returns the Vertex ID triplet for the specified Triangle.
+	* @param TriangleID indicates the triangle to query on the Target Mesh.
+	* @param bIsValidTriangle will be false on return if the Triangle ID does not exist in the Target Mesh, in that case the returned vertex triplet will be (-1, -1, -1).
+	*/
 	UFUNCTION(BlueprintPure, Category = "GeometryScript|MeshQueries", meta=(ScriptMethod))
 	static UPARAM(DisplayName = "Indices") FIntVector GetTriangleIndices( UDynamicMesh* TargetMesh, int32 TriangleID, bool& bIsValidTriangle );
 
+	/*
+	* Returns a TriangleList of all Triangle Vertex ID triplets in a mesh.
+	* @param bSkipGaps if false there will be a one-to-one correspondence between Triangle ID and entries in the triangle list and invalid triplets of (-1,-1,-1) will correspond to Triangle IDs not found in the Target Mesh.
+	* @param bHasTriangleIDGaps will be false on return if the mesh had no gaps in Triangle IDs or if bSkipGaps was set to true.
+	*/
 	UFUNCTION(BlueprintCallable, Category = "GeometryScript|MeshQueries", meta=(ScriptMethod))
 	static UPARAM(DisplayName = "Target Mesh") UDynamicMesh* 
 	GetAllTriangleIndices( UDynamicMesh* TargetMesh, FGeometryScriptTriangleList& TriangleList, bool bSkipGaps, bool& bHasTriangleIDGaps);
 
+	/*
+	* Returns the 3D positions (in the mesh local space) of the three vertices of the requested triangle.
+	* If the Triangle ID is not an element of the Target Mesh, all three vertices will be returned as (0, 0, 0) and bIsValidTriangle will be set to false.
+	*/
 	UFUNCTION(BlueprintPure, Category = "GeometryScript|MeshQueries", meta=(ScriptMethod))
 	static void GetTrianglePositions( UDynamicMesh* TargetMesh, int32 TriangleID, bool& bIsValidTriangle, FVector& Vertex1, FVector& Vertex2, FVector& Vertex3 );
 
@@ -110,26 +164,49 @@ public:
 		FVector& BarycentricCoords );
 
 
-
+	/**
+	* Gets the number of vertices in the mesh. Note this may be less than the number of Vertex IDs used as some vertices may have been deleted.
+	*/
 	UFUNCTION(BlueprintPure, Category = "GeometryScript|MeshQueries", meta=(ScriptMethod))
 	static UPARAM(DisplayName = "Vertex Count") int32 GetVertexCount( UDynamicMesh* TargetMesh );
 
+	/**
+	* Gets the number of Vertex IDs in the mesh, which may be larger than the Vertex Count, if the mesh is not dense (e.g.  after deleting vertices).
+	*/
 	UFUNCTION(BlueprintPure, Category = "GeometryScript|MeshQueries", meta=(ScriptMethod))
 	static UPARAM(DisplayName = "Num Vertex IDs") int32 GetNumVertexIDs( UDynamicMesh* TargetMesh );
 
+	/**
+	* Returns true if there are gaps in the Vertex ID list. For example, Get Number of Vertex IDs is greater than Get Vertex Count.
+	*/
 	UFUNCTION(BlueprintPure, Category = "GeometryScript|MeshQueries", meta=(ScriptMethod))
 	static bool GetHasVertexIDGaps( UDynamicMesh* TargetMesh );
 
+	/**
+	* Returns true if Vertex ID refers to a valid vertex in the Target Mesh.
+	*/
 	UFUNCTION(BlueprintPure, Category = "GeometryScript|MeshQueries", meta=(ScriptMethod))
 	static bool IsValidVertexID( UDynamicMesh* TargetMesh, int32 VertexID );
 
+	/**
+	* Returns an IndexList of all Vertex IDs in mesh.
+	*/
 	UFUNCTION(BlueprintCallable, Category = "GeometryScript|MeshQueries", meta=(ScriptMethod))
 	static UPARAM(DisplayName = "Target Mesh") UDynamicMesh* 
 	GetAllVertexIDs( UDynamicMesh* TargetMesh, FGeometryScriptIndexList& VertexIDList, bool& bHasVertexIDGaps );
 
+	/**
+	* Gets the 3D position of a mesh vertex in the mesh local space, by Vertex ID.
+	*/
 	UFUNCTION(BlueprintPure, Category = "GeometryScript|MeshQueries", meta=(ScriptMethod))
 	static UPARAM(DisplayName = "Position") FVector GetVertexPosition( UDynamicMesh* TargetMesh, int32 VertexID, bool& bIsValidVertex );
 
+	/**
+	* Returns a Vector List of all the mesh vertex 3D positions (possibly large!).
+	* @param bSkipGaps if false there will be a one-to-one correspondence between Vertex ID and entries in the Position List
+	* where a zero vector (0,0,0) will correspond to Vertex IDs not found in the Target Mesh.
+	* @param bHasVertexIDGaps will be false if the mesh had no gaps in Vertex IDs or if bSkipGaps was set to true.
+	*/
 	UFUNCTION(BlueprintCallable, Category = "GeometryScript|MeshQueries", meta=(ScriptMethod))
 	static UPARAM(DisplayName = "Target Mesh") UDynamicMesh* 
 	GetAllVertexPositions( UDynamicMesh* TargetMesh, FGeometryScriptVectorList& PositionList, bool bSkipGaps, bool& bHasVertexIDGaps );
@@ -139,16 +216,24 @@ public:
 	// UV Queries
 	//
 
-
+	/**
+	* Gets the number of UV Channels on the Target Mesh.
+	*/
 	UFUNCTION(BlueprintPure, Category = "GeometryScript|MeshQueries", meta=(ScriptMethod), DisplayName = "Get Num UV Channels")
 	static UPARAM(DisplayName = "Num UV Channels") int32 GetNumUVSets( UDynamicMesh* TargetMesh );
 
+	/**
+	* Gets the 2D bounding box of all UVs in the UV Channel.  If the UV Channel does not exist, or if the UV Channel is empty, the resulting box will be invalid.
+	*/
 	UFUNCTION(BlueprintPure, Category = "GeometryScript|MeshQueries", meta=(ScriptMethod))
 	static UPARAM(DisplayName = "Bounding Box") FBox2D GetUVSetBoundingBox( UDynamicMesh* TargetMesh, 
 																			UPARAM(DisplayName = "UV Channel") int UVSetIndex, 
 																			UPARAM(DisplayName = "Is Valid UV Channel") bool& bIsValidUVSet, 
 																			UPARAM(DisplayName = "UV Channel Is Empty") bool& bUVSetIsEmpty );
-
+	/**
+	* Returns the UV values associated with the three vertices of the triangle in the specified UV Channel.
+	* If the Triangle does not exist in the mesh or if no UVs are set in the specified UV Channel for the triangle, the resulting values will be (0,0) and bHaveValidUVs will be set to false.
+	*/
 	UFUNCTION(BlueprintPure, Category = "GeometryScript|MeshQueries", meta=(ScriptMethod))
 	static void GetTriangleUVs( UDynamicMesh* TargetMesh, UPARAM(DisplayName = "UV Channel") int32 UVSetIndex, int32 TriangleID, FVector2D& UV1, FVector2D& UV2, FVector2D& UV3, bool& bHaveValidUVs );
 
@@ -179,9 +264,9 @@ public:
 	static UPARAM(DisplayName = "Enabled") bool GetHasTriangleNormals( UDynamicMesh* TargetMesh );
 
 	/**
-	 * For the specified TriangleID of the TargetMesh, get the Normal vectors at each vertex of the Triangle.
-	 * These Normals will be taken from the Normal Overlay, ie they will potentially be split-normals.
-	 * @param bTriHasValidNormals will be returned true if TriangleID exists in TargetMesh and has Normals set
+	 * For the specified TriangleID of the Target Mesh, get the Normal vectors at each vertex of the Triangle.
+	 * These Normals will be taken from the Normal Overlay, i.e. they will potentially be split-normals.
+	 * @param bTriHasValidNormals will be returned true if TriangleID exists in TargetMesh and has Normals set.
 	 */
 	UFUNCTION(BlueprintCallable, Category = "GeometryScript|MeshQueries", meta=(ScriptMethod))
 	static UPARAM(DisplayName = "Target Mesh") UDynamicMesh* 
@@ -196,7 +281,7 @@ public:
 	/**
 	 * Compute the interpolated Normal (A*Normal1 + B*Normal2 + C*Normal3), where (A,B,C)=BarycentricCoords and the Normals are taken
 	 * from the specified TriangleID in the Normal layer of the TargetMesh. 
-	 * @param bTriHasValidNormals will be returned true if TriangleID exists in TargetMesh and has Normals set, and otherwise will be returned false
+	 * @param bTriHasValidNormals will be returned true if TriangleID exists in TargetMesh and has Normals set, and otherwise will be returned false.
 	 */
 	UFUNCTION(BlueprintCallable, Category = "GeometryScript|MeshQueries", meta=(ScriptMethod))
 	static UPARAM(DisplayName = "Target Mesh") UDynamicMesh* 
@@ -208,9 +293,9 @@ public:
 		FVector& InterpolatedNormal );
 
 	/**
-	 * For the specified TriangleID of the TargetMesh, get the Normal and Tangent vectors at each vertex of the Triangle.
-	 * These Normals/Tangents will be taken from the Normal and Tangents Overlays, ie they will potentially be split-normals.
-	 * @param bTriHasValidElements will be returned true if TriangleID exists in TargetMesh and has Normals and Tangents set
+	 * For the specified Triangle ID of the TargetMesh, get the Normal and Tangent vectors at each vertex of the Triangle.
+	 * These Normals/Tangents will be taken from the Normal and Tangents Overlays, i.e. they will potentially be split-normals.
+	 * @param bTriHasValidElements will be returned true if TriangleID exists in TargetMesh and has Normals and Tangents set.
 	 */
 	UFUNCTION(BlueprintCallable, Category = "GeometryScript|MeshQueries", meta=(ScriptMethod))
 	static UPARAM(DisplayName = "Target Mesh") UDynamicMesh* 
@@ -284,6 +369,9 @@ public:
 	// Material Queries
 	//
 
+	/**
+	 * Returns true if the mesh has Material IDs available/enabled.
+	 */
 	UFUNCTION(BlueprintPure, Category = "GeometryScript|MeshQueries", meta=(ScriptMethod))
 	static UPARAM(DisplayName = "Enabled") bool GetHasMaterialIDs( UDynamicMesh* TargetMesh );
 
@@ -292,9 +380,15 @@ public:
 	// Polygroup Queries
 	//
 
+	/**
+	* Returns true if the mesh has a standard PolyGroup Layer.
+	*/
 	UFUNCTION(BlueprintPure, Category = "GeometryScript|MeshQueries", meta=(ScriptMethod, DisplayName = "GetHasPolyGroups"))
 	static UPARAM(DisplayName = "Enabled") bool GetHasPolygroups( UDynamicMesh* TargetMesh );
 
+	/**
+	* Returns the count of extended PolyGroup Layers.
+	*/
 	UFUNCTION(BlueprintPure, Category = "GeometryScript|MeshQueries", meta=(ScriptMethod, DisplayName = "GetNumExtendedPolyGroupLayers"))
 	static UPARAM(DisplayName = "Num Layers") int GetNumExtendedPolygroupLayers( UDynamicMesh* TargetMesh );
 
diff --git a/Engine/Plugins/Experimental/GeometryScripting/Source/GeometryScriptingCore/Public/GeometryScript/MeshSelectionFunctions.h b/Engine/Plugins/Experimental/GeometryScripting/Source/GeometryScriptingCore/Public/GeometryScript/MeshSelectionFunctions.h
--- a/Engine/Plugins/Experimental/GeometryScripting/Source/GeometryScriptingCore/Public/GeometryScript/MeshSelectionFunctions.h
+++ b/Engine/Plugins/Experimental/GeometryScripting/Source/GeometryScriptingCore/Public/GeometryScript/MeshSelectionFunctions.h
@@ -117,7 +117,7 @@ public:
 		FGeometryScriptMeshSelection& Selection);
 
 	/**
-	 * Convert a Mesh Selection to an Index List
+	 * Convert a Mesh Selection to an Index Array
 	 * @param ConvertToType optional parameter specifying the type of Index List to convert to. If Set to Any, no conversion will be performed.
 	 */
 	UFUNCTION(BlueprintCallable, Category = "GeometryScript|MeshSelection", meta=(ScriptMethod))
diff --git a/Engine/Plugins/Experimental/GeometryScripting/Source/GeometryScriptingCore/Public/GeometryScript/MeshSimplifyFunctions.h b/Engine/Plugins/Experimental/GeometryScripting/Source/GeometryScriptingCore/Public/GeometryScript/MeshSimplifyFunctions.h
--- a/Engine/Plugins/Experimental/GeometryScripting/Source/GeometryScriptingCore/Public/GeometryScript/MeshSimplifyFunctions.h
+++ b/Engine/Plugins/Experimental/GeometryScripting/Source/GeometryScriptingCore/Public/GeometryScript/MeshSimplifyFunctions.h
@@ -85,13 +85,21 @@ class GEOMETRYSCRIPTINGCORE_API UGeometryScriptLibrary_MeshSimplifyFunctions : p
 	GENERATED_BODY()
 public:
 
+	/**
+	* Simplifies planar areas of the mesh that have more triangles than necessary. Note that it does not change the 3D shape of the mesh.
+	* Planar regions are identified by comparison of face normals using a Angle Threshold in the Options.
+	*/
 	UFUNCTION(BlueprintCallable, Category = "GeometryScript|Simplification", meta=(ScriptMethod))
 	static UPARAM(DisplayName = "Target Mesh") UDynamicMesh* 
 	ApplySimplifyToPlanar(  
 		UDynamicMesh* TargetMesh, 
 		FGeometryScriptPlanarSimplifyOptions Options,
 		UGeometryScriptDebug* Debug = nullptr);
-
+	
+	/**
+	* Simplifies the mesh down to the PolyGroup Topology. For example, the high-level faces of the mesh PolyGroups. 
+	* Another example would be on a default Box-Sphere where simplifying to PolyGroup topology produces a box.
+	*/
 	UFUNCTION(BlueprintCallable, Category = "GeometryScript|Simplification", meta=(ScriptMethod, DisplayName = "ApplySimplifyToPolyGroupTopology"))
 	static UPARAM(DisplayName = "Target Mesh") UDynamicMesh* 
 	ApplySimplifyToPolygroupTopology(  
@@ -100,6 +108,9 @@ public:
 		FGeometryScriptGroupLayer GroupLayer,
 		UGeometryScriptDebug* Debug = nullptr);
 
+	/**
+	* Simplifies the mesh until a target triangle count is reached. Behavior can be additionally controlled with the Options. 
+	*/
 	UFUNCTION(BlueprintCallable, Category = "GeometryScript|Simplification", meta=(ScriptMethod))
 	static UPARAM(DisplayName = "Target Mesh") UDynamicMesh* 
 	ApplySimplifyToTriangleCount(  
@@ -108,6 +119,9 @@ public:
 		FGeometryScriptSimplifyMeshOptions Options,
 		UGeometryScriptDebug* Debug = nullptr);
 
+	/**
+	* Simplifies the mesh until a target vertex count is reached. Behavior can be additionally controlled with the Options. 
+	*/
 	UFUNCTION(BlueprintCallable, Category = "GeometryScript|Simplification", meta=(ScriptMethod))
 	static UPARAM(DisplayName = "Target Mesh") UDynamicMesh* 
 	ApplySimplifyToVertexCount(  
@@ -116,6 +130,10 @@ public:
 		FGeometryScriptSimplifyMeshOptions Options,
 		UGeometryScriptDebug* Debug = nullptr);
 
+	/**
+	* Simplifies the mesh to a target geometric tolerance. Stops when any further simplification would result in a deviation from the input mesh larger than the tolerance.
+	* Behavior can be additionally controlled with the Options. 
+	*/
 	UFUNCTION(BlueprintCallable, Category = "GeometryScript|Simplification", meta=(ScriptMethod))
 	static UPARAM(DisplayName = "Target Mesh") UDynamicMesh* 
 	ApplySimplifyToTolerance(  
diff --git a/Engine/Plugins/Experimental/GeometryScripting/Source/GeometryScriptingCore/Public/GeometryScript/MeshSpatialFunctions.h b/Engine/Plugins/Experimental/GeometryScripting/Source/GeometryScriptingCore/Public/GeometryScript/MeshSpatialFunctions.h
--- a/Engine/Plugins/Experimental/GeometryScripting/Source/GeometryScriptingCore/Public/GeometryScript/MeshSpatialFunctions.h
+++ b/Engine/Plugins/Experimental/GeometryScripting/Source/GeometryScriptingCore/Public/GeometryScript/MeshSpatialFunctions.h
@@ -18,6 +18,7 @@ public:
 	UPROPERTY(EditAnywhere, BlueprintReadWrite, Category = Options)
 	float MaxDistance = 0;
 
+	/** When true, allows the provided BHV to be used even if the mesh has been updated */
 	UPROPERTY(EditAnywhere, BlueprintReadWrite, Category = Options)
 	bool bAllowUnsafeModifiedQueries = false;
 
@@ -57,16 +58,26 @@ class GEOMETRYSCRIPTINGCORE_API UGeometryScriptLibrary_MeshSpatial : public UBlu
 {
 	GENERATED_BODY()
 public:
+
+	/**
+	* Reset the Bounding Volume Hierarchy (BVH) by clearing all the internal data.
+	*/
 	UFUNCTION(BlueprintCallable, Category = "GeometryScript|Spatial", meta = (ScriptMethod))
 	static void ResetBVH(UPARAM(ref) FGeometryScriptDynamicMeshBVH& ResetBVH);
-
+	
+	/**
+	* Builds a Bounding Volume Hierarchy (BVH) object for a mesh that can be used with multiple spatial queries.
+	*/
 	UFUNCTION(BlueprintCallable, Category = "GeometryScript|Spatial", meta=(ScriptMethod))
 	static UPARAM(DisplayName = "Target Mesh") UDynamicMesh* 
 	BuildBVHForMesh(
 		UDynamicMesh* TargetMesh,
 		FGeometryScriptDynamicMeshBVH& OutputBVH,
 		UGeometryScriptDebug* Debug = nullptr );
-
+	
+	/**
+	* Checks if the provided Bounding Volume Hierarchy (BVH) can still be used with the Mesh — it generally returns false if the mesh has been changed.
+	*/
 	UFUNCTION(BlueprintCallable, Category = "GeometryScript|Spatial", meta=(ScriptMethod))
 	static UPARAM(DisplayName = "Target Mesh") UDynamicMesh* 
 	IsBVHValidForMesh(
@@ -75,6 +86,9 @@ public:
 		bool& bIsValid,
 		UGeometryScriptDebug* Debug = nullptr );
 
+	/**
+	* Rebuilds the Bounding Volume Hierarchy (BVH) for the mesh in-place, which can reduce memory allocations, compared to building a new BVH.
+	*/
 	UFUNCTION(BlueprintCallable, Category = "GeometryScript|Spatial", meta=(ScriptMethod))
 	static UPARAM(DisplayName = "Target Mesh") UDynamicMesh* 
 	RebuildBVHForMesh(
@@ -83,6 +97,14 @@ public:
 		bool bOnlyIfInvalid = true,
 		UGeometryScriptDebug* Debug = nullptr );
 
+	/**
+	* Finds the nearest point (Nearest Result) on the Target Mesh to a given 3D point (Query Point) by using the Query BVH.
+	* @param QueryBVH a BVH associated with the Target Mesh
+	* @param QueryPoint a 3D location relative to the local space of the mesh
+	* @param NearestResult on return, holds the nearest point on the mesh to the QueryPoint
+	* @param Outcome will be either Found or Not Found depending on the success of the query.
+	* Note NearestResult.bValid will be false if the query failed. 
+	*/
 	UFUNCTION(BlueprintCallable, Category = "GeometryScript|Spatial", meta=(ScriptMethod, ExpandEnumAsExecs = "Outcome"))
 	static UPARAM(DisplayName = "Target Mesh") UDynamicMesh* 
 	FindNearestPointOnMesh(
@@ -94,7 +116,11 @@ public:
 		EGeometryScriptSearchOutcomePins& Outcome,
 		UGeometryScriptDebug* Debug = nullptr );
 
-
+	/**
+	* Finds the nearest intersection of a 3D ray with the mesh by using the Query BVH.  
+	* Note, depending on the Ray Origin and Ray Direction, there is the possibility that the ray might not intersect with the Target Mesh.  
+	* Should the ray miss, the HitResult.bHit will be false and the Outcome  will be Not Found.
+	*/
 	UFUNCTION(BlueprintCallable, Category = "GeometryScript|Spatial", meta=(ScriptMethod, ExpandEnumAsExecs = "Outcome"))
 	static UPARAM(DisplayName = "Target Mesh") UDynamicMesh* 
 	FindNearestRayIntersectionWithMesh(
@@ -107,7 +133,12 @@ public:
 		EGeometryScriptSearchOutcomePins& Outcome,
 		UGeometryScriptDebug* Debug = nullptr );
 
-
+	/**
+	* Tests if a point is inside the mesh using the Fast Winding Number query and data stored in the BVH.
+	* @param QueryBVH is an acceleration structure previously built with this mesh.
+	* @param QueryPoint the point in the mesh's 3D local space.
+	* @param Options control the fast winding number threshold 
+	*/
 	UFUNCTION(BlueprintCallable, Category = "GeometryScript|Spatial", meta=(ScriptMethod, ExpandEnumAsExecs = "Outcome"))
 	static UPARAM(DisplayName = "Target Mesh") UDynamicMesh* 
 	IsPointInsideMesh(
diff --git a/Engine/Plugins/Experimental/GeometryScripting/Source/GeometryScriptingCore/Public/GeometryScript/MeshTransformFunctions.h b/Engine/Plugins/Experimental/GeometryScripting/Source/GeometryScriptingCore/Public/GeometryScript/MeshTransformFunctions.h
--- a/Engine/Plugins/Experimental/GeometryScripting/Source/GeometryScriptingCore/Public/GeometryScript/MeshTransformFunctions.h
+++ b/Engine/Plugins/Experimental/GeometryScripting/Source/GeometryScriptingCore/Public/GeometryScript/MeshTransformFunctions.h
@@ -17,6 +17,9 @@ class GEOMETRYSCRIPTINGCORE_API UGeometryScriptLibrary_MeshTransformFunctions :
 	GENERATED_BODY()
 public:
 
+	/**
+	* Applies the provided Transform to the vertices of a Mesh.
+	*/
 	UFUNCTION(BlueprintCallable, Category = "GeometryScript|Transforms", meta=(ScriptMethod, AdvancedDisplay = "bFixOrientationForNegativeScale"))
 	static UPARAM(DisplayName = "Target Mesh") UDynamicMesh*
 	TransformMesh(
@@ -25,13 +28,19 @@ public:
 		bool bFixOrientationForNegativeScale = true,
 		UGeometryScriptDebug* Debug = nullptr);
 
+	/**
+	* Applies the provided Translation to the vertices of a Mesh.
+	*/
 	UFUNCTION(BlueprintCallable, Category = "GeometryScript|Transforms", meta=(ScriptMethod))
 	static UPARAM(DisplayName = "Target Mesh") UDynamicMesh* 
 	TranslateMesh(
 		UDynamicMesh* TargetMesh,
 		FVector Translation,
 		UGeometryScriptDebug* Debug = nullptr );
-
+	
+	/**
+	* Rotates the mesh relative to the specified Rotation Origin.
+	*/
 	UFUNCTION(BlueprintCallable, Category = "GeometryScript|Transforms", meta=(ScriptMethod))
 	static UPARAM(DisplayName = "Target Mesh") UDynamicMesh* 
 	RotateMesh(
@@ -39,7 +48,10 @@ public:
 		FRotator Rotation,
 		FVector RotationOrigin = FVector(0,0,0),
 		UGeometryScriptDebug* Debug = nullptr );
-
+	
+	/**
+	* Applies the provided Scale transformation relative to the Scale Origin to the vertices of a Mesh.
+	*/
 	UFUNCTION(BlueprintCallable, Category = "GeometryScript|Transforms", meta=(ScriptMethod, AdvancedDisplay = "bFixOrientationForNegativeScale"))
 	static UPARAM(DisplayName = "Target Mesh") UDynamicMesh* 
 	ScaleMesh(
@@ -51,7 +63,9 @@ public:
 
 
 
-
+	/**
+	* Applies the given Transform to the vertices identified by the Selection of the mesh.
+	*/
 	UFUNCTION(BlueprintCallable, Category = "GeometryScript|Transforms", meta=(ScriptMethod, AdvancedDisplay = "bFixOrientationForNegativeScale"))
 	static UPARAM(DisplayName = "Target Mesh") UDynamicMesh*
 	TransformMeshSelection(
@@ -60,6 +74,9 @@ public:
 		FTransform Transform,
 		UGeometryScriptDebug* Debug = nullptr);
 
+	/**
+	* Applies the given Translation to the vertices identified by the Selection of the mesh.
+	*/ 
 	UFUNCTION(BlueprintCallable, Category = "GeometryScript|Transforms", meta=(ScriptMethod))
 	static UPARAM(DisplayName = "Target Mesh") UDynamicMesh* 
 	TranslateMeshSelection(
@@ -68,6 +85,9 @@ public:
 		FVector Translation,
 		UGeometryScriptDebug* Debug = nullptr );
 
+	/**
+	* Rotates the selected part of the mesh relative to the specified Rotation Origin.
+	*/
 	UFUNCTION(BlueprintCallable, Category = "GeometryScript|Transforms", meta=(ScriptMethod))
 	static UPARAM(DisplayName = "Target Mesh") UDynamicMesh* 
 	RotateMeshSelection(
@@ -77,6 +97,9 @@ public:
 		FVector RotationOrigin = FVector(0,0,0),
 		UGeometryScriptDebug* Debug = nullptr );
 
+	/**
+	* Applies the given Scale transform relative to the Scale Origin to the selection part of the mesh.
+	*/
 	UFUNCTION(BlueprintCallable, Category = "GeometryScript|Transforms", meta=(ScriptMethod, AdvancedDisplay = "bFixOrientationForNegativeScale"))
 	static UPARAM(DisplayName = "Target Mesh") UDynamicMesh* 
 	ScaleMeshSelection(
diff --git a/Engine/Plugins/Experimental/GeometryScripting/Source/GeometryScriptingCore/Public/GeometryScript/MeshUVFunctions.h b/Engine/Plugins/Experimental/GeometryScripting/Source/GeometryScriptingCore/Public/GeometryScript/MeshUVFunctions.h
--- a/Engine/Plugins/Experimental/GeometryScripting/Source/GeometryScriptingCore/Public/GeometryScript/MeshUVFunctions.h
+++ b/Engine/Plugins/Experimental/GeometryScripting/Source/GeometryScriptingCore/Public/GeometryScript/MeshUVFunctions.h
@@ -146,6 +146,9 @@ class GEOMETRYSCRIPTINGCORE_API UGeometryScriptLibrary_MeshUVFunctions : public
 	GENERATED_BODY()
 public:
 
+	/**
+	* Set the number of UV Channels on the Target Mesh. If not already enabled, this will enable the mesh attributes.
+	*/
 	UFUNCTION(BlueprintCallable, Category = "GeometryScript|UVs", meta=(ScriptMethod), DisplayName = "Set Num UV Channels")
 	static UPARAM(DisplayName = "Target Mesh") UDynamicMesh* 
 	SetNumUVSets( 
@@ -153,6 +156,9 @@ public:
 		UPARAM(DisplayName = "Num UV Channels") int NumUVSets,
 		UGeometryScriptDebug* Debug = nullptr );
 
+	/**
+	* Copy the data in one UV Channel to another UV Channel on the same Target Mesh.
+	*/
 	UFUNCTION(BlueprintCallable, Category = "GeometryScript|UVs", meta=(ScriptMethod), DisplayName = "Copy UV Channel")
 	static UPARAM(DisplayName = "Target Mesh") UDynamicMesh* 
 	CopyUVSet( 
@@ -161,7 +167,9 @@ public:
 		UPARAM(DisplayName = "To UV Channel")   int ToUVSet,
 		UGeometryScriptDebug* Debug = nullptr );
 
-
+	/**
+	* Sets the UVs of a mesh triangle in the given UV Channel.
+	*/
 	UFUNCTION(BlueprintCallable, Category = "GeometryScript|UVs", meta=(ScriptMethod))
 	static UPARAM(DisplayName = "Target Mesh") UDynamicMesh* 
 	SetMeshTriangleUVs( 
@@ -172,7 +180,10 @@ public:
 		bool& bIsValidTriangle, 
 		bool bDeferChangeNotifications = false );
 
-
+	/**
+	* Update all selected UV values in the specified UV Channel by adding the Translation value to each.
+	* If the provided Selection is empty, the Translation is applied to the entire UV Channel.
+	*/
 	UFUNCTION(BlueprintCallable, Category = "GeometryScript|UVs", meta=(ScriptMethod))
 	static UPARAM(DisplayName = "Target Mesh") UDynamicMesh* 
 	TranslateMeshUVs( 
@@ -181,7 +192,11 @@ public:
 		FVector2D Translation,
 		FGeometryScriptMeshSelection Selection,
 		UGeometryScriptDebug* Debug = nullptr );
-
+	
+	/**
+	* Update all selected UV values in the specified UV Channel by Scale, mathematically the new value is given by (UV - ScaleOrigin) * Scale + ScaleOrigin
+	* If the provided Selection is empty, the update is applied to the entire UV Channel.
+	*/
 	UFUNCTION(BlueprintCallable, Category = "GeometryScript|UVs", meta=(ScriptMethod))
 	static UPARAM(DisplayName = "Target Mesh") UDynamicMesh* 
 	ScaleMeshUVs( 
@@ -192,6 +207,10 @@ public:
 		FGeometryScriptMeshSelection Selection,
 		UGeometryScriptDebug* Debug = nullptr );
 
+	/**
+	* Update all the selected UV values in the specified UV Channel by a rotation of Rotation Angle (in degrees) relative to the Rotation Origin.
+	* If the provided Selection is empty, the update is applied to the entire UV Channel.
+	*/
 	UFUNCTION(BlueprintCallable, Category = "GeometryScript|UVs", meta=(ScriptMethod))
 	static UPARAM(DisplayName = "Target Mesh") UDynamicMesh* 
 	RotateMeshUVs( 
@@ -216,6 +235,9 @@ public:
 		UGeometryScriptDebug* Debug = nullptr );
 
 
+	/**
+	* Using Box Projection, update the UVs in the UV Channel for an entire mesh or a subset defined by a non-empty Selection.
+	*/
 	UFUNCTION(BlueprintCallable, Category = "GeometryScript|UVs", meta=(ScriptMethod))
 	static UPARAM(DisplayName = "Target Mesh") UDynamicMesh* 
 	SetMeshUVsFromBoxProjection( 
@@ -226,7 +248,9 @@ public:
 		int MinIslandTriCount = 2,
 		UGeometryScriptDebug* Debug = nullptr );
 
-
+	/**
+	* Using Cylinder Projection, update the UVs in the UV Channel for an entire mesh or a subset defined by a non-empty Selection.
+	*/
 	UFUNCTION(BlueprintCallable, Category = "GeometryScript|UVs", meta=(ScriptMethod))
 	static UPARAM(DisplayName = "Target Mesh") UDynamicMesh* 
 	SetMeshUVsFromCylinderProjection( 
@@ -238,7 +262,10 @@ public:
 		UGeometryScriptDebug* Debug = nullptr );
 
 
-
+	/**
+	* Recomputes UVs in the UV Channel for a Mesh based on different types of well-defined UV islands, such as existing UV islands, PolyGroups, 
+	* or a subset of the mesh based on a non-empty Selection.
+	*/
 	UFUNCTION(BlueprintCallable, Category = "GeometryScript|UVs", meta=(ScriptMethod))
 	static UPARAM(DisplayName = "Target Mesh") UDynamicMesh* 
 	RecomputeMeshUVs( 
@@ -248,7 +275,9 @@ public:
 		FGeometryScriptMeshSelection Selection,
 		UGeometryScriptDebug* Debug = nullptr );
 
-
+	/**
+	* Packs the existing UV islands in the specified UV Channel into standard UV space based on the Repack Options.
+	*/
 	UFUNCTION(BlueprintCallable, Category = "GeometryScript|UVs", meta=(ScriptMethod))
 	static UPARAM(DisplayName = "Target Mesh") UDynamicMesh* 
 	RepackMeshUVs( 
@@ -257,6 +286,9 @@ public:
 		FGeometryScriptRepackUVsOptions RepackOptions,
 		UGeometryScriptDebug* Debug = nullptr );
 
+	/**
+	* Computes new UVs for the specified UV Channel using PatchBuilder method in the Options, and optionally packs.
+	*/
 	UFUNCTION(BlueprintCallable, Category = "GeometryScript|UVs", meta=(ScriptMethod))
 	static UPARAM(DisplayName = "Target Mesh") UDynamicMesh* 
 	AutoGeneratePatchBuilderMeshUVs( 
@@ -265,6 +297,9 @@ public:
 		FGeometryScriptPatchBuilderOptions Options,
 		UGeometryScriptDebug* Debug = nullptr );
 
+	/**
+	* Computes new UVs for the specified UV Channel using XAtlas, and optionally packs.
+	*/
 	UFUNCTION(BlueprintCallable, Category = "GeometryScript|UVs", meta=(ScriptMethod))
 	static UPARAM(DisplayName = "Target Mesh") UDynamicMesh* 
 	AutoGenerateXAtlasMeshUVs( 
diff --git a/Engine/Plugins/Experimental/GeometryScripting/Source/GeometryScriptingCore/Public/GeometryScript/MeshVoxelFunctions.h b/Engine/Plugins/Experimental/GeometryScripting/Source/GeometryScriptingCore/Public/GeometryScript/MeshVoxelFunctions.h
--- a/Engine/Plugins/Experimental/GeometryScripting/Source/GeometryScriptingCore/Public/GeometryScript/MeshVoxelFunctions.h
+++ b/Engine/Plugins/Experimental/GeometryScripting/Source/GeometryScriptingCore/Public/GeometryScript/MeshVoxelFunctions.h
@@ -119,13 +119,19 @@ class GEOMETRYSCRIPTINGCORE_API UGeometryScriptLibrary_MeshVoxelFunctions : publ
 	GENERATED_BODY()
 public:
 
+	/**
+	* Replaces the mesh with a voxelized-and-meshed approximation (VoxWrap operation).
+	*/
 	UFUNCTION(BlueprintCallable, Category = "GeometryScript|Voxel", meta=(ScriptMethod))
 	static UPARAM(DisplayName = "Target Mesh") UDynamicMesh*
 	ApplyMeshSolidify(
 		UDynamicMesh* TargetMesh,
 		FGeometryScriptSolidifyOptions Options,
 		UGeometryScriptDebug* Debug = nullptr);
-
+	
+	/**
+	* Replaces the mesh with an SDF-based offset mesh approximation.
+	*/
 	UFUNCTION(BlueprintCallable, Category = "GeometryScript|Voxel", meta=(ScriptMethod))
 	static UPARAM(DisplayName = "Target Mesh") UDynamicMesh*
 	ApplyMeshMorphology(
diff --git a/Engine/Plugins/Experimental/GeometryScripting/Source/GeometryScriptingCore/Public/GeometryScript/PolyPathFunctions.h b/Engine/Plugins/Experimental/GeometryScripting/Source/GeometryScriptingCore/Public/GeometryScript/PolyPathFunctions.h
--- a/Engine/Plugins/Experimental/GeometryScripting/Source/GeometryScriptingCore/Public/GeometryScript/PolyPathFunctions.h
+++ b/Engine/Plugins/Experimental/GeometryScripting/Source/GeometryScriptingCore/Public/GeometryScript/PolyPathFunctions.h
@@ -42,23 +42,44 @@ class GEOMETRYSCRIPTINGCORE_API UGeometryScriptLibrary_PolyPathFunctions : publi
 {
 	GENERATED_BODY()
 public:
-
+	
+	/**
+	* Returns the number of vertices in the the PolyPath.
+	*/
 	UFUNCTION(BlueprintPure, Category = "GeometryScript|PolyPath", meta=(ScriptMethod))
 	static int GetPolyPathNumVertices(FGeometryScriptPolyPath PolyPath);
 
+	/**
+	* Returns the index of the last vertex in the PolyPath.  
+	*/
 	UFUNCTION(BlueprintPure, Category = "GeometryScript|PolyPath", meta=(ScriptMethod))
 	static int GetPolyPathLastIndex(FGeometryScriptPolyPath PolyPath);
 
+	/**
+	* Returns the 3D position of the requested vertex in the PolyPath.
+	* If the Index does not correspond to a vertex in the PolyPath, a Zero Vector (0,0,0) will be returned. 
+	* @param Index specifies a vertex in the PolyPath.
+	* @param bIsValidIndex will be false on return if the Index is not included in the PolyPath. 
+	*/
 	UFUNCTION(BlueprintPure, Category = "GeometryScript|PolyPath", meta=(ScriptMethod))
 	static FVector GetPolyPathVertex(FGeometryScriptPolyPath PolyPath, int Index, bool& bIsValidIndex);
 
+	/**
+	* Returns the local tangent vector of the PolyPath at the specified vertex index.
+	* If the Index does not correspond to a vertex in the PolyPath, a Zero Vector (0,0,0) will be returned. 
+	* @param Index specifies a vertex in the PolyPath
+	* @param bIsValidIndex will be false on return if the Index is not included in the PolyPath
+	*/
 	UFUNCTION(BlueprintPure, Category = "GeometryScript|PolyPath", meta = (ScriptMethod))
 	static FVector GetPolyPathTangent(FGeometryScriptPolyPath PolyPath, int Index, bool& bIsValidIndex);
 
+	/**
+	* Returns the length of the PolyPath.
+	*/
 	UFUNCTION(BlueprintCallable, Category = "GeometryScript|PolyPath", meta = (ScriptMethod))
 	static double GetPolyPathArcLength(FGeometryScriptPolyPath PolyPath);
 
-	/** Find the index of the vertex closest to a given point.  Returns -1 if path has no vertices. */
+	/** Find the index of the vertex closest to a given point.  Returns -1 if the PolyPath has no vertices. */
 	UFUNCTION(BlueprintCallable, Category = "GeometryScript|PolyPath", meta = (ScriptMethod))
 	static int32 GetNearestVertexIndex(FGeometryScriptPolyPath PolyPath, FVector Point);
 
@@ -117,27 +138,51 @@ public:
 		bool bIncludeScale = true);
 
 
+	/**
+	* Populates an array of 3D vectors with the PolyPath vertex locations.
+	*/
 	UFUNCTION(BlueprintCallable, Category = "GeometryScript|PolyPath", meta=(ScriptMethod))
 	static void ConvertPolyPathToArray(FGeometryScriptPolyPath PolyPath, TArray<FVector>& VertexArray);
 
+	/**
+	* Creates a PolyPath from an array of 3D position vectors.
+	*/
 	UFUNCTION(BlueprintCallable, Category = "GeometryScript|PolyPath")
 	static void ConvertArrayToPolyPath(const TArray<FVector>& VertexArray, FGeometryScriptPolyPath& PolyPath);
 
+	/**
+	* Creates an array of 2D Vectors with the PolyPath vertex locations projected onto the XY plane.
+	*/
 	UFUNCTION(BlueprintCallable, Category = "GeometryScript|PolyPath", meta = (ScriptMethod))
 	static void ConvertPolyPathToArrayOfVector2D(FGeometryScriptPolyPath PolyPath, TArray<FVector2D>& VertexArray);
 
+	/**
+	* Creates a PolyPath from an array of 2D position vectors. The Z-coordinate of the corresponding PolyPath vertices will be zero.
+	*/
 	UFUNCTION(BlueprintCallable, Category = "GeometryScript|PolyPath")
 	static void ConvertArrayOfVector2DToPolyPath(const TArray<FVector2D>& VertexArray, FGeometryScriptPolyPath& PolyPath);
 
+	/**
+	* Returns an array of 3D vectors with the PolyPath vertex locations.
+	*/
 	UFUNCTION(BlueprintPure, meta = (DisplayName = "PolyPath To Array Of Vector", CompactNodeTitle = "->", BlueprintAutocast), Category = "GeometryScript|PolyPath")
 	static TArray<FVector> Conv_GeometryScriptPolyPathToArray(FGeometryScriptPolyPath PolyPath);
 
+	/**
+	* Returns an array of 2D Vectors with the PolyPath vertex locations projected onto the XY plane.
+	*/
 	UFUNCTION(BlueprintPure, meta = (DisplayName = "PolyPath To Array Of Vector2D", CompactNodeTitle = "->", BlueprintAutocast), Category = "GeometryScript|PolyPath")
 	static TArray<FVector2D> Conv_GeometryScriptPolyPathToArrayOfVector2D(FGeometryScriptPolyPath PolyPath);
 
+	/** 
+	* Returns a PolyPath created from an array of 3D position vectors.
+	*/
 	UFUNCTION(BlueprintPure, meta = (DisplayName = "Array Of Vector To PolyPath", CompactNodeTitle = "->", BlueprintAutocast), Category = "GeometryScript|PolyPath")
 	static FGeometryScriptPolyPath Conv_ArrayToGeometryScriptPolyPath(const TArray<FVector>& PathVertices);
 
+	/**
+	* Returns a PolyPath created from an array of 2D position vectors. The Z-coordinate of the corresponding PolyPath vertices will be zero.
+	*/
 	UFUNCTION(BlueprintPure, meta = (DisplayName = "Array Of Vector2D To PolyPath", CompactNodeTitle = "->", BlueprintAutocast), Category = "GeometryScript|PolyPath")
 	static FGeometryScriptPolyPath Conv_ArrayOfVector2DToGeometryScriptPolyPath(const TArray<FVector2D>& PathVertices);
 };
diff --git a/Engine/Plugins/Experimental/GeometryScripting/Source/GeometryScriptingCore/Public/GeometryScript/VectorMathFunctions.h b/Engine/Plugins/Experimental/GeometryScripting/Source/GeometryScriptingCore/Public/GeometryScript/VectorMathFunctions.h
--- a/Engine/Plugins/Experimental/GeometryScripting/Source/GeometryScriptingCore/Public/GeometryScript/VectorMathFunctions.h
+++ b/Engine/Plugins/Experimental/GeometryScripting/Source/GeometryScriptingCore/Public/GeometryScript/VectorMathFunctions.h
@@ -137,20 +137,35 @@ public:
 	static void 
 	ScalarBlendInPlace(FGeometryScriptScalarList ScalarListA, UPARAM(ref) FGeometryScriptScalarList& ScalarListB, double ConstantA = 1.0, double ConstantB = 1.0);
 
-
+	/**
+	 * Returns a Scalar List constructed with each element is the product (ConstantMultiplier * A * B) 
+	 * where A and B are the corresponding elements from ScalarListA and ScalarListB accordingly. 
+	 * If ScalarListA and ScalarListB have different lengths, no operation will be performed and an empty Scalar List will be returned.
+	 */
 	UFUNCTION(BlueprintCallable, Category = "GeometryScript|Math|VectorMath", meta=(ScriptMethod))
 	static FGeometryScriptScalarList 
 	ScalarMultiply(FGeometryScriptScalarList ScalarListA, FGeometryScriptScalarList ScalarListB, double ConstantMultiplier = 1.0);
 
+	/**
+	 * Compute (ConstantMultiplier * A * B)  where A and B are the corresponding elements from ScalarListA and ScalarListB accordingly, and store the 
+	 * result in ScalarListA. 
+	 * If ScalarListA and ScalarListB have different lengths, the computation will be skipped.
+	 */
 	UFUNCTION(BlueprintCallable, Category = "GeometryScript|Math|VectorMath", meta=(ScriptMethod))
 	static void 
 	ScalarMultiplyInPlace(FGeometryScriptScalarList ScalarListA, UPARAM(ref) FGeometryScriptScalarList& ScalarListB, double ConstantMultiplier = 1.0);
 
 
+	/**
+	 * Returns a Scalar List of the same length as the input Scalar List, with the elements computed as (Constant * A) where A is the corresponding element in the input List.
+	 */
 	UFUNCTION(BlueprintCallable, Category = "GeometryScript|Math|VectorMath")
 	static FGeometryScriptScalarList 
 	ConstantScalarMultiply(double Constant, FGeometryScriptScalarList ScalarList);
 
+	/**
+	 * Compute (Constant * A) for each element A is the Scalar List, and the result is stored in the original Scalar List. 
+	 */
 	UFUNCTION(BlueprintCallable, Category = "GeometryScript|Math|VectorMath")
 	static void 
 	ConstantScalarMultiplyInPlace(double Constant, UPARAM(ref) FGeometryScriptScalarList& ScalarList);
