Expose Delaunay Triangulation to Geometry Script

diff --git a/Engine/Plugins/Experimental/GeometryScripting/Source/GeometryScriptingCore/Private/MeshPrimitiveFunctions.cpp b/Engine/Plugins/Experimental/GeometryScripting/Source/GeometryScriptingCore/Private/MeshPrimitiveFunctions.cpp
--- a/Engine/Plugins/Experimental/GeometryScripting/Source/GeometryScriptingCore/Private/MeshPrimitiveFunctions.cpp
+++ b/Engine/Plugins/Experimental/GeometryScripting/Source/GeometryScriptingCore/Private/MeshPrimitiveFunctions.cpp
@@ -29,60 +29,99 @@ using namespace UE::Geometry;
 #define LOCTEXT_NAMESPACE "UGeometryScriptLibrary_MeshPrimitiveFunctions"
 
 
-static void AppendPrimitive(
-	UDynamicMesh* TargetMesh, 
-	FMeshShapeGenerator* Generator, 
-	FTransform Transform, 
-	FGeometryScriptPrimitiveOptions PrimitiveOptions,
-	FVector3d PreTranslate = FVector3d::Zero())
+static void ApplyPrimitiveOptionsToMesh(FDynamicMesh3& Mesh, const FTransform& Transform, FGeometryScriptPrimitiveOptions PrimitiveOptions, FVector3d PreTranslate = FVector3d::Zero())
 {
-	auto ApplyOptionsToMesh = [&Transform, &PrimitiveOptions, PreTranslate](FDynamicMesh3& Mesh)
+	if (PreTranslate.SquaredLength() > 0)
 	{
-		if (PreTranslate.SquaredLength() > 0)
-		{
-			MeshTransforms::Translate(Mesh, PreTranslate);
-		}
+		MeshTransforms::Translate(Mesh, PreTranslate);
+	}
 
-		MeshTransforms::ApplyTransform(Mesh, (FTransformSRT3d)Transform, true);
-		if (PrimitiveOptions.PolygroupMode == EGeometryScriptPrimitivePolygroupMode::SingleGroup)
+	MeshTransforms::ApplyTransform(Mesh, (FTransformSRT3d)Transform, true);
+	if (PrimitiveOptions.PolygroupMode == EGeometryScriptPrimitivePolygroupMode::SingleGroup)
+	{
+		for (int32 tid : Mesh.TriangleIndicesItr())
 		{
-			for (int32 tid : Mesh.TriangleIndicesItr())
-			{
-				Mesh.SetTriangleGroup(tid, 0);
-			}
+			Mesh.SetTriangleGroup(tid, 0);
 		}
-		if (PrimitiveOptions.bFlipOrientation)
+	}
+	if (PrimitiveOptions.bFlipOrientation)
+	{
+		Mesh.ReverseOrientation(true);
+		if (Mesh.HasAttributes())
 		{
-			Mesh.ReverseOrientation(true);
-			if (Mesh.HasAttributes())
+			FDynamicMeshNormalOverlay* Normals = Mesh.Attributes()->PrimaryNormals();
+			for (int elemid : Normals->ElementIndicesItr())
 			{
-				FDynamicMeshNormalOverlay* Normals = Mesh.Attributes()->PrimaryNormals();
-				for (int elemid : Normals->ElementIndicesItr())
-				{
-					Normals->SetElement(elemid, -Normals->GetElement(elemid));
-				}
+				Normals->SetElement(elemid, -Normals->GetElement(elemid));
 			}
 		}
-	};
+	}
+}
 
+static void AppendPrimitive(
+	UDynamicMesh* TargetMesh, 
+	FMeshShapeGenerator* Generator, 
+	FTransform Transform, 
+	FGeometryScriptPrimitiveOptions PrimitiveOptions,
+	FVector3d PreTranslate = FVector3d::Zero())
+{
 	if (TargetMesh->IsEmpty())
 	{
 		TargetMesh->EditMesh([&](FDynamicMesh3& EditMesh)
 		{
 			EditMesh.Copy(Generator);
-			ApplyOptionsToMesh(EditMesh);
+			ApplyPrimitiveOptionsToMesh(EditMesh, Transform, PrimitiveOptions, PreTranslate);
 		}, EDynamicMeshChangeType::GeneralEdit, EDynamicMeshAttributeChangeFlags::Unknown, false);
 	}
 	else
 	{
 		FDynamicMesh3 TempMesh(Generator);
-		ApplyOptionsToMesh(TempMesh);
+		ApplyPrimitiveOptionsToMesh(TempMesh, Transform, PrimitiveOptions, PreTranslate);
 		TargetMesh->EditMesh([&](FDynamicMesh3& EditMesh)
 		{
 			FMeshIndexMappings TmpMappings;
 			FDynamicMeshEditor Editor(&EditMesh);
 			Editor.AppendMesh(&TempMesh, TmpMappings);
+		}, EDynamicMeshChangeType::GeneralEdit, EDynamicMeshAttributeChangeFlags::Unknown, false);
+	}
+}
 
+static void AppendPrimitiveWithVertexIndexMap(
+	UDynamicMesh* TargetMesh,
+	TArray<int32>& OutVertexIndexMap,
+	FMeshShapeGenerator* Generator,
+	FTransform Transform,
+	FGeometryScriptPrimitiveOptions PrimitiveOptions,
+	FVector3d PreTranslate = FVector3d::Zero())
+{
+	if (TargetMesh->IsEmpty())
+	{
+		TargetMesh->EditMesh([&](FDynamicMesh3& EditMesh)
+		{
+			EditMesh.Copy(Generator);
+			ApplyPrimitiveOptionsToMesh(EditMesh, Transform, PrimitiveOptions, PreTranslate);
+			// For a clean/empty mesh, the vertex map is just an identity map
+			OutVertexIndexMap.SetNum(EditMesh.VertexCount());
+			for (int32 Idx = 0; Idx < OutVertexIndexMap.Num(); ++Idx)
+			{
+				OutVertexIndexMap[Idx] = Idx;
+			}
+		}, EDynamicMeshChangeType::GeneralEdit, EDynamicMeshAttributeChangeFlags::Unknown, false);
+	}
+	else
+	{
+		FDynamicMesh3 TempMesh(Generator);
+		ApplyPrimitiveOptionsToMesh(TempMesh, Transform, PrimitiveOptions, PreTranslate);
+		TargetMesh->EditMesh([&](FDynamicMesh3& EditMesh)
+		{
+			FMeshIndexMappings TmpMappings;
+			FDynamicMeshEditor Editor(&EditMesh);
+			Editor.AppendMesh(&TempMesh, TmpMappings);
+			OutVertexIndexMap.Init(INDEX_NONE, Generator->Vertices.Num());
+			for (TPair<int32, int32> IndexToID : TmpMappings.GetVertexMap().GetForwardMap())
+			{
+				OutVertexIndexMap[IndexToID.Key] = IndexToID.Value;
+			}
 		}, EDynamicMeshChangeType::GeneralEdit, EDynamicMeshAttributeChangeFlags::Unknown, false);
 	}
 }
@@ -1088,7 +1127,7 @@ UDynamicMesh* UGeometryScriptLibrary_MeshPrimitiveFunctions::AppendVoronoiDiagra
 {
 	if (TargetMesh == nullptr)
 	{
-		UE::Geometry::AppendError(Debug, EGeometryScriptErrorType::InvalidInputs, LOCTEXT("PrimitiveFunctions_AppendVoronoiDiagram2D_NullTarget", "AppendVoronoiDiagram3d: TargetMesh is Null"));
+		UE::Geometry::AppendError(Debug, EGeometryScriptErrorType::InvalidInputs, LOCTEXT("PrimitiveFunctions_AppendVoronoiDiagram2D_NullTarget", "AppendVoronoiDiagram2D: TargetMesh is Null"));
 		return TargetMesh;
 	}
 
@@ -1103,7 +1142,7 @@ UDynamicMesh* UGeometryScriptLibrary_MeshPrimitiveFunctions::AppendVoronoiDiagra
 
 	if (!bTriSuccess)
 	{
-		UE::Geometry::AppendError(Debug, EGeometryScriptErrorType::InvalidInputs, LOCTEXT("PrimitiveFunctions_AppendVoronoiDiagram2D_GenFailed", "AppendVoronoiDiagram2D: Voronoi diagram generation failed"));
+		UE::Geometry::AppendError(Debug, EGeometryScriptErrorType::OperationFailed, LOCTEXT("PrimitiveFunctions_AppendVoronoiDiagram2D_GenFailed", "AppendVoronoiDiagram2D: Voronoi diagram generation failed"));
 		return TargetMesh;
 	}
 
@@ -1167,4 +1206,112 @@ UDynamicMesh* UGeometryScriptLibrary_MeshPrimitiveFunctions::AppendVoronoiDiagra
 
 
 
+UDynamicMesh* UGeometryScriptLibrary_MeshPrimitiveFunctions::AppendDelaunayTriangulation2D(
+	UDynamicMesh* TargetMesh,
+	FGeometryScriptPrimitiveOptions PrimitiveOptions,
+	FTransform Transform,
+	const TArray<FVector2D>& VertexPositions,
+	const TArray<FIntPoint>& ConstrainedEdges,
+	FGeometryScriptConstrainedDelaunayTriangulationOptions TriangulationOptions,
+	TArray<int32>& PositionsToVertexIDs,
+	bool& bHasDuplicateVertices,
+	UGeometryScriptDebug* Debug)
+{
+	bHasDuplicateVertices = false;
+	PositionsToVertexIDs.Reset();
+
+	if (TargetMesh == nullptr)
+	{
+		UE::Geometry::AppendError(Debug, EGeometryScriptErrorType::InvalidInputs, LOCTEXT("PrimitiveFunctions_AppendDelaunayTriangulation2D_NullTarget", "AppendDelaunayTriangulation2D: TargetMesh is Null"));
+		return TargetMesh;
+	}
+
+	if (VertexPositions.Num() < 3)
+	{
+		UE::Geometry::AppendError(Debug, EGeometryScriptErrorType::InvalidInputs, LOCTEXT("PrimitiveFunctions_AppendDelaunayTriangulation2D_TooFewVertices", "AppendDelaunayTriangulation2D: VertexPositions array requires at least 3 positions"));
+		return TargetMesh;
+	}
+
+	UE::Geometry::FDelaunay2 Delaunay;
+	TArray<FIndex2i> ConvertedEdges;
+	if (!ConstrainedEdges.IsEmpty())
+	{
+		ConvertedEdges.Reserve(ConstrainedEdges.Num());
+		for (FIntPoint Pt : ConstrainedEdges)
+		{
+			if (VertexPositions.IsValidIndex(Pt.X) && VertexPositions.IsValidIndex(Pt.Y))
+			{
+				ConvertedEdges.Emplace(Pt.X, Pt.Y);
+			}
+			else if (TriangulationOptions.bValidateEdgesInResult)
+			{
+				UE::Geometry::AppendError(Debug, EGeometryScriptErrorType::InvalidInputs, LOCTEXT("PrimitiveFunctions_AppendDelaunayTriangulation2D_InvalidEdges", "AppendDelaunayTriangulation2D: Constrained Edges must only contain valid indices into Vertex Positions"));
+				return TargetMesh;
+			}
+		}
+		Delaunay.bValidateEdges = TriangulationOptions.bValidateEdgesInResult;
+		Delaunay.bAutomaticallyFixEdgesToDuplicateVertices = true;
+	}
+	else
+	{
+		Delaunay.bAutomaticallyFixEdgesToDuplicateVertices = TriangulationOptions.bRemoveDuplicateVertices;
+	}
+	bool bTriSuccess = Delaunay.Triangulate(VertexPositions, ConvertedEdges);
+
+	if (!bTriSuccess)
+	{
+		UE::Geometry::AppendError(Debug, EGeometryScriptErrorType::OperationFailed, LOCTEXT("PrimitiveFunctions_AppendDelaunayTriangulation2D_GenFailed", "AppendDelaunayTriangulation2D: Delaunay triangulation failed"));
+		return TargetMesh;
+	}
+	
+	bHasDuplicateVertices = Delaunay.HasDuplicates();
+
+	FFlatTriangulationMeshGenerator TriangulationMeshGen;
+	TriangulationMeshGen.bReverseOrientation = true; // because FDelaunay2 generates triangles with reversed orientation vs what we want
+	TriangulationMeshGen.Vertices2D = VertexPositions;
+	if (ConstrainedEdges.IsEmpty() || TriangulationOptions.ConstrainedEdgesFillMode == EGeometryScriptPolygonFillMode::All)
+	{
+		TriangulationMeshGen.Triangles2D = Delaunay.GetTriangles();
+	}
+	else
+	{
+		// Convert FillMode enum to the GeometryCore library equivalent.  Note the "All" case is handled above.
+		EGeometryScriptPolygonFillMode GeometryScriptFillMode = TriangulationOptions.ConstrainedEdgesFillMode;
+		FDelaunay2::EFillMode GeometryCoreFillMode =
+			GeometryScriptFillMode == EGeometryScriptPolygonFillMode::Solid ? FDelaunay2::EFillMode::Solid :
+			GeometryScriptFillMode == EGeometryScriptPolygonFillMode::PositiveWinding ? FDelaunay2::EFillMode::PositiveWinding :
+			GeometryScriptFillMode == EGeometryScriptPolygonFillMode::NonZeroWinding ? FDelaunay2::EFillMode::NonZeroWinding :
+			GeometryScriptFillMode == EGeometryScriptPolygonFillMode::NegativeWinding ? FDelaunay2::EFillMode::NegativeWinding :
+			FDelaunay2::EFillMode::OddWinding;
+		TriangulationMeshGen.Triangles2D = Delaunay.GetFilledTriangles(ConvertedEdges, GeometryCoreFillMode);
+	}
+
+	if (TriangulationMeshGen.Vertices2D.Num() > 2 && TriangulationMeshGen.Triangles2D.Num() > 0)
+	{
+		AppendPrimitiveWithVertexIndexMap(TargetMesh, PositionsToVertexIDs, &TriangulationMeshGen.Generate(), Transform, PrimitiveOptions);
+		
+		// If requested, remove duplicate vertices from the final mesh
+		if (TriangulationOptions.bRemoveDuplicateVertices && bHasDuplicateVertices)
+		{
+			TargetMesh->EditMesh([&](FDynamicMesh3& EditMesh)
+			{
+				for (int32 Idx = 0; Idx < PositionsToVertexIDs.Num(); ++Idx)
+				{
+					int32 RemapIdx = Delaunay.RemapIfDuplicate(Idx);
+					if (Idx != RemapIdx)
+					{
+						int32 DupID = PositionsToVertexIDs[Idx];
+						checkSlow(!EditMesh.IsReferencedVertex(DupID));
+						EditMesh.RemoveVertex(DupID);
+						PositionsToVertexIDs[Idx] = INDEX_NONE;
+					}
+				}
+			});
+		}
+	}
+
+	return TargetMesh;
+}
+
+
 #undef LOCTEXT_NAMESPACE
\ No newline at end of file
diff --git a/Engine/Plugins/Experimental/GeometryScripting/Source/GeometryScriptingCore/Public/GeometryScript/MeshPrimitiveFunctions.h b/Engine/Plugins/Experimental/GeometryScripting/Source/GeometryScriptingCore/Public/GeometryScript/MeshPrimitiveFunctions.h
--- a/Engine/Plugins/Experimental/GeometryScripting/Source/GeometryScriptingCore/Public/GeometryScript/MeshPrimitiveFunctions.h
+++ b/Engine/Plugins/Experimental/GeometryScripting/Source/GeometryScriptingCore/Public/GeometryScript/MeshPrimitiveFunctions.h
@@ -30,6 +30,23 @@ enum class EGeometryScriptPrimitiveUVMode : uint8
 	ScaleToFill = 1
 };
 
+UENUM(BlueprintType)
+enum class EGeometryScriptPolygonFillMode : uint8
+{
+	// Keep all triangles, regardless of whether they were enclosed by constrained edges
+	All = 0,
+	// Fill everything inside the outer boundaries of constrained edges, ignoring edge orientation and any internal holes
+	Solid = 1,
+	// Fill where the 'winding number' is positive
+	PositiveWinding = 2,
+	// Fill where the 'winding number' is not zero
+	NonZeroWinding = 3,
+	// Fill where the 'winding number' is negative
+	NegativeWinding = 4,
+	// Fill where the 'winding number' is an odd number
+	OddWinding = 5
+};
+
 
 USTRUCT(BlueprintType)
 struct GEOMETRYSCRIPTINGCORE_API FGeometryScriptPrimitiveOptions
@@ -95,6 +112,29 @@ public:
 };
 
 
+USTRUCT(BlueprintType)
+struct GEOMETRYSCRIPTINGCORE_API FGeometryScriptConstrainedDelaunayTriangulationOptions
+{
+	GENERATED_BODY()
+public:
+	/** How to decide which parts of the shape defined by constrained edges should be filled with triangles */
+	UPROPERTY(EditAnywhere, BlueprintReadWrite, Category = Options)
+	EGeometryScriptPolygonFillMode ConstrainedEdgesFillMode = EGeometryScriptPolygonFillMode::All;
+
+	/**
+	 * Whether the triangulation should be considered a failure if it doesn't include the requested Constrained Edges.
+	 * (Edges may be missing e.g. due to intersecting edges in the input.) 
+	 */
+	UPROPERTY(EditAnywhere, BlueprintReadWrite, Category = Options)
+	bool bValidateEdgesInResult = true;
+
+	/** Whether to remove duplicate vertices from the output.  If false, duplicate vertices will not be used in any triangles, but will remain in the output mesh. */
+	UPROPERTY(EditAnywhere, BlueprintReadWrite, Category = Options)
+	bool bRemoveDuplicateVertices = false;
+
+};
+
+
 UCLASS(meta = (ScriptName = "GeometryScript_Primitives"))
 class GEOMETRYSCRIPTINGCORE_API UGeometryScriptLibrary_MeshPrimitiveFunctions : public UBlueprintFunctionLibrary
 {
@@ -468,6 +508,25 @@ public:
 		FGeometryScriptVoronoiOptions VoronoiOptions,
 		UGeometryScriptDebug* Debug = nullptr);
 
+	/**
+	* Generates a Delaunay Triangulation of the provided vertices, and appends it to the Target Mesh.
+	* If optional Constrained Edges are provided, will generate a Constrained Delaunay Triangulation which connects the specified vertices with edges.
+	* On success, all vertices are always appended to the output mesh, though duplicate vertices will not be used in any triangles and may optionally be removed.
+	* Use PositionsToVertexIDs to map indices in the input VertexPositions array to vertex IDs in the Dynamic Mesh.
+	*/
+	UFUNCTION(BlueprintCallable, Category = "GeometryScript|Primitives", meta = (ScriptMethod))
+	static UPARAM(DisplayName = "Target Mesh") UDynamicMesh*
+	AppendDelaunayTriangulation2D(
+		UDynamicMesh* TargetMesh,
+		FGeometryScriptPrimitiveOptions PrimitiveOptions,
+		FTransform Transform,
+		const TArray<FVector2D>& VertexPositions,
+		const TArray<FIntPoint>& ConstrainedEdges,
+		FGeometryScriptConstrainedDelaunayTriangulationOptions TriangulationOptions,
+		TArray<int32>& PositionsToVertexIDs,
+		bool& bHasDuplicateVertices,
+		UGeometryScriptDebug* Debug = nullptr);
+
 
 
 	//---------------------------------------------
diff --git a/Engine/Source/Runtime/GeometryCore/Private/CompGeom/Delaunay2.cpp b/Engine/Source/Runtime/GeometryCore/Private/CompGeom/Delaunay2.cpp
--- a/Engine/Source/Runtime/GeometryCore/Private/CompGeom/Delaunay2.cpp
+++ b/Engine/Source/Runtime/GeometryCore/Private/CompGeom/Delaunay2.cpp
@@ -25,6 +25,7 @@ struct FDelaunay2Connectivity
 
 	void Empty(int32 ExpectedMaxVertices = 0)
 	{
+		bHasDuplicates = false;
 		EdgeToVert.Empty(ExpectedMaxVertices * 2 * 3);
 		DisableVertexAdjacency();
 	}
@@ -539,6 +540,18 @@ struct FDelaunay2Connectivity
 		return bTrackDuplicateVertices;
 	}
 
+	// called during triangulation when a duplicate vertex is found
+	void SetHasDuplicates()
+	{
+		bHasDuplicates = true;
+	}
+
+	// @return whether duplicate vertices were detected during triangulation
+	bool HasDuplicates() const
+	{
+		return bHasDuplicates;
+	}
+
 	void FixDuplicatesOnEdge(FIndex2i& Edge) const
 	{
 		const int32* OrigA = DuplicateVertices.Find(Edge.A);
@@ -553,6 +566,12 @@ struct FDelaunay2Connectivity
 		}
 	}
 
+	int32 RemapIfDuplicate(int32 Index) const
+	{
+		const int32* Remap = DuplicateVertices.Find(Index);
+		return Remap ? *Remap : Index;
+	}
+
 	// Get any edge BC opposite vertex A, such that triangle ABC is in the mesh (or return InvalidIndex edge if no such edge is present)
 	// Before calling this frequently, consider calling EnableVertexAdjacency()
 	// Not thread-safe if cache is enabled, because it may try to update the cache
@@ -640,6 +659,7 @@ protected:
 	TArray<int32> VertexAdjCache;
 	bool bUseAdjCache = false;
 
+	bool bHasDuplicates = false;
 	bool bTrackDuplicateVertices = false;
 	TMap<int32, int32> DuplicateVertices;
 
@@ -1389,14 +1409,19 @@ namespace DelaunayInternal
 			int32 DuplicateOf = -1;
 			constexpr bool bAssumeDelaunay = true; // initial construction, before constraint edges, so safe to assume Delaunay
 			FIndex3i ContainingTri = WalkToContainingTri<RealType>(Random, Connectivity, Vertices, SearchTri, Vertex, bAssumeDelaunay, DuplicateOf);
-			if (DuplicateOf >= 0 || ContainingTri[0] == FDelaunay2Connectivity::InvalidIndex)
+			if (DuplicateOf >= 0)
 			{
-				if (DuplicateOf >= 0 && Connectivity.HasDuplicateTracking())
+				Connectivity.SetHasDuplicates();
+				if (Connectivity.HasDuplicateTracking())
 				{
 					Connectivity.MarkDuplicateVertex(DuplicateOf, Vertex);
 				}
 				continue;
 			}
+			if (ContainingTri[0] == FDelaunay2Connectivity::InvalidIndex)
+			{
+				continue;
+			}
 			SearchTri = Insert<RealType>(Connectivity, Vertices, ContainingTri, Vertex);
 			checkSlow(SearchTri.A != FDelaunay2Connectivity::InvalidIndex);
 		}
@@ -1557,5 +1582,16 @@ bool FDelaunay2::HasEdge(const FIndex2i& Edge, bool bRemapDuplicates)
 	return Connectivity->HasEdge(Edge);
 }
 
+bool FDelaunay2::HasDuplicates() const
+{
+	return Connectivity->HasDuplicates();
+}
+
+int32 FDelaunay2::RemapIfDuplicate(int32 Index) const
+{
+	checkSlow(Connectivity->HasDuplicateTracking());
+	return Connectivity->RemapIfDuplicate(Index);
+}
+
 } // end namespace UE::Geometry
 } // end namespace UE
\ No newline at end of file
diff --git a/Engine/Source/Runtime/GeometryCore/Public/CompGeom/Delaunay2.h b/Engine/Source/Runtime/GeometryCore/Public/CompGeom/Delaunay2.h
--- a/Engine/Source/Runtime/GeometryCore/Public/CompGeom/Delaunay2.h
+++ b/Engine/Source/Runtime/GeometryCore/Public/CompGeom/Delaunay2.h
@@ -226,6 +226,14 @@ public:
 	// @return true if the edge is in the triangulation
 	bool HasEdge(const FIndex2i& Edge, bool bRemapDuplicates);
 
+	// @return true if the input had duplicate vertices, and so not all vertices could be used in the triangulation.
+	bool HasDuplicates() const;
+
+	// @return if bAutomaticallyFixEdgesToDuplicateVertices was set during triangulation, will return the remapped vertex index for any duplicate vertices, or Index if the vertex was not remapped.
+	int32 RemapIfDuplicate(int32 Index) const;
+
+
+
 	/**
 	 * Get Voronoi diagram cells as dual of the Delaunay triangulation.  You must call Triangulate() before calling this.
 	 * @param Vertices			Vertices of the Delaunay triangulation, to be used as the sites of the Voronoi diagram
