make edges array optional for Geometry Script's Delaunay Triangulation node

diff --git a/Engine/Plugins/Experimental/GeometryScripting/Source/GeometryScriptingCore/Public/GeometryScript/MeshPrimitiveFunctions.h b/Engine/Plugins/Experimental/GeometryScripting/Source/GeometryScriptingCore/Public/GeometryScript/MeshPrimitiveFunctions.h
--- a/Engine/Plugins/Experimental/GeometryScripting/Source/GeometryScriptingCore/Public/GeometryScript/MeshPrimitiveFunctions.h
+++ b/Engine/Plugins/Experimental/GeometryScripting/Source/GeometryScriptingCore/Public/GeometryScript/MeshPrimitiveFunctions.h
@@ -514,7 +514,7 @@ public:
 	* On success, all vertices are always appended to the output mesh, though duplicate vertices will not be used in any triangles and may optionally be removed.
 	* Use PositionsToVertexIDs to map indices in the input VertexPositions array to vertex IDs in the Dynamic Mesh.
 	*/
-	UFUNCTION(BlueprintCallable, Category = "GeometryScript|Primitives", meta = (ScriptMethod))
+	UFUNCTION(BlueprintCallable, Category = "GeometryScript|Primitives", meta = (ScriptMethod, AutoCreateRefTerm = "ConstrainedEdges"))
 	static UPARAM(DisplayName = "Target Mesh") UDynamicMesh*
 	AppendDelaunayTriangulation2D(
 		UDynamicMesh* TargetMesh,
