Added two new functions for creating static and skeletal meshes from an array of LODs instead of just a single mesh.
Use the DynamicMesh bone attributes (if available) to reindex the skin weight bone indices with the respect to the input skeleton.

diff --git a/Engine/Plugins/Experimental/GeometryScripting/Source/GeometryScriptingEditor/Private/CreateNewAssetUtilityFunctions.cpp b/Engine/Plugins/Experimental/GeometryScripting/Source/GeometryScriptingEditor/Private/CreateNewAssetUtilityFunctions.cpp
--- a/Engine/Plugins/Experimental/GeometryScripting/Source/GeometryScriptingEditor/Private/CreateNewAssetUtilityFunctions.cpp
+++ b/Engine/Plugins/Experimental/GeometryScripting/Source/GeometryScriptingEditor/Private/CreateNewAssetUtilityFunctions.cpp
@@ -3,11 +3,15 @@
 #include "GeometryScript/CreateNewAssetUtilityFunctions.h"
 
 #include "DynamicMesh/DynamicMesh3.h"
+#include "DynamicMesh/DynamicVertexSkinWeightsAttribute.h"
+#include "DynamicMesh/DynamicBoneAttribute.h"
 #include "UDynamicMesh.h"
 
 #include "Engine/StaticMesh.h"
 #include "Engine/SkeletalMesh.h"
 #include "Engine/StaticMeshActor.h"
+#include "Animation/Skeleton.h"
+#include "ReferenceSkeleton.h"
 #include "StaticMeshAttributes.h"
 #include "StaticMeshOperations.h"
 #include "RenderingThread.h"
@@ -145,33 +149,43 @@ AVolume* UGeometryScriptLibrary_CreateNewAssetFunctions::CreateNewVolumeFromMesh
 
 
 
-UStaticMesh* UGeometryScriptLibrary_CreateNewAssetFunctions::CreateNewStaticMeshAssetFromMesh(
-	UDynamicMesh* FromDynamicMesh, 
+UStaticMesh* UGeometryScriptLibrary_CreateNewAssetFunctions::CreateNewStaticMeshAssetFromMeshLODs(
+	TArray<UDynamicMesh*> FromDynamicMeshLODs, 
 	FString AssetPathAndName,
 	FGeometryScriptCreateNewStaticMeshAssetOptions Options,
 	EGeometryScriptOutcomePins& Outcome,
 	UGeometryScriptDebug* Debug)
 {
 	Outcome = EGeometryScriptOutcomePins::Failure;
-	if (FromDynamicMesh == nullptr)
+	if (FromDynamicMeshLODs.IsEmpty())
 	{
-		UE::Geometry::AppendError(Debug, EGeometryScriptErrorType::InvalidInputs, LOCTEXT("CreateNewStaticMeshAssetFromMesh_InvalidInput1", "CreateNewStaticMeshAssetFromMesh: FromDynamicMesh is Null"));
+		AppendError(Debug, EGeometryScriptErrorType::InvalidInputs, LOCTEXT("CreateNewStaticMeshAssetFromMeshLODs_EmptyLODArray", "CreateNewStaticMeshAssetFromMeshLODs: LOD array is empty"));
 		return nullptr;
 	}
-	if (FromDynamicMesh->GetTriangleCount() == 0)
-	{
-		UE::Geometry::AppendError(Debug, EGeometryScriptErrorType::InvalidInputs, LOCTEXT("CreateNewStaticMeshAssetFromMesh_InvalidInput3", "CreateNewStaticMeshAssetFromMesh: FromDynamicMesh has zero triangles"));
-		return nullptr;
+
+	for (int LodIdx = 0; LodIdx < FromDynamicMeshLODs.Num(); ++LodIdx)
+	{	
+		const UDynamicMesh* FromDynamicMesh = FromDynamicMeshLODs[LodIdx];
+		if (FromDynamicMesh == nullptr)
+		{
+			UE::Geometry::AppendError(Debug, EGeometryScriptErrorType::InvalidInputs, FText::Format(LOCTEXT("CreateNewStaticMeshAssetFromMeshLODs_LODIsNull", "CreateNewStaticMeshAssetFromMeshLODs: LOD {0} is Null"), FText::AsNumber(LodIdx)));
+			return nullptr;
+		}
+		if (FromDynamicMesh->GetTriangleCount() == 0)
+		{
+			UE::Geometry::AppendError(Debug, EGeometryScriptErrorType::InvalidInputs, FText::Format(LOCTEXT("CreateNewStaticMeshAssetFromMeshLODs_ZeroTriangles", "CreateNewStaticMeshAssetFromMeshLODs: LOD {0} has no triangles"), FText::AsNumber(LodIdx)));
+			return nullptr;
+		}
 	}
 	// todo: other safety checks
 
-	GEditor->BeginTransaction(LOCTEXT("CreateNewStaticMeshAssetFromMesh_Transaction", "Create StaticMesh"));
+	GEditor->BeginTransaction(LOCTEXT("CreateNewStaticMeshAssetFromMeshLODs_Transaction", "Create StaticMesh"));
 
 	UE::AssetUtils::FStaticMeshAssetOptions AssetOptions;
 	AssetPathAndName = UPackageTools::SanitizePackageName(AssetPathAndName);
 	AssetOptions.NewAssetPath = AssetPathAndName;
 
-	AssetOptions.NumSourceModels = 1;
+	AssetOptions.NumSourceModels = FromDynamicMeshLODs.Num();
 
 	// CreateStaticMeshAsset below will handle this, but we could allow passing in materials as an option...
 	//AssetOptions.NumMaterialSlots = CreateMeshParams.Materials.Num();
@@ -186,19 +200,32 @@ UStaticMesh* UGeometryScriptLibrary_CreateNewAssetFunctions::CreateNewStaticMesh
 	AssetOptions.bCreatePhysicsBody = Options.bEnableCollision;
 	AssetOptions.CollisionType = Options.CollisionMode;
 
-	FDynamicMesh3 LODMesh;
-	FromDynamicMesh->ProcessMesh([&](const FDynamicMesh3& ReadMesh)
+	/**
+	 * We are making a copy of each LOD mesh since UDynamicMesh can potentially be editable asynchronously in the future, 
+	 * so we should not hold onto the pointer outside the function.
+	 */
+	TArray<FDynamicMesh3> CopyFromDynamicMeshLODs;
+	CopyFromDynamicMeshLODs.SetNum(FromDynamicMeshLODs.Num());
+
+	for (int LodIdx = 0; LodIdx < FromDynamicMeshLODs.Num(); ++LodIdx) 
 	{
-		LODMesh = ReadMesh;
-	});
-	AssetOptions.SourceMeshes.DynamicMeshes.Add(&LODMesh);
+		const UDynamicMesh* LODMesh = FromDynamicMeshLODs[LodIdx];
+		FDynamicMesh3* CopyLODMesh = &CopyFromDynamicMeshLODs[LodIdx];
+
+		LODMesh->ProcessMesh([CopyLODMesh](const FDynamicMesh3& ReadMesh)
+		{
+			*CopyLODMesh = ReadMesh;
+		});
+
+		AssetOptions.SourceMeshes.DynamicMeshes.Add(CopyLODMesh);
+	}
 
 	UE::AssetUtils::FStaticMeshResults ResultData;
 	UE::AssetUtils::ECreateStaticMeshResult AssetResult = UE::AssetUtils::CreateStaticMeshAsset(AssetOptions, ResultData);
 
 	if (AssetResult != UE::AssetUtils::ECreateStaticMeshResult::Ok)
 	{
-		UE::Geometry::AppendError(Debug, EGeometryScriptErrorType::OperationFailed, LOCTEXT("CreateNewStaticMeshAssetFromMesh_Failed", "CreateNewStaticMeshAssetFromMesh: Failed to create new Asset"));
+		UE::Geometry::AppendError(Debug, EGeometryScriptErrorType::OperationFailed, LOCTEXT("CreateNewStaticMeshAssetFromMeshLODs_Failed", "CreateNewStaticMeshAssetFromMeshLODs: Failed to create new Asset"));
 		return nullptr;
 	}
 
@@ -214,9 +241,19 @@ UStaticMesh* UGeometryScriptLibrary_CreateNewAssetFunctions::CreateNewStaticMesh
 	return NewStaticMesh;
 }
 
+UStaticMesh* UGeometryScriptLibrary_CreateNewAssetFunctions::CreateNewStaticMeshAssetFromMesh(
+	UDynamicMesh* FromDynamicMesh, 
+	FString AssetPathAndName,
+	FGeometryScriptCreateNewStaticMeshAssetOptions Options,
+	EGeometryScriptOutcomePins& Outcome,
+	UGeometryScriptDebug* Debug)
+{
+	return CreateNewStaticMeshAssetFromMeshLODs({FromDynamicMesh}, AssetPathAndName, Options, Outcome, Debug);
+}
+
 
-USkeletalMesh* UGeometryScriptLibrary_CreateNewAssetFunctions::CreateNewSkeletalMeshAssetFromMesh(
-	UDynamicMesh* FromDynamicMesh,
+USkeletalMesh* UGeometryScriptLibrary_CreateNewAssetFunctions::CreateNewSkeletalMeshAssetFromMeshLODs(
+	TArray<UDynamicMesh*> FromDynamicMeshLODs,
 	USkeleton* InSkeleton,
 	FString AssetPathAndName, 
 	FGeometryScriptCreateNewSkeletalMeshAssetOptions Options,
@@ -226,37 +263,47 @@ USkeletalMesh* UGeometryScriptLibrary_CreateNewAssetFunctions::CreateNewSkeletal
 	using namespace UE::AssetUtils;
 
 	Outcome = EGeometryScriptOutcomePins::Failure;
-	if (FromDynamicMesh == nullptr)
+	if (FromDynamicMeshLODs.IsEmpty())
 	{
-		AppendError(Debug, EGeometryScriptErrorType::InvalidInputs, LOCTEXT("CreateNewSkeletalMeshAssetFromMesh_InvalidInput1", "CreateNewSkeletalMeshAssetFromMesh: FromDynamicMesh is Null"));
+		AppendError(Debug, EGeometryScriptErrorType::InvalidInputs, LOCTEXT("CreateNewSkeletalMeshAssetFromMeshLODs_EmptyLODArray", "CreateNewSkeletalMeshAssetFromMeshLODs: FromDynamicMesh array is empty"));
 		return nullptr;
 	}
-	if (FromDynamicMesh->GetTriangleCount() == 0)
-	{
-		AppendError(Debug, EGeometryScriptErrorType::InvalidInputs, LOCTEXT("CreateNewSkeletalMeshAssetFromMesh_InvalidInput2", "CreateNewSkeletalMeshAssetFromMesh: FromDynamicMesh has zero triangles"));
-		return nullptr;
-	}
-	if (FromDynamicMesh->GetMeshRef().Attributes()->GetSkinWeightsAttributes().Num() == 0)
-	{
-		AppendError(Debug, EGeometryScriptErrorType::InvalidInputs, LOCTEXT("CreateNewSkeletalMeshAssetFromMesh_InvalidInput3", "CreateNewSkeletalMeshAssetFromMesh: FromDynamicMesh has no skin weight attributes"));
-		return nullptr;
-	}
-	if (InSkeleton == nullptr)
-	{
-		AppendError(Debug, EGeometryScriptErrorType::InvalidInputs, LOCTEXT("CreateNewSkeletalMeshAssetFromMesh_InvalidSkeleton", "CreateNewSkeletalMeshAssetFromMesh: Skeleton is Null"));
-		return nullptr;
+
+	for (int LodIdx = 0; LodIdx < FromDynamicMeshLODs.Num(); ++LodIdx)
+	{	
+		const UDynamicMesh* FromDynamicMesh = FromDynamicMeshLODs[LodIdx];
+		if (FromDynamicMesh == nullptr)
+		{
+			UE::Geometry::AppendError(Debug, EGeometryScriptErrorType::InvalidInputs, FText::Format(LOCTEXT("CreateNewSkeletalMeshAssetFromMeshLODs_LODIsNull", "CreateNewSkeletalMeshAssetFromMeshLODs: LOD {0} is Null"), FText::AsNumber(LodIdx)));
+			return nullptr;
+		}
+		if (FromDynamicMesh->GetTriangleCount() == 0)
+		{
+			AppendError(Debug, EGeometryScriptErrorType::InvalidInputs, FText::Format(LOCTEXT("CreateNewSkeletalMeshAssetFromMeshLODs_ZeroTriangles", "CreateNewSkeletalMeshAssetFromMeshLODs: LOD {0} has no triangles"), FText::AsNumber(LodIdx)));
+			return nullptr;
+		}
+		if (FromDynamicMesh->GetMeshRef().HasAttributes() == false || FromDynamicMesh->GetMeshRef().Attributes()->GetSkinWeightsAttributes().Num() == 0)
+		{
+			AppendError(Debug, EGeometryScriptErrorType::InvalidInputs, FText::Format(LOCTEXT("CreateNewSkeletalMeshAssetFromMeshLODs_NoSkinWeights", "CreateNewSkeletalMeshAssetFromMeshLODs: LOD {0} has no skin weight attributes"), FText::AsNumber(LodIdx)));
+			return nullptr;
+		}
+		if (InSkeleton == nullptr)
+		{
+			AppendError(Debug, EGeometryScriptErrorType::InvalidInputs, LOCTEXT("CreateNewSkeletalMeshAssetFromMeshLODs_NullSkeleton", "CreateNewSkeletalMeshAssetFromMeshLODs: Skeleton is Null"));
+			return nullptr;
+		}
 	}
 	
 	// todo: other safety checks
 
-	GEditor->BeginTransaction(LOCTEXT("CreateNewSkeletalMeshAssetFromMesh_Transaction", "Create SkeletalMesh"));
+	GEditor->BeginTransaction(LOCTEXT("CreateNewSkeletalMeshAssetFromMeshLODs_Transaction", "Create SkeletalMesh"));
 
 	FSkeletalMeshAssetOptions AssetOptions;
 	AssetPathAndName = UPackageTools::SanitizePackageName(AssetPathAndName);
 	AssetOptions.NewAssetPath = AssetPathAndName;
 	AssetOptions.Skeleton = InSkeleton;
 
-	AssetOptions.NumSourceModels = 1;
+	AssetOptions.NumSourceModels = FromDynamicMeshLODs.Num();
 
 	if (!Options.Materials.IsEmpty())
 	{
@@ -277,19 +324,81 @@ USkeletalMesh* UGeometryScriptLibrary_CreateNewAssetFunctions::CreateNewSkeletal
 	AssetOptions.bEnableRecomputeNormals = Options.bEnableRecomputeNormals;
 	AssetOptions.bEnableRecomputeTangents = Options.bEnableRecomputeTangents;
 
-	FDynamicMesh3 LODMesh;
-	FromDynamicMesh->ProcessMesh([&](const FDynamicMesh3& ReadMesh)
+	/**
+	 * We are making a copy of each LOD mesh since UDynamicMesh can potentially be editable asynchronously in the future, 
+	 * so we should not hold onto the pointer outside the function.
+	 */
+	TArray<FDynamicMesh3> CopyFromDynamicMeshLODs;
+	CopyFromDynamicMeshLODs.SetNum(FromDynamicMeshLODs.Num());
+
+	for (int LodIdx = 0; LodIdx < FromDynamicMeshLODs.Num(); ++LodIdx) 
 	{
-		LODMesh = ReadMesh;
-	});
-	AssetOptions.SourceMeshes.DynamicMeshes.Add(&LODMesh);
+		const UDynamicMesh* LODMesh = FromDynamicMeshLODs[LodIdx];
+		FDynamicMesh3* CopyLODMesh = &CopyFromDynamicMeshLODs[LodIdx];
+
+		LODMesh->ProcessMesh([CopyLODMesh](const FDynamicMesh3& ReadMesh)
+		{
+			*CopyLODMesh = ReadMesh;
+		});
+	}
+
+	// Check if all LODs have bone attributes.
+	bool bMeshHasBoneAttribs = true;
+	for (const FDynamicMesh3& FromDynamicMesh : CopyFromDynamicMeshLODs)
+	{
+		bMeshHasBoneAttribs = bMeshHasBoneAttribs && FromDynamicMesh.Attributes()->HasBones();
+	}
+
+	if (!bMeshHasBoneAttribs) // If at least one LOD doesn't contain the bone attributes then add LOD mesh as is
+	{
+		for (const FDynamicMesh3& FromDynamicMesh : CopyFromDynamicMeshLODs)
+		{
+			AssetOptions.SourceMeshes.DynamicMeshes.Add(&FromDynamicMesh);
+		}
+	}
+	else // If bone attributes are available then attempt to reindex the weights
+	{	
+		TArray<FName> ToSkeleton;
+		ToSkeleton.Reserve(InSkeleton->GetReferenceSkeleton().GetRawRefBoneInfo().Num());
+		for (const FMeshBoneInfo& BoneInfo : InSkeleton->GetReferenceSkeleton().GetRawRefBoneInfo())
+		{
+			ToSkeleton.Add(BoneInfo.Name);
+		}
+
+		for (int LodIdx = 0; LodIdx < CopyFromDynamicMeshLODs.Num(); ++LodIdx)
+		{
+			FDynamicMesh3& FromDynamicMesh = CopyFromDynamicMeshLODs[LodIdx];
+			FDynamicMeshAttributeSet* AttribSet = FromDynamicMesh.Attributes();
+
+			// Check if the skeleton we are trying to bind the mesh to is the same as the current mesh skeleton.
+			const TArray<FName>& FromSkeleton = AttribSet->GetBoneNames()->GetAttribValues();
+
+			if (FromSkeleton != ToSkeleton)
+			{
+				for (const TPair<FName, TUniquePtr<FDynamicMeshVertexSkinWeightsAttribute>>& Entry : AttribSet->GetSkinWeightsAttributes())
+				{
+					FDynamicMeshVertexSkinWeightsAttribute* SkinWeightAttrib = Entry.Value.Get();
+					
+					// Reindex the bone indices
+					if (SkinWeightAttrib->ReindexBoneIndicesToSkeleton(FromSkeleton, ToSkeleton) == false)
+					{
+						AppendError(Debug, EGeometryScriptErrorType::OperationFailed, FText::Format(LOCTEXT("CreateNewSkeletalMeshAssetFromMeshLODs_FailedReindexing", "CreateNewSkeletalMeshAssetFromMeshLODs: LOD {0} has invalid skinning data or the bone data is not compatible with the specified skeleton."), FText::AsNumber(LodIdx)));
+						return nullptr;
+					}
+				}
+			}
+
+			AssetOptions.SourceMeshes.DynamicMeshes.Add(&FromDynamicMesh);
+		}
+	}
+	
 
 	FSkeletalMeshResults ResultData;
 	const ECreateSkeletalMeshResult AssetResult = CreateSkeletalMeshAsset(AssetOptions, ResultData);
 
 	if (AssetResult != ECreateSkeletalMeshResult::Ok)
 	{
-		AppendError(Debug, EGeometryScriptErrorType::OperationFailed, LOCTEXT("CreateNewSkeletalMeshAssetFromMesh_Failed", "CreateNewSkeletalMeshAssetFromMesh: Failed to create new Asset"));
+		AppendError(Debug, EGeometryScriptErrorType::OperationFailed, LOCTEXT("CreateNewSkeletalMeshAssetFromMeshLODs_Failed", "CreateNewSkeletalMeshAssetFromMeshLODs: Failed to create new Asset"));
 		return nullptr;
 	}
 
@@ -303,6 +412,17 @@ USkeletalMesh* UGeometryScriptLibrary_CreateNewAssetFunctions::CreateNewSkeletal
 }
 
 
+USkeletalMesh* UGeometryScriptLibrary_CreateNewAssetFunctions::CreateNewSkeletalMeshAssetFromMesh(
+	UDynamicMesh* FromDynamicMeshLODs,
+	USkeleton* InSkeleton,
+	FString AssetPathAndName, 
+	FGeometryScriptCreateNewSkeletalMeshAssetOptions Options,
+	EGeometryScriptOutcomePins& Outcome, 
+	UGeometryScriptDebug* Debug)
+{
+	return CreateNewSkeletalMeshAssetFromMeshLODs({FromDynamicMeshLODs}, InSkeleton, AssetPathAndName, Options, Outcome, Debug);
+}
+
 UTexture2D* UGeometryScriptLibrary_CreateNewAssetFunctions::CreateNewTexture2DAsset(
 		UTexture2D* FromTexture, 
 		FString AssetPathAndName,
diff --git a/Engine/Plugins/Experimental/GeometryScripting/Source/GeometryScriptingEditor/Public/GeometryScript/CreateNewAssetUtilityFunctions.h b/Engine/Plugins/Experimental/GeometryScripting/Source/GeometryScriptingEditor/Public/GeometryScript/CreateNewAssetUtilityFunctions.h
--- a/Engine/Plugins/Experimental/GeometryScripting/Source/GeometryScriptingEditor/Public/GeometryScript/CreateNewAssetUtilityFunctions.h
+++ b/Engine/Plugins/Experimental/GeometryScripting/Source/GeometryScriptingEditor/Public/GeometryScript/CreateNewAssetUtilityFunctions.h
@@ -128,6 +128,9 @@ public:
 		EGeometryScriptOutcomePins& Outcome,
 		UGeometryScriptDebug* Debug = nullptr);
 
+	/**
+	 * Create a new StaticMesh asset from a DynamicMesh. Save the asset at the AssetPathAndName location.
+	 */
 	UFUNCTION(BlueprintCallable, Category = "GeometryScript|AssetManagement", meta = (ExpandEnumAsExecs = "Outcome"))
 	static UPARAM(DisplayName = "Static Mesh Asset") UStaticMesh* 
 	CreateNewStaticMeshAssetFromMesh(
@@ -137,6 +140,23 @@ public:
 		EGeometryScriptOutcomePins& Outcome,
 		UGeometryScriptDebug* Debug = nullptr);
 	
+	/**
+	 * Create a new StaticMesh asset from a collection of LODs represented by an array of DynamicMeshes.
+	 * Save the asset at the AssetPathAndName location.
+	 */
+	UFUNCTION(BlueprintCallable, Category = "GeometryScript|AssetManagement", meta = (ExpandEnumAsExecs = "Outcome"))
+	static UPARAM(DisplayName = "Static Mesh Asset") UStaticMesh* 
+	CreateNewStaticMeshAssetFromMeshLODs(
+		TArray<UDynamicMesh*> FromDynamicMesh, 
+		FString AssetPathAndName,
+		FGeometryScriptCreateNewStaticMeshAssetOptions Options,
+		EGeometryScriptOutcomePins& Outcome,
+		UGeometryScriptDebug* Debug = nullptr);
+	
+	/**
+	 * Create a new SkeletalMesh asset from a DynamicMesh and a Skeleton.
+	 * Save the asset at the AssetPathAndName location.
+	 */
 	UFUNCTION(BlueprintCallable, Category = "GeometryScript|AssetManagement", meta = (ExpandEnumAsExecs = "Outcome"))
 	static UPARAM(DisplayName = "Skeletal Mesh Asset") USkeletalMesh* 
 	CreateNewSkeletalMeshAssetFromMesh(
@@ -146,6 +166,20 @@ public:
 		FGeometryScriptCreateNewSkeletalMeshAssetOptions Options,
 		EGeometryScriptOutcomePins& Outcome,
 		UGeometryScriptDebug* Debug = nullptr);
+	
+	/**
+	 * Create a new SkeletalMesh asset from a collection of LODs represented by an array of DynamicMeshes and a Skeleton.
+	 * Save the asset at the AssetPathAndName location.
+	 */
+	UFUNCTION(BlueprintCallable, Category = "GeometryScript|AssetManagement", meta = (ExpandEnumAsExecs = "Outcome"))
+	static UPARAM(DisplayName = "Skeletal Mesh Asset") USkeletalMesh* 
+	CreateNewSkeletalMeshAssetFromMeshLODs(
+		TArray<UDynamicMesh*> FromDynamicMeshLODs, 
+		USkeleton* InSkeleton,
+		FString AssetPathAndName,
+		FGeometryScriptCreateNewSkeletalMeshAssetOptions Options,
+		EGeometryScriptOutcomePins& Outcome,
+		UGeometryScriptDebug* Debug = nullptr);
 
 	UFUNCTION(BlueprintCallable, Category = "GeometryScript|AssetManagement", meta = (ExpandEnumAsExecs = "Outcome"))
 	static UPARAM(DisplayName = "Texture 2D Asset") UTexture2D* 
