Added geometry script functions for:
-Discarding bone attributes.
-Copying bone attributes from one mesh to another.
-Getting an index of the bone by name.

diff --git a/Engine/Plugins/Experimental/GeometryScripting/Source/GeometryScriptingCore/Private/MeshBoneWeightFunctions.cpp b/Engine/Plugins/Experimental/GeometryScripting/Source/GeometryScriptingCore/Private/MeshBoneWeightFunctions.cpp
--- a/Engine/Plugins/Experimental/GeometryScripting/Source/GeometryScriptingCore/Private/MeshBoneWeightFunctions.cpp
+++ b/Engine/Plugins/Experimental/GeometryScripting/Source/GeometryScriptingCore/Private/MeshBoneWeightFunctions.cpp
@@ -7,6 +7,7 @@
 #include "DynamicMesh/DynamicMesh3.h"
 #include "DynamicMesh/DynamicMeshAttributeSet.h"
 #include "DynamicMesh/DynamicVertexSkinWeightsAttribute.h"
+#include "DynamicMesh/DynamicBoneAttribute.h"
 #include "SkinningOps/SkinBindingOp.h"
 #include "UDynamicMesh.h"
 
@@ -342,5 +343,102 @@ UDynamicMesh* UGeometryScriptLibrary_MeshBoneWeightFunctions::ComputeSmoothBoneW
 	return TargetMesh;
 }
 
+UDynamicMesh* UGeometryScriptLibrary_MeshBoneWeightFunctions::CopyBonesFromMesh(
+	UDynamicMesh* SourceMesh, 
+	UDynamicMesh* TargetMesh, 
+	UGeometryScriptDebug* Debug)
+{
+	if (SourceMesh == nullptr)
+	{
+		AppendError(Debug, EGeometryScriptErrorType::InvalidInputs, LOCTEXT("CopyBonesFromMesh_InvalidSourceMesh", "CopyBonesFromMesh: SourceMesh is Null"));
+		return TargetMesh;
+	}
+
+	if (TargetMesh == nullptr)
+	{
+		AppendError(Debug, EGeometryScriptErrorType::InvalidInputs, LOCTEXT("CopyBonesFromMesh_InvalidTargetMesh", "CopyBonesFromMesh: TargetMesh is Null"));
+		return TargetMesh;
+	}
+
+	SourceMesh->EditMesh([&](const FDynamicMesh3& ReadMesh)
+	{
+		TargetMesh->EditMesh([&](FDynamicMesh3& EditMesh)
+		{
+			if (!ReadMesh.HasAttributes() || !ReadMesh.Attributes()->HasBones())
+			{
+				AppendWarning(Debug, EGeometryScriptErrorType::InvalidInputs, LOCTEXT("CopyBonesFromMesh_SourceMeshHasNoBones", "SourceMesh has no bone attributes"));
+				return;
+			}
+
+			if (!ReadMesh.Attributes()->CheckBoneValidity(EValidityCheckFailMode::ReturnOnly))
+			{
+				AppendError(Debug, EGeometryScriptErrorType::InvalidInputs, LOCTEXT("CopyBonesFromMesh_InvalidSourceMeshBones", "SourceMesh has invalid bone attributes"));
+				return;
+			}
+
+			if (!EditMesh.HasAttributes())
+			{
+				EditMesh.EnableAttributes();
+			}
+				
+			EditMesh.Attributes()->CopyBoneAttributes(*ReadMesh.Attributes());
+		}, EDynamicMeshChangeType::AttributeEdit, EDynamicMeshAttributeChangeFlags::Unknown, false);
+	});
+
+	return TargetMesh;
+}
+
+UDynamicMesh* UGeometryScriptLibrary_MeshBoneWeightFunctions::DiscardBonesFromMesh(
+	UDynamicMesh* TargetMesh, 
+	UGeometryScriptDebug* Debug)
+{
+	if (TargetMesh == nullptr)
+	{
+		AppendError(Debug, EGeometryScriptErrorType::InvalidInputs, LOCTEXT("DiscardBonesFromMesh_InvalidTargetMesh", "DiscardBonesFromMesh: TargetMesh is Null"));
+		return TargetMesh;
+	}
+
+	TargetMesh->EditMesh([&](FDynamicMesh3& EditMesh)
+	{
+		if (EditMesh.HasAttributes())
+		{
+			EditMesh.Attributes()->DisableBones();
+		}
+		
+	}, EDynamicMeshChangeType::AttributeEdit, EDynamicMeshAttributeChangeFlags::Unknown, false);
+
+	return TargetMesh;
+}
+
+UDynamicMesh* UGeometryScriptLibrary_MeshBoneWeightFunctions::GetBoneIndexForName(
+	UDynamicMesh* TargetMesh,
+	FString BoneName,
+	bool& bIsValidBoneName,
+	int& BoneIndex,
+	UGeometryScriptDebug* Debug)
+{
+	if (TargetMesh == nullptr)
+	{
+		AppendError(Debug, EGeometryScriptErrorType::InvalidInputs, LOCTEXT("GetBoneIndexForName_InvalidTargetMesh", "GetBoneIndexForName: TargetMesh is Null"));
+		return TargetMesh;
+	}
+
+	TargetMesh->ProcessMesh([&](const FDynamicMesh3& EditMesh)
+	{	
+		if (!EditMesh.HasAttributes() || !EditMesh.Attributes()->HasBones())
+		{
+			AppendError(Debug, EGeometryScriptErrorType::InvalidInputs, LOCTEXT("GetBoneIndexForName_TargetMeshHasNoBones", "TargetMesh has no bone attributes"));
+			return;
+		}
+
+		// INDEX_NONE if BoneName doesn't exist in the bone names attribute
+		BoneIndex = EditMesh.Attributes()->GetBoneNames()->GetAttribValues().Find(FName(BoneName));
+		
+		bIsValidBoneName = BoneIndex == INDEX_NONE;
+
+	});
+
+	return TargetMesh;
+}
 
 #undef LOCTEXT_NAMESPACE
diff --git a/Engine/Plugins/Experimental/GeometryScripting/Source/GeometryScriptingCore/Public/GeometryScript/MeshBoneWeightFunctions.h b/Engine/Plugins/Experimental/GeometryScripting/Source/GeometryScriptingCore/Public/GeometryScript/MeshBoneWeightFunctions.h
--- a/Engine/Plugins/Experimental/GeometryScripting/Source/GeometryScriptingCore/Public/GeometryScript/MeshBoneWeightFunctions.h
+++ b/Engine/Plugins/Experimental/GeometryScripting/Source/GeometryScriptingCore/Public/GeometryScript/MeshBoneWeightFunctions.h
@@ -198,4 +198,41 @@ public:
 		FGeometryScriptBoneWeightProfile Profile = FGeometryScriptBoneWeightProfile(),
 		UGeometryScriptDebug* Debug = nullptr);
 	
+	/**
+	 * Copy the bone attributes (skeleton) from the SourceMesh to the TargetMesh.
+	 * @param SourceMesh Mesh we are copying the bone attributes from.
+	 * @param TargetMesh Mesh we are copying the bone attributes to.
+	 */
+	UFUNCTION(BlueprintCallable, Category = "GeometryScript|MeshQueries|BoneWeights", meta=(ScriptMethod))
+	static UPARAM(DisplayName = "Target Mesh") UDynamicMesh* 
+	CopyBonesFromMesh(
+		UDynamicMesh* SourceMesh,
+		UDynamicMesh* TargetMesh,
+		UGeometryScriptDebug* Debug = nullptr);
+
+	/**
+	 * Discard the bone attributes (skeleton) from the TargetMesh.
+	 * @param TargetMesh Mesh we are discarding bone attributes from.
+	 */
+	UFUNCTION(BlueprintCallable, Category = "GeometryScript|MeshQueries|BoneWeights", meta=(ScriptMethod))
+	static UPARAM(DisplayName = "Target Mesh") UDynamicMesh* 
+	DiscardBonesFromMesh(
+		UDynamicMesh* TargetMesh,
+		UGeometryScriptDebug* Debug = nullptr);
+
+	/**
+	 * Get the index of the bone with the given name.
+	 * @param TargetMesh Mesh containing the bone attributes.
+	 * @param BoneName The name of the bone whose index we are getting.
+	 * @param bIsValidBoneName Set to true if the TargetMesh contains a bone with the given name, false otherwise.
+	 * @param BoneIndex The index of the bone with the given name. Will be set to -1 if bIsValidBoneName is false.
+	 */
+	UFUNCTION(BlueprintCallable, Category = "GeometryScript|MeshQueries|BoneWeights", meta=(ScriptMethod))
+	static UPARAM(DisplayName = "Target Mesh") UDynamicMesh* 
+	GetBoneIndexForName(
+		UDynamicMesh* TargetMesh,
+		FString BoneName,
+		bool& bIsValidBoneName,
+		int& BoneIndex,
+		UGeometryScriptDebug* Debug = nullptr);
 };
\ No newline at end of file
