Added a geometry script function for transferring bone weights from one mesh to another.

diff --git a/Engine/Plugins/Experimental/GeometryScripting/Source/GeometryScriptingCore/Private/MeshBoneWeightFunctions.cpp b/Engine/Plugins/Experimental/GeometryScripting/Source/GeometryScriptingCore/Private/MeshBoneWeightFunctions.cpp
--- a/Engine/Plugins/Experimental/GeometryScripting/Source/GeometryScriptingCore/Private/MeshBoneWeightFunctions.cpp
+++ b/Engine/Plugins/Experimental/GeometryScripting/Source/GeometryScriptingCore/Private/MeshBoneWeightFunctions.cpp
@@ -10,6 +10,8 @@
 #include "DynamicMesh/DynamicBoneAttribute.h"
 #include "SkinningOps/SkinBindingOp.h"
 #include "UDynamicMesh.h"
+#include "Operations/TransferBoneWeights.h"
+
 
 #include UE_INLINE_GENERATED_CPP_BY_NAME(MeshBoneWeightFunctions)
 
@@ -343,6 +345,88 @@ UDynamicMesh* UGeometryScriptLibrary_MeshBoneWeightFunctions::ComputeSmoothBoneW
 	return TargetMesh;
 }
 
+
+UDynamicMesh* UGeometryScriptLibrary_MeshBoneWeightFunctions::TransferBoneWeightsFromMesh(
+	UDynamicMesh* SourceMesh,
+	UDynamicMesh* TargetMesh,
+	FGeometryScriptTransferBoneWeightsOptions Options,
+	UGeometryScriptDebug* Debug)
+{
+	using namespace UE::Geometry;
+
+	if (SourceMesh == nullptr)
+	{
+		AppendError(Debug, EGeometryScriptErrorType::InvalidInputs, LOCTEXT("TransferBoneWeightsFromMesh_InvalidSourceMesh", "TransferBoneWeightsFromMesh: Source Mesh is Null"));
+		return TargetMesh;
+	}
+	if (TargetMesh == nullptr)
+	{
+		AppendError(Debug, EGeometryScriptErrorType::InvalidInputs, LOCTEXT("TransferBoneWeightsFromMesh_InvalidTargetMesh", "TransferBoneWeightsFromMesh: Target Mesh is Null"));
+		return TargetMesh;
+	}
+
+	SourceMesh->ProcessMesh([&](const FDynamicMesh3& ReadMesh)
+	{
+		if (!ReadMesh.HasAttributes() || !ReadMesh.Attributes()->HasBones())
+		{
+			UE::Geometry::AppendError(Debug, EGeometryScriptErrorType::InvalidInputs, LOCTEXT("TransferBoneWeightsFromMesh_NoBones", "Source Mesh has no bone attribute"));
+			return;
+		}
+		if (ReadMesh.Attributes()->GetNumBones() == 0)
+		{
+			UE::Geometry::AppendError(Debug, EGeometryScriptErrorType::InvalidInputs, LOCTEXT("TransferBoneWeightsFromMesh_EmptyBones", "Source Mesh has an empty bone attribute"));
+			return;
+		}
+
+		FTransferBoneWeights TransferBoneWeights(&ReadMesh, Options.SourceProfile.GetProfileName());
+		TransferBoneWeights.TransferMethod = static_cast<FTransferBoneWeights::ETransferBoneWeightsMethod>(Options.TransferMethod);
+		if (TransferBoneWeights.Validate() != EOperationValidationResult::Ok)
+		{
+			UE::Geometry::AppendError(Debug, EGeometryScriptErrorType::OperationFailed, LOCTEXT("TransferBoneWeightsFromMesh_ValidationFailed", "TransferBoneWeightsFromMesh: Invalid parameters were set for the transfer weight operator"));
+			return;
+		}
+
+		TargetMesh->EditMesh([&](FDynamicMesh3& EditMesh)
+		{
+			if (!EditMesh.HasAttributes())
+			{
+				EditMesh.EnableAttributes();
+			}
+			
+			if (EditMesh.Attributes()->HasBones())
+			{
+				// If the TargetMesh has bone attributes, but we want to use the SourceMesh bone attributes, then we copy.
+				// Otherwise, nothing to do, and we use the target mesh bone attributes.
+				if (Options.OutputTargetMeshBones == EOutputTargetMeshBones::SourceBones)
+				{
+					EditMesh.Attributes()->CopyBoneAttributes(*ReadMesh.Attributes());
+				}
+			}
+			else
+			{
+				// If the TargetMesh has no bone attributes, then we must use the SourceMesh bone attributes. Otherwise, throw an error.
+				if (Options.OutputTargetMeshBones == EOutputTargetMeshBones::SourceBones)
+				{
+					EditMesh.Attributes()->CopyBoneAttributes(*ReadMesh.Attributes());
+				}
+				else 
+				{
+					UE::Geometry::AppendError(Debug, EGeometryScriptErrorType::InvalidInputs, LOCTEXT("TransferBoneWeightsFromMesh_NoTargetMeshBones", "TransferBoneWeightsFromMesh: TargetMesh has no bone attributes but the OutputTargetMeshBones option is set to TargetBones"));
+				}
+			}
+			
+			if (!TransferBoneWeights.Compute(EditMesh, FTransformSRT3d::Identity(), Options.TargetProfile.GetProfileName()))
+			{
+				UE::Geometry::AppendError(Debug, EGeometryScriptErrorType::OperationFailed, LOCTEXT("TransferBoneWeightsFromMesh_TransferFailed", "TransferBoneWeightsFromMesh: Failed to transfer the weights"));
+			}
+			
+		}, EDynamicMeshChangeType::AttributeEdit, EDynamicMeshAttributeChangeFlags::Unknown, false);
+	});
+
+	return TargetMesh;
+}
+
+
 UDynamicMesh* UGeometryScriptLibrary_MeshBoneWeightFunctions::CopyBonesFromMesh(
 	UDynamicMesh* SourceMesh, 
 	UDynamicMesh* TargetMesh, 
diff --git a/Engine/Plugins/Experimental/GeometryScripting/Source/GeometryScriptingCore/Public/GeometryScript/MeshBoneWeightFunctions.h b/Engine/Plugins/Experimental/GeometryScripting/Source/GeometryScriptingCore/Public/GeometryScript/MeshBoneWeightFunctions.h
--- a/Engine/Plugins/Experimental/GeometryScripting/Source/GeometryScriptingCore/Public/GeometryScript/MeshBoneWeightFunctions.h
+++ b/Engine/Plugins/Experimental/GeometryScripting/Source/GeometryScriptingCore/Public/GeometryScript/MeshBoneWeightFunctions.h
@@ -73,8 +73,54 @@ struct GEOMETRYSCRIPTINGCORE_API FGeometryScriptSmoothBoneWeightsOptions
 	int32 VoxelResolution = 256;
 };
 
+UENUM(BlueprintType)
+enum class ETransferBoneWeightsMethod : uint8
+{
+	/** For every vertex on the TargetMesh, find the closest point on the surface of the SourceMesh and transfer 
+	 * bone weights from it. This is usually a point on the SourceMesh triangle where the bone weights are computed via 
+	 * interpolation of the bone weights at the vertices of the triangle via barycentric coordinates.
+	 */
+	ClosestPointOnSurface = 0
+};
+
+UENUM(BlueprintType)
+enum class EOutputTargetMeshBones : uint8
+{
+	/** When transferring weights, the SourceMesh bone attriubtes will be copied over to the TargetMesh, replacing any 
+	 * existing one, and transferred weights will be indexing the copied bone attributes.
+	 */
+	SourceBones = 0,
+
+    /** When transferring weights, if the TargetMesh has bone attributes, then the transferred SourceMesh weights will be 
+     * reindexed with respect to the TargetMesh bones. In cases where a transferred SourceMesh weight refers to a bone 
+     * not present in the TargetMesh bone attributes, then that weight is simply skipped, and an error message with 
+     * information about the bone will be printed to the Output Log. For best results, the TargetMesh bone attributes 
+     * should be a superset of all the bones that are indexed by the transferred weights.
+     */
+	TargetBones = 1
+};
+
+USTRUCT(BlueprintType)
+struct GEOMETRYSCRIPTINGCORE_API FGeometryScriptTransferBoneWeightsOptions
+{
+	GENERATED_BODY();
+	
+	/** The type of algorithm to use for transferring the bone weights. */
+	UPROPERTY(EditAnywhere, BlueprintReadWrite, Category = Options)
+	ETransferBoneWeightsMethod TransferMethod = ETransferBoneWeightsMethod::ClosestPointOnSurface;
 
+	/** Chooses which bone attributes to use for transferring weights to the TargetMesh. */
+	UPROPERTY(EditAnywhere, BlueprintReadWrite, Category = Options)
+	EOutputTargetMeshBones OutputTargetMeshBones = EOutputTargetMeshBones::SourceBones;
+	
+	/** The identifier for the source bone/skin weight profile used to transfer the weights from. */
+	UPROPERTY(EditAnywhere, BlueprintReadWrite, Category = Options)
+	FGeometryScriptBoneWeightProfile SourceProfile = FGeometryScriptBoneWeightProfile();
 
+	/** The identifier for the source bone/skin weight profile used to transfer the weights to. */
+	UPROPERTY(EditAnywhere, BlueprintReadWrite, Category = Options)
+	FGeometryScriptBoneWeightProfile TargetProfile = FGeometryScriptBoneWeightProfile();
+};
 
 UCLASS(meta = (ScriptName = "GeometryScript_BoneWeights"))
 class GEOMETRYSCRIPTINGCORE_API UGeometryScriptLibrary_MeshBoneWeightFunctions : public UBlueprintFunctionLibrary
@@ -198,6 +244,22 @@ public:
 		FGeometryScriptBoneWeightProfile Profile = FGeometryScriptBoneWeightProfile(),
 		UGeometryScriptDebug* Debug = nullptr);
 	
+	/** 
+	 * Transfer the bone weights from the SourceMesh to the TargetMesh. Assumes that the meshes are aligned. Otherwise, 
+	 * use the TransformMesh geometry script function to align them.
+	 * 
+	 * @param SourceMesh The mesh we are transferring the weights from.
+	 * @param TargetMesh The mesh we are transferring the weights to.
+	 * @param Options The options to set for the transfer weight algorithm.
+	 */
+	UFUNCTION(BlueprintCallable, Category = "GeometryScript|MeshQueries|BoneWeights", meta = (ScriptMethod))
+	static UPARAM(DisplayName = "Target Mesh") UDynamicMesh* 
+	TransferBoneWeightsFromMesh(
+		UDynamicMesh* SourceMesh,
+		UDynamicMesh* TargetMesh,
+		FGeometryScriptTransferBoneWeightsOptions Options = FGeometryScriptTransferBoneWeightsOptions(),
+		UGeometryScriptDebug* Debug = nullptr);
+
 	/**
 	 * Copy the bone attributes (skeleton) from the SourceMesh to the TargetMesh.
 	 * @param SourceMesh Mesh we are copying the bone attributes from.
