Geometry Scripting - update comments for a number of nodes. these were reviewed earlier but I broke them out of the original submission because they might generate merge conflicts.

diff --git a/Engine/Plugins/Experimental/GeometryScripting/Source/GeometryScriptingCore/Public/GeometryScript/MeshRepairFunctions.h b/Engine/Plugins/Experimental/GeometryScripting/Source/GeometryScriptingCore/Public/GeometryScript/MeshRepairFunctions.h
--- a/Engine/Plugins/Experimental/GeometryScripting/Source/GeometryScriptingCore/Public/GeometryScript/MeshRepairFunctions.h
+++ b/Engine/Plugins/Experimental/GeometryScripting/Source/GeometryScriptingCore/Public/GeometryScript/MeshRepairFunctions.h
@@ -165,12 +165,18 @@ class GEOMETRYSCRIPTINGCORE_API UGeometryScriptLibrary_MeshRepairFunctions : pub
 	GENERATED_BODY()
 public:
 
+	/**
+	* Compacts the mesh's vertices and triangles to remove any "holes" in the Vertex ID or Triangle ID lists.
+	*/
 	UFUNCTION(BlueprintCallable, Category = "GeometryScript|Repair", meta=(ScriptMethod))
 	static UPARAM(DisplayName = "Target Mesh") UDynamicMesh* 
 	CompactMesh(  
 		UDynamicMesh* TargetMesh, 
 		UGeometryScriptDebug* Debug = nullptr);
 
+	/**
+	* Attempts to resolve T-Junctions in the mesh by addition of vertices and welding.
+	*/
 	UFUNCTION(BlueprintCallable, Category = "GeometryScript|Repair", meta=(ScriptMethod))
 	static UPARAM(DisplayName = "Target Mesh") UDynamicMesh* 
 	ResolveMeshTJunctions(  
@@ -178,6 +184,9 @@ public:
 		FGeometryScriptResolveTJunctionOptions ResolveOptions,
 		UGeometryScriptDebug* Debug = nullptr);
 
+	/**
+	* Welds any open boundary edges of the mesh together if possible in order to remove "cracks."
+	*/
 	UFUNCTION(BlueprintCallable, Category = "GeometryScript|Repair", meta=(ScriptMethod))
 	static UPARAM(DisplayName = "Target Mesh") UDynamicMesh* 
 	WeldMeshEdges(  
@@ -185,6 +194,12 @@ public:
 		FGeometryScriptWeldEdgesOptions WeldOptions,
 		UGeometryScriptDebug* Debug = nullptr);
 
+	/**
+	* Tries to fill all open boundary loops (such as holes in the geometry surface) of a mesh.
+	* @param FillOptions specifies the method used to fill the holes.
+	* @param NumFilledHoles reports the number of holes filled by the function.
+	* @param NumFailedHolesFills reports the detected holes that were unable to be filled.
+	*/
 	UFUNCTION(BlueprintCallable, Category = "GeometryScript|Repair", meta=(ScriptMethod))
 	static UPARAM(DisplayName = "Target Mesh") UDynamicMesh* 
 	FillAllMeshHoles(  
@@ -194,6 +209,9 @@ public:
 		int32& NumFailedHoleFills,
 		UGeometryScriptDebug* Debug = nullptr);
 
+	/*
+	* Removes connected components of the mesh that have a volume, area, or triangle count below a threshold as specified by the Options.
+	*/
 	UFUNCTION(BlueprintCallable, Category = "GeometryScript|Repair", meta=(ScriptMethod))
 	static UPARAM(DisplayName = "Target Mesh") UDynamicMesh* 
 	RemoveSmallComponents(  
@@ -201,6 +219,10 @@ public:
 		FGeometryScriptRemoveSmallComponentOptions Options,
 		UGeometryScriptDebug* Debug = nullptr);
 
+	/**
+	* Removes any triangles in the mesh that are not visible from the exterior view, under various definitions of "visible" and "outside"
+	* as specified by the Options.
+	*/
 	UFUNCTION(BlueprintCallable, Category = "GeometryScript|Repair", meta=(ScriptMethod))
 	static UPARAM(DisplayName = "Target Mesh") UDynamicMesh* 
 	RemoveHiddenTriangles(  
@@ -208,6 +230,10 @@ public:
 		FGeometryScriptRemoveHiddenTrianglesOptions Options,
 		UGeometryScriptDebug* Debug = nullptr);
 
+	/**
+	* Splits Bowties in the mesh and/or the attributes.  A Bowtie is formed when a single vertex is connected to more than two boundary edges, 
+	* and splitting duplicates the shared vertex so each triangle will have a unique copy.
+	*/
 	UFUNCTION(BlueprintCallable, Category = "GeometryScript|Repair", meta=(ScriptMethod))
 	static UPARAM(DisplayName = "Target Mesh") UDynamicMesh* 
 	SplitMeshBowties(  
@@ -216,6 +242,9 @@ public:
 		bool bAttributeBowties = true,
 		UGeometryScriptDebug* Debug = nullptr);
 
+	/**
+	* Removes triangles that have area or edge length below specified amounts depending on the Options requested.
+	*/
 	UFUNCTION(BlueprintCallable, Category = "GeometryScript|Repair", meta=(ScriptMethod))
 	static UPARAM(DisplayName = "Target Mesh") UDynamicMesh* 
 	RepairMeshDegenerateGeometry(  
diff --git a/Engine/Plugins/Experimental/GeometryScripting/Source/GeometryScriptingCore/Public/GeometryScript/MeshSubdivideFunctions.h b/Engine/Plugins/Experimental/GeometryScripting/Source/GeometryScriptingCore/Public/GeometryScript/MeshSubdivideFunctions.h
--- a/Engine/Plugins/Experimental/GeometryScripting/Source/GeometryScriptingCore/Public/GeometryScript/MeshSubdivideFunctions.h
+++ b/Engine/Plugins/Experimental/GeometryScripting/Source/GeometryScriptingCore/Public/GeometryScript/MeshSubdivideFunctions.h
@@ -53,6 +53,11 @@ class GEOMETRYSCRIPTINGCORE_API UGeometryScriptLibrary_MeshSubdivideFunctions :
 	GENERATED_BODY()
 public:
 
+
+	
+	/**
+	 * Apply PN Tessellation to the Target Mesh as controlled by the Tessellation Level and the Options.
+	 */
 	UFUNCTION(BlueprintCallable, Category = "GeometryScript|Subdivide", meta=(ScriptMethod))
 	static UPARAM(DisplayName = "Target Mesh") UDynamicMesh* 
 	ApplyPNTessellation(
@@ -61,6 +66,9 @@ public:
 		int TessellationLevel = 3,
 		UGeometryScriptDebug* Debug = nullptr );
 	
+	/**
+	* Apply Uniform Tessellation to the Target Mesh as controlled by the Tessellation Level and the Options.
+	*/
 	UFUNCTION(BlueprintCallable, Category = "GeometryScript|Subdivide", meta=(ScriptMethod))
 	static UPARAM(DisplayName = "Target Mesh") UDynamicMesh* 
 	ApplyUniformTessellation(
@@ -68,6 +76,13 @@ public:
 		int TessellationLevel = 3,
 		UGeometryScriptDebug* Debug = nullptr );
 
+	/**
+	* Selectively Tessellate a Selection of the Target Mesh or possibly the entire mesh as controlled by 
+	* the Options.
+	* @param Selection selects the triangles of the mesh to be tessellated.
+	* @param Options controls the behavior of the tessellation if the Selection is empty.
+	* @param TessellationLevel determines the amount of tessellation
+	*/
 	UFUNCTION(BlueprintCallable, Category = "GeometryScript|Subdivide", meta=(ScriptMethod))
 	static UPARAM(DisplayName = "Target Mesh") UDynamicMesh* 
 	ApplySelectiveTessellation(
diff --git a/Engine/Plugins/Experimental/GeometryScripting/Source/GeometryScriptingCore/Public/GeometryScript/TextureMapFunctions.h b/Engine/Plugins/Experimental/GeometryScripting/Source/GeometryScriptingCore/Public/GeometryScript/TextureMapFunctions.h
--- a/Engine/Plugins/Experimental/GeometryScripting/Source/GeometryScriptingCore/Public/GeometryScript/TextureMapFunctions.h
+++ b/Engine/Plugins/Experimental/GeometryScripting/Source/GeometryScriptingCore/Public/GeometryScript/TextureMapFunctions.h
@@ -46,7 +46,7 @@ class GEOMETRYSCRIPTINGCORE_API UGeometryScriptLibrary_TextureMapFunctions : pub
 public:
 
 	/**
-	 * Sample the the given TextureMap at the list of UV positions and return the color at each position in ColorList output
+	 * Samples the given TextureMap at the list of UV positions and returns the color at each position in ColorList output.
 	 */
 	UFUNCTION(BlueprintCallable, Category = "GeometryScript|TextureUtils")
 	static void 
