Fix part of Horde Issue 279297 caused by duplicate LOCTEXT key

diff --git a/Engine/Plugins/Experimental/GeometryScripting/Source/GeometryScriptingEditor/Private/EditorTextureMapFunctions.cpp b/Engine/Plugins/Experimental/GeometryScripting/Source/GeometryScriptingEditor/Private/EditorTextureMapFunctions.cpp
--- a/Engine/Plugins/Experimental/GeometryScripting/Source/GeometryScriptingEditor/Private/EditorTextureMapFunctions.cpp
+++ b/Engine/Plugins/Experimental/GeometryScripting/Source/GeometryScriptingEditor/Private/EditorTextureMapFunctions.cpp
@@ -84,7 +84,7 @@ FGeometryScriptChannelPackResult UGeometryScriptLibrary_EditorTextureMapFunction
 		"ChannelPack: All input textures were null, this node requires at least one non-null input texture");
 	const FText MismatchingResolutionError = LOCTEXT("ChannelPack_MismatchingResolution",
 		"ChannelPack: All non-null source textures must have matching sizes/resolutions");
-	const FText NoValidChannelsForOutputError = LOCTEXT("ChannelPack_CouldNotCopyDataToOutput",
+	const FText NoValidChannelsForOutputError = LOCTEXT("ChannelPack_NoValidChannelsForOutput",
 		"ChannelPack: Deduced an output texture which has no valid channels");
 	const FText CouldNotCreateOutputError = LOCTEXT("ChannelPack_CouldNotCreateOutput",
 		"ChannelPack: Could not create output texture");
