GeometryScript: improve comments on ComputeTangents with note about DynamicMeshComponent tangents mode
#rnx

diff --git a/Engine/Plugins/Experimental/GeometryScripting/Source/GeometryScriptingCore/Public/GeometryScript/MeshNormalsFunctions.h b/Engine/Plugins/Experimental/GeometryScripting/Source/GeometryScriptingCore/Public/GeometryScript/MeshNormalsFunctions.h
--- a/Engine/Plugins/Experimental/GeometryScripting/Source/GeometryScriptingCore/Public/GeometryScript/MeshNormalsFunctions.h
+++ b/Engine/Plugins/Experimental/GeometryScripting/Source/GeometryScriptingCore/Public/GeometryScript/MeshNormalsFunctions.h
@@ -206,6 +206,7 @@ public:
 
 	/**
 	 * Recompute Tangents for the TargetMesh, using the method and settings specified by FGeometryScriptTangentsOptions
+	 * @note If recomputing Tangents for use with a DynamicMeshComponent, it is also necessary to set the Tangents Type on the Component to "Externally Provided"
 	 */
 	UFUNCTION(BlueprintCallable, Category = "GeometryScript|Normals", meta=(ScriptMethod))
 	static UPARAM(DisplayName = "Target Mesh") UDynamicMesh* 
@@ -219,6 +220,7 @@ public:
 	 * Set all tangents in the TargetMesh Tangents Overlays to the specified per-vertex tangents
 	 * @param TangentXList per-vertex tangent vectors. Size must be equal to the MaxVertexID of TargetMesh  (ie non-compact TargetMesh is supported)
 	 * @param TangentYList per-vertex bitangent/binormal vectors. Size must be equal to TangentXList
+	 * @note If setting Tangents for use with a DynamicMeshComponent, it is also necessary to set the Tangents Type on the Component to "Externally Provided"
 	 */
 	UFUNCTION(BlueprintCallable, Category = "GeometryScript|Normals", meta=(ScriptMethod))
 	static UPARAM(DisplayName = "Target Mesh") UDynamicMesh*
