GeometryScript: fix paste-o in AppendDisc that clamped radial samples to 3

diff --git a/Engine/Plugins/Experimental/GeometryScripting/Source/GeometryScriptingCore/Private/MeshPrimitiveFunctions.cpp b/Engine/Plugins/Experimental/GeometryScripting/Source/GeometryScriptingCore/Private/MeshPrimitiveFunctions.cpp
--- a/Engine/Plugins/Experimental/GeometryScripting/Source/GeometryScriptingCore/Private/MeshPrimitiveFunctions.cpp
+++ b/Engine/Plugins/Experimental/GeometryScripting/Source/GeometryScriptingCore/Private/MeshPrimitiveFunctions.cpp
@@ -967,7 +967,7 @@ UDynamicMesh* UGeometryScriptLibrary_MeshPrimitiveFunctions::AppendDisc(
 	UseGenerator->Radius = FMath::Max(FMathf::ZeroTolerance, Radius);
 	UseGenerator->Normal = FVector3f::UnitZ();
 	UseGenerator->AngleSamples = FMath::Max(3, AngleSteps);
-	UseGenerator->RadialSamples = FMath::Max(3, SpokeSteps);
+	UseGenerator->RadialSamples = FMath::Max(1, SpokeSteps);
 	UseGenerator->StartAngle = StartAngle;
 	UseGenerator->EndAngle = EndAngle;
 	UseGenerator->bSinglePolygroup = (PrimitiveOptions.PolygroupMode != EGeometryScriptPrimitivePolygroupMode::PerQuad);
