Geometry Script: Renamed GetBoneIndexForName function.

#rnx

diff --git a/Engine/Plugins/Experimental/GeometryScripting/Source/GeometryScriptingCore/Private/MeshBoneWeightFunctions.cpp b/Engine/Plugins/Experimental/GeometryScripting/Source/GeometryScriptingCore/Private/MeshBoneWeightFunctions.cpp
--- a/Engine/Plugins/Experimental/GeometryScripting/Source/GeometryScriptingCore/Private/MeshBoneWeightFunctions.cpp
+++ b/Engine/Plugins/Experimental/GeometryScripting/Source/GeometryScriptingCore/Private/MeshBoneWeightFunctions.cpp
@@ -494,16 +494,16 @@ UDynamicMesh* UGeometryScriptLibrary_MeshBoneWeightFunctions::DiscardBonesFromMe
 	return TargetMesh;
 }
 
-UDynamicMesh* UGeometryScriptLibrary_MeshBoneWeightFunctions::GetBoneIndexForName(
+UDynamicMesh* UGeometryScriptLibrary_MeshBoneWeightFunctions::GetBoneIndex(
 	UDynamicMesh* TargetMesh,
-	FString BoneName,
+	FName BoneName,
 	bool& bIsValidBoneName,
 	int& BoneIndex,
 	UGeometryScriptDebug* Debug)
 {
 	if (TargetMesh == nullptr)
 	{
-		AppendError(Debug, EGeometryScriptErrorType::InvalidInputs, LOCTEXT("GetBoneIndexForName_InvalidTargetMesh", "GetBoneIndexForName: TargetMesh is Null"));
+		AppendError(Debug, EGeometryScriptErrorType::InvalidInputs, LOCTEXT("GetBoneIndex_InvalidTargetMesh", "GetBoneIndex: TargetMesh is Null"));
 		return TargetMesh;
 	}
 
@@ -511,12 +511,12 @@ UDynamicMesh* UGeometryScriptLibrary_MeshBoneWeightFunctions::GetBoneIndexForNam
 	{	
 		if (!EditMesh.HasAttributes() || !EditMesh.Attributes()->HasBones())
 		{
-			AppendError(Debug, EGeometryScriptErrorType::InvalidInputs, LOCTEXT("GetBoneIndexForName_TargetMeshHasNoBones", "TargetMesh has no bone attributes"));
+			AppendError(Debug, EGeometryScriptErrorType::InvalidInputs, LOCTEXT("GetBoneIndex_TargetMeshHasNoBones", "TargetMesh has no bone attributes"));
 			return;
 		}
 
 		// INDEX_NONE if BoneName doesn't exist in the bone names attribute
-		BoneIndex = EditMesh.Attributes()->GetBoneNames()->GetAttribValues().Find(FName(BoneName));
+		BoneIndex = EditMesh.Attributes()->GetBoneNames()->GetAttribValues().Find(BoneName);
 		
 		bIsValidBoneName = BoneIndex == INDEX_NONE;
 
diff --git a/Engine/Plugins/Experimental/GeometryScripting/Source/GeometryScriptingCore/Public/GeometryScript/MeshBoneWeightFunctions.h b/Engine/Plugins/Experimental/GeometryScripting/Source/GeometryScriptingCore/Public/GeometryScript/MeshBoneWeightFunctions.h
--- a/Engine/Plugins/Experimental/GeometryScripting/Source/GeometryScriptingCore/Public/GeometryScript/MeshBoneWeightFunctions.h
+++ b/Engine/Plugins/Experimental/GeometryScripting/Source/GeometryScriptingCore/Public/GeometryScript/MeshBoneWeightFunctions.h
@@ -291,9 +291,9 @@ public:
 	 */
 	UFUNCTION(BlueprintCallable, Category = "GeometryScript|MeshQueries|BoneWeights", meta=(ScriptMethod))
 	static UPARAM(DisplayName = "Target Mesh") UDynamicMesh* 
-	GetBoneIndexForName(
+	GetBoneIndex(
 		UDynamicMesh* TargetMesh,
-		FString BoneName,
+		FName BoneName,
 		bool& bIsValidBoneName,
 		int& BoneIndex,
 		UGeometryScriptDebug* Debug = nullptr);
