Geometry Script: Changed MeshCreateBoneWeights function type from BlueprintPure to BlueprintCallable.

diff --git a/Engine/Plugins/Experimental/GeometryScripting/Source/GeometryScriptingCore/Public/GeometryScript/MeshBoneWeightFunctions.h b/Engine/Plugins/Experimental/GeometryScripting/Source/GeometryScriptingCore/Public/GeometryScript/MeshBoneWeightFunctions.h
--- a/Engine/Plugins/Experimental/GeometryScripting/Source/GeometryScriptingCore/Public/GeometryScript/MeshBoneWeightFunctions.h
+++ b/Engine/Plugins/Experimental/GeometryScripting/Source/GeometryScriptingCore/Public/GeometryScript/MeshBoneWeightFunctions.h
@@ -147,7 +147,7 @@ public:
 	 * @param bReplaceExistingProfile if true, if the Profile already exists, it is reset
 	 * @param Profile identifier for the bone/skin weight profile
 	 */
-	UFUNCTION(BlueprintPure, Category = "GeometryScript|MeshQueries|BoneWeights", meta=(ScriptMethod))
+	UFUNCTION(BlueprintCallable, Category = "GeometryScript|MeshQueries|BoneWeights", meta=(ScriptMethod))
 	static UPARAM(DisplayName = "Target Mesh") UDynamicMesh* 
 	MeshCreateBoneWeights( 
 		UDynamicMesh* TargetMesh,
