Geometry Scripting - Mesh Geodesics.  Add two functions that compute the shortest distance between mesh points:  GetShortestVertexPath (using Dijkstra ) and GetShortestSurfacePath (using intrinsic mesh edge flips.)

diff --git a/Engine/Plugins/Experimental/GeometryScripting/Source/GeometryScriptingCore/Private/MeshGeodesicFunctions.cpp b/Engine/Plugins/Experimental/GeometryScripting/Source/GeometryScriptingCore/Private/MeshGeodesicFunctions.cpp
new file mode 100644
--- /dev/null
+++ b/Engine/Plugins/Experimental/GeometryScripting/Source/GeometryScriptingCore/Private/MeshGeodesicFunctions.cpp
@@ -0,0 +1,340 @@
+// Copyright Epic Games, Inc. All Rights Reserved.
+
+#include "GeometryScript/MeshGeodesicFunctions.h"
+#include "UDynamicMesh.h"
+
+#include "DynamicMesh/DynamicMesh3.h"
+#include "DynamicMesh/Public/Operations/GeodesicPath.h"
+#include "Math/UnrealMathUtility.h"
+#include "Parameterization/MeshDijkstra.h"
+
+
+
+#include UE_INLINE_GENERATED_CPP_BY_NAME(MeshGeodesicFunctions)
+
+
+using namespace UE::Geometry;
+
+#define LOCTEXT_NAMESPACE "UGeometryScriptLibrary_MeshGeodesicFunctions"
+
+
+
+UDynamicMesh* UGeometryScriptLibrary_MeshGeodesicFunctions::GetShortestVertexPath( UDynamicMesh* TargetMesh,
+                                                                                   int32 StartVID,
+																			       int32 EndVID,
+																			       FGeometryScriptIndexList& VertexIDList,
+																			       bool& bFoundErrors,
+																			       UGeometryScriptDebug* Debug)
+{
+	bFoundErrors = true;
+
+	if (TargetMesh == nullptr)
+	{
+		UE::Geometry::AppendError(Debug, EGeometryScriptErrorType::InvalidInputs, LOCTEXT("GetShortestVertexPath_InvalidInput", "GetShortestVertexPath: TargetMesh is Null"));
+		return TargetMesh;
+	}
+
+	if (StartVID == EndVID)
+	{
+		UE::Geometry::AppendError(Debug, EGeometryScriptErrorType::InvalidInputs, LOCTEXT("GetShortestVertexPath_DuplicateInputs", "GetShortestVertexPath: Start and End vertex are the same"));
+		return TargetMesh;
+	}
+
+	auto MeshHasVertex = [TargetMesh](int32 VID)
+	{ 
+		bool Result = false;
+		TargetMesh->ProcessMesh([VID, &Result](const FDynamicMesh3& Mesh)
+			{
+				Result =  Mesh.IsVertex(VID);
+			}); 
+		return Result;
+	};
+
+
+	const bool bMeshHasStartVID = MeshHasVertex(StartVID);
+	const bool bMeshHasEndVID   = MeshHasVertex(EndVID);
+
+	if (!bMeshHasStartVID)
+	{
+		UE::Geometry::AppendError(Debug, EGeometryScriptErrorType::InvalidInputs, LOCTEXT("GetShortestVertexPath_InvalidStart", "GetShortestVertexPath: Start vertex not part of mesh"));
+		return TargetMesh;
+	}
+
+	if (!bMeshHasEndVID)
+	{
+		UE::Geometry::AppendError(Debug, EGeometryScriptErrorType::InvalidInputs, LOCTEXT("GetShortestVertexPath_InvalidEnd", "GetShortestVertexPath: End vertex not part of mesh"));
+		return TargetMesh;
+	}
+	
+	
+	const TArray<int32> DijkstraVertexPath = [&]
+	{
+		TArray<int32> VertexPath;
+		TargetMesh->ProcessMesh([StartVID, EndVID, &VertexPath](const FDynamicMesh3& ReadMesh)
+		{
+			typedef TMeshDijkstra<FDynamicMesh3>  FMeshDijkstra;
+			FMeshDijkstra MeshDijkstra(&ReadMesh);
+
+			TArray<FMeshDijkstra::FSeedPoint> SeedPoints;
+			FMeshDijkstra::FSeedPoint& Seed = SeedPoints.AddDefaulted_GetRef();
+			Seed.PointID = EndVID;
+
+			// compute graph distances
+			MeshDijkstra.ComputeToTargetPoint(SeedPoints, StartVID);
+
+			// Find path from StartVID to the seed point (EndVID) 
+			MeshDijkstra.FindPathToNearestSeed(StartVID, VertexPath);
+		});
+		return MoveTemp(VertexPath);
+	}();
+
+	if (DijkstraVertexPath.Num() == 0)
+	{
+		UE::Geometry::AppendError(Debug, EGeometryScriptErrorType::InvalidInputs, LOCTEXT("GetShortestVertexPath_Failed", "GetShortestVertexPath: Failed to find connecting path"));
+		return TargetMesh;
+	}
+
+
+	// convert to correct output type
+	{
+		VertexIDList.Reset(EGeometryScriptIndexType::Vertex);
+		VertexIDList.List->Append(DijkstraVertexPath);
+	}
+
+	bFoundErrors = false;
+	return TargetMesh;
+
+}
+ 
+UDynamicMesh* UGeometryScriptLibrary_MeshGeodesicFunctions::GetShortestSurfacePath( UDynamicMesh* TargetMesh,
+																					int32 StartTID,
+																					FVector StartTriCoords,
+																					int32 EndTID,
+																					FVector EndTriCoords,
+																					FGeometryScriptPolyPath& ShortestPath,
+																					bool& bFoundErrors,
+																					UGeometryScriptDebug* Debug)
+{
+	bFoundErrors = true;
+
+	ShortestPath.Reset();
+
+	if (TargetMesh == nullptr)
+	{
+		UE::Geometry::AppendError(Debug, EGeometryScriptErrorType::InvalidInputs, LOCTEXT("GetShortestSurfacePath_InvalidInput", "GetShortestSurfacePath: TargetMesh is Null"));
+		return TargetMesh;
+	}
+
+
+
+	auto MeshHasTriangle = [TargetMesh](const int32& TID)->bool
+	{ 
+		bool Result = false;
+		TargetMesh->ProcessMesh([TID, &Result](const FDynamicMesh3& Mesh)
+			{
+				Result =  Mesh.IsTriangle(TID);
+			});
+		return Result;	
+	};
+
+
+
+	
+	const bool bHasStartTID = MeshHasTriangle(StartTID);
+	const bool bHasEndTID   = MeshHasTriangle(EndTID);
+
+	if (!bHasStartTID)
+	{
+		UE::Geometry::AppendError(Debug, EGeometryScriptErrorType::InvalidInputs, LOCTEXT("GetShortestSurfacePath_InvalidStartTriangle", "GetShortestSurfacePath: Start Point triangle not part of mesh"));
+		return TargetMesh;
+	}
+
+	if (!bHasEndTID)
+	{
+		UE::Geometry::AppendError(Debug, EGeometryScriptErrorType::InvalidInputs, LOCTEXT("GetShortestSurfacePath_InvalidEndTriangle", "GetShortestSurfacePath: End Point triangle not part of mesh"));
+		return TargetMesh;
+	}
+
+	auto MakeValidBaryCentric = [](const FVector& Vec)->FVector
+	{
+		// any negative value indicates invalid BC
+		if (Vec.X < 0 || Vec.Y < 0 || Vec.Z < 0)
+		{
+			return FVector(1./3., 1./3., 1./3.);
+		}
+
+		const double Sum = Vec.X + Vec.Y + Vec.Z;
+		constexpr double Tol = 0.05;
+		if (FMath::Abs(Sum - 1.0)  > Tol) // allow for a small amount of error before we declare invalid BC 
+		{
+			return FVector(1./3., 1./3., 1./3.);
+		}
+		else
+		{
+			return (1./Sum) * Vec;
+		}
+	};
+
+	const FVector StartBC = MakeValidBaryCentric(StartTriCoords);
+	const FVector EndBC   = MakeValidBaryCentric(EndTriCoords);
+
+	
+	// trivial case when both start and end points are on the same triangle.
+	if (StartTID == EndTID)
+	{
+		
+		if (StartBC == EndBC)
+		{
+			UE::Geometry::AppendError(Debug, EGeometryScriptErrorType::InvalidInputs, LOCTEXT("GetShortestSurfacePath_DuplicateInputs", "GetShortestSurfacePath: Start and End Point are the same"));
+			return TargetMesh;
+		}
+
+		auto GetTriBaryPoint = [TargetMesh](const int32 TID, const FVector& BaryCoords)
+		{
+
+			FVector PositionOfBaryPoint = FVector::Zero();
+			TargetMesh->ProcessMesh([TID, &BaryCoords, &PositionOfBaryPoint](const FDynamicMesh3& Mesh)
+				{
+					PositionOfBaryPoint = Mesh.GetTriBaryPoint(TID, BaryCoords.X, BaryCoords.Y, BaryCoords.Z);
+				});
+
+			return PositionOfBaryPoint;
+		};
+
+		const FVector StartPosition = GetTriBaryPoint(StartTID, StartBC);
+		const FVector EndPosition   = GetTriBaryPoint(EndTID, EndBC);
+		
+		ShortestPath.bClosedLoop = false;
+		ShortestPath.Path->Add(StartPosition);
+		ShortestPath.Path->Add(EndPosition);
+
+		bFoundErrors = false;
+		return TargetMesh;
+	}
+
+
+
+	// compute the geodesic using the intrinsic mesh.
+	{
+
+		// Currently the intrinsic mesh geodesic only connects vertices on of the dynamic mesh.
+		// @todo see if this limitation can be removed..
+		FDynamicMesh3 TmpMesh;
+		TargetMesh->ProcessMesh([&TmpMesh](const FDynamicMesh3& SrcMesh){TmpMesh = SrcMesh;});
+
+		int32 StartVID = FDynamicMesh3::InvalidID;
+		int32 EndVID = FDynamicMesh3::InvalidID;
+		
+		// add start and end vertices to the tmp mesh
+		{
+			DynamicMeshInfo::FPokeTriangleInfo PokeInfo;
+			if (TmpMesh.PokeTriangle(StartTID, StartBC, PokeInfo) == EMeshResult::Ok)
+			{
+				  StartVID = PokeInfo.NewVertex;
+			}
+			if (TmpMesh.PokeTriangle(EndTID, EndBC, PokeInfo) == EMeshResult::Ok)
+			{
+				EndVID = PokeInfo.NewVertex;
+			}
+
+			// somehow we were unable to create the start or end point.  perhaps the barycoords weren't valid...
+			if (StartVID == FDynamicMesh3::InvalidID || EndVID == FDynamicMesh3::InvalidID)
+			{
+				UE::Geometry::AppendError(Debug, EGeometryScriptErrorType::InvalidInputs, LOCTEXT("GetShortestSurfacePath_PokeFailure", "GetShortestSurfacePath: failed"));
+				return TargetMesh;
+			}
+		}
+		
+
+		// create the initial path that will be deformed into the shortest path.
+		const TArray<FEdgePath::FDirectedSegment> DirectedSegments = [&TmpMesh, StartVID, EndVID]
+		{
+			// find the initial path by using Dijkstra
+			const TArray<int32> DijkstraVertexPath = [&]
+			{
+				typedef TMeshDijkstra<FDynamicMesh3>  FMeshDijkstra;
+				FMeshDijkstra MeshDijkstra(&TmpMesh);
+
+				TArray<FMeshDijkstra::FSeedPoint> SeedPoints;
+				FMeshDijkstra::FSeedPoint& Seed = SeedPoints.AddDefaulted_GetRef();
+				Seed.PointID = EndVID;
+
+				// compute distances from EndVID, up to StartVID
+				MeshDijkstra.ComputeToTargetPoint(SeedPoints, StartVID);
+
+				// make a vertex path from Seed (EndVID) to StartVID
+				TArray<int32> VertexPath;
+				MeshDijkstra.FindPathToNearestSeed(StartVID, VertexPath);
+				return MoveTemp(VertexPath);
+			}();
+														
+			// convert the initial path into a sequence of directed edges.
+			TArray<FEdgePath::FDirectedSegment> SegmentPath;
+			SegmentPath.Empty();
+			if (DijkstraVertexPath.Num() == 0)
+			{
+				return MoveTemp(SegmentPath);
+			}
+
+			SegmentPath.Reserve(DijkstraVertexPath.Num() - 1);
+
+			for (int32 i = 1; i < DijkstraVertexPath.Num(); ++i)
+			{
+				int32 VID = DijkstraVertexPath[i];
+				int32 PVID = DijkstraVertexPath[i - 1];
+
+				int32 EID = TmpMesh.FindEdge(VID, PVID);
+				if (EID != FDynamicMesh3::InvalidID)
+				{
+					const FIndex2i EdgeV = TmpMesh.GetEdgeV(EID);
+
+					FEdgePath::FDirectedSegment& Segment = SegmentPath.AddDefaulted_GetRef();
+					Segment.EID = EID;
+					Segment.HeadIndex = (EdgeV.B == VID) ? 1 : 0;
+				}
+			}
+
+			return MoveTemp(SegmentPath);
+		}();
+
+		// failed to find any path connecting the start and end point.  Most likely they are not in the same connected component.
+		if ( DirectedSegments.Num() == 0)
+		{
+			UE::Geometry::AppendError(Debug, EGeometryScriptErrorType::InvalidInputs, LOCTEXT("GetShortestSurfacePath_DijkstraFailure", "GetShortestSurfacePath: failed to find a path connecting the points"));
+			return TargetMesh;
+		}
+
+		// populate the deformable path
+		FDeformableEdgePath DeformableEdgePath(TmpMesh, DirectedSegments);
+
+		// minimize the deformable path
+		FDeformableEdgePath::FEdgePathDeformationInfo ResultInfo;
+		DeformableEdgePath.Minimize(ResultInfo);
+
+		// convert minimized deformable path to a PolyPath
+		{
+			TArray<FVector>& GeodesicPath = *ShortestPath.Path;
+			GeodesicPath.Empty();
+
+			constexpr double CoalesceThreshold = 0.01; // essentially a welding threshold for adjacent path points.
+
+			TArray<FDeformableEdgePath::FSurfacePoint> SurfacePathPoints = DeformableEdgePath.AsSurfacePoints(CoalesceThreshold);
+			
+			GeodesicPath.Reserve(SurfacePathPoints.Num());
+			
+			const FDynamicMesh3* SurfaceMesh = DeformableEdgePath.GetIntrinsicMesh().GetNormalCoordinates().SurfaceMesh;
+			for (const FDeformableEdgePath::FSurfacePoint& SurfacePoint : SurfacePathPoints)
+			{
+				bool bValid;
+				const FVector3d Pos = AsR3Position(SurfacePoint, *SurfaceMesh, bValid);
+				GeodesicPath.Add(Pos);
+			}
+		}
+
+		bFoundErrors = false;
+		return TargetMesh;
+	}
+
+}
+
+#undef LOCTEXT_NAMESPACE
\ No newline at end of file
diff --git a/Engine/Plugins/Experimental/GeometryScripting/Source/GeometryScriptingCore/Public/GeometryScript/MeshGeodesicFunctions.h b/Engine/Plugins/Experimental/GeometryScripting/Source/GeometryScriptingCore/Public/GeometryScript/MeshGeodesicFunctions.h
new file mode 100644
--- /dev/null
+++ b/Engine/Plugins/Experimental/GeometryScripting/Source/GeometryScriptingCore/Public/GeometryScript/MeshGeodesicFunctions.h
@@ -0,0 +1,66 @@
+// Copyright Epic Games, Inc. All Rights Reserved.
+
+#pragma once
+
+#include "CoreMinimal.h"
+#include "Kismet/BlueprintFunctionLibrary.h"
+#include "GeometryScript/GeometryScriptTypes.h"
+
+#include "MeshGeodesicFunctions.generated.h"
+
+class UDynamicMesh;
+
+UCLASS(meta = (ScriptName = "GeometryScript_MeshGeodesic"))
+class GEOMETRYSCRIPTINGCORE_API UGeometryScriptLibrary_MeshGeodesicFunctions : public UBlueprintFunctionLibrary
+{
+	GENERATED_BODY()
+public:
+
+
+
+	/**
+	* Computes a vertex list that represents the shortest path constrained to travel along mesh triangle edges between the prescribed start and end vertex.
+	* This can fail if the Start and End points are within separate connected components of the mesh.
+	* @param TargetMesh defines the surface where the path is computed.
+	* @param StartVertexID  indicates ID of mesh vertex that defines the starting point of the path.
+	* @param EndVertexID  indicates ID of the mesh vertex that defined the end point of the path.
+	* @param VertexIDList, if found this will hold on return a list of mesh vertex IDs that define the path from StartVertexID to EndVertexID.
+	* @param bFoundErrors will be false on success.  
+	*/
+	UFUNCTION(BlueprintCallable, Category = "GeometryScript|MeshGeodesic", meta=(ScriptMethod))
+	static UPARAM(DisplayName = "Target Mesh") UDynamicMesh* 
+	GetShortestVertexPath(
+		UDynamicMesh* TargetMesh,
+		int32 StartVertexID,
+		int32 EndVertexID,
+		FGeometryScriptIndexList& VertexIDList,
+		bool& bFoundErrors,
+		UGeometryScriptDebug* Debug = nullptr);
+
+	/**
+	* Computes a PolyPath that represents the shortest mesh surface path between two prescribed points on the provided mesh.
+	* This can fail if the Start and End points are within separate connected components of the mesh.
+	* @param TargetMesh defines the surface where the path is computed.
+	* @param StartTriangleID the ID of mesh Triangle that contains the start point of the path.
+	* @param StartBaryCoords indicates the location of start point within the start triangle, in terms of barycentric coordinates.
+	* @param EndTriangleID the ID of mesh Triangle that contains the end point of the path.
+	* @param EndBaryCoords indicates the location of the end point within the end triangle, in terms of barycentric coordinates.
+	* @param ShortestPath, if found this will hold on return a PolyPath that defines the shortest path along the mesh surface connecting the start and end points.
+	* @param bFoundErrors, will be false on success.  
+	* Note, Barycentric coordinates are of the form (a,b,c) where each entry is positive and a + b + c = 1.  
+	*       If the provided coordinates are invalid, the value (1/3, 1/3, 1/3) will be used.
+	*/
+	UFUNCTION(BlueprintCallable, Category = "GeometryScript|MeshGeodesic", meta=(ScriptMethod))
+	static UPARAM(DisplayName = "Target Mesh") UDynamicMesh* 
+	GetShortestSurfacePath(
+		UDynamicMesh* TargetMesh,
+		int32 StartTriangleID,
+		FVector StartBaryCoords,
+		int32 EndTriangleID,
+		FVector EndBaryCoords,
+		FGeometryScriptPolyPath& ShortestPath,
+		bool& bFoundErrors,
+		UGeometryScriptDebug* Debug = nullptr);
+
+	
+};
\ No newline at end of file
