Remove unnecessary pathing which no longer compiles in 5.3
#rnx

diff --git a/Engine/Plugins/Experimental/GeometryScripting/Source/GeometryScriptingCore/Private/MeshGeodesicFunctions.cpp b/Engine/Plugins/Experimental/GeometryScripting/Source/GeometryScriptingCore/Private/MeshGeodesicFunctions.cpp
--- a/Engine/Plugins/Experimental/GeometryScripting/Source/GeometryScriptingCore/Private/MeshGeodesicFunctions.cpp
+++ b/Engine/Plugins/Experimental/GeometryScripting/Source/GeometryScriptingCore/Private/MeshGeodesicFunctions.cpp
@@ -4,7 +4,7 @@
 #include "UDynamicMesh.h"
 
 #include "DynamicMesh/DynamicMesh3.h"
-#include "DynamicMesh/Public/Operations/GeodesicPath.h"
+#include "Operations/GeodesicPath.h"
 #include "Math/UnrealMathUtility.h"
 #include "Parameterization/MeshDijkstra.h"
 
