add new geometryscript functions for manipulating lists of general polygons, and also for a class for simple 2D polygons (so that we can have arrays of polygons to support general polygon holes)

diff --git a/Engine/Plugins/Experimental/GeometryScripting/Source/GeometryScriptingCore/Private/GeometryScriptTypes.cpp b/Engine/Plugins/Experimental/GeometryScripting/Source/GeometryScriptingCore/Private/GeometryScriptTypes.cpp
--- a/Engine/Plugins/Experimental/GeometryScripting/Source/GeometryScriptingCore/Private/GeometryScriptTypes.cpp
+++ b/Engine/Plugins/Experimental/GeometryScripting/Source/GeometryScriptingCore/Private/GeometryScriptTypes.cpp
@@ -3,6 +3,7 @@
 #include "GeometryScript/GeometryScriptTypes.h"
 #include "GeometryBase.h"
 
+#include "Curve/GeneralPolygon2.h"
 
 
 FGeometryScriptDebugMessage UE::Geometry::MakeScriptError(EGeometryScriptErrorType ErrorTypeIn, const FText& MessageIn)
@@ -59,3 +60,13 @@ void UE::Geometry::AppendWarning(TArray<FGeometryScriptDebugMessage>* DebugMessa
 		DebugMessages->Add(Result);
 	}
 }
+
+
+void FGeometryScriptGeneralPolygonList::Reset()
+{
+	if (!Polygons.IsValid())
+	{
+		Polygons = MakeShared<TArray<UE::Geometry::FGeneralPolygon2d>>();
+	}
+	Polygons->Reset();
+}
\ No newline at end of file
diff --git a/Engine/Plugins/Experimental/GeometryScripting/Source/GeometryScriptingCore/Private/MeshPrimitiveFunctions.cpp b/Engine/Plugins/Experimental/GeometryScripting/Source/GeometryScriptingCore/Private/MeshPrimitiveFunctions.cpp
--- a/Engine/Plugins/Experimental/GeometryScripting/Source/GeometryScriptingCore/Private/MeshPrimitiveFunctions.cpp
+++ b/Engine/Plugins/Experimental/GeometryScripting/Source/GeometryScriptingCore/Private/MeshPrimitiveFunctions.cpp
@@ -1313,5 +1313,45 @@ UDynamicMesh* UGeometryScriptLibrary_MeshPrimitiveFunctions::AppendDelaunayTrian
 	return TargetMesh;
 }
 
+UDynamicMesh* UGeometryScriptLibrary_MeshPrimitiveFunctions::AppendPolygonListTriangulation(
+	UDynamicMesh* TargetMesh,
+	FGeometryScriptPrimitiveOptions PrimitiveOptions,
+	FTransform Transform,
+	FGeometryScriptGeneralPolygonList PolygonList,
+	UGeometryScriptDebug* Debug)
+{
+	if (TargetMesh == nullptr)
+	{
+		UE::Geometry::AppendError(Debug, EGeometryScriptErrorType::InvalidInputs, LOCTEXT("PrimitiveFunctions_AppendPolygonListTriangulation_NullTarget", "AppendDelaunayTriangulation2D: TargetMesh is Null"));
+		return TargetMesh;
+	}
+	if (!PolygonList.Polygons)
+	{
+		UE::Geometry::AppendError(Debug, EGeometryScriptErrorType::InvalidInputs, LOCTEXT("PrimitiveFunctions_AppendPolygonListTriangulation_NullPolygonList", "AppendPolygonListTriangulation: Polygon List is not set"));
+		return TargetMesh;
+	}
+
+	// Triangulate and append each polygon separately
+	bool bAnyFailed = false;
+	for (const UE::Geometry::FGeneralPolygon2d& Polygon : *PolygonList.Polygons)
+	{
+		UE::Geometry::FDelaunay2 Delaunay;
+		FFlatTriangulationMeshGenerator TriangulationMeshGen;
+		TriangulationMeshGen.bReverseOrientation = true; // because FDelaunay2 generates triangles with reversed orientation vs what we want
+		Delaunay.Triangulate(Polygon, &TriangulationMeshGen.Triangles2D, &TriangulationMeshGen.Vertices2D);
+
+		if (TriangulationMeshGen.Vertices2D.Num() > 2 && TriangulationMeshGen.Triangles2D.Num() > 0)
+		{
+			AppendPrimitive(TargetMesh, &TriangulationMeshGen.Generate(), Transform, PrimitiveOptions);
+		}
+	}
+	if (bAnyFailed)
+	{
+		UE::Geometry::AppendWarning(Debug, EGeometryScriptErrorType::OperationFailed, LOCTEXT("PrimitiveFunctions_AppendPolygonListTriangulation_TriangulationFailed", "AppendPolygonListTriangulation: Triangulation failed, but may still contain a partial result."));
+	}
+
+	return TargetMesh;
+}
+
 
 #undef LOCTEXT_NAMESPACE
\ No newline at end of file
diff --git a/Engine/Plugins/Experimental/GeometryScripting/Source/GeometryScriptingCore/Private/PolygonFunctions.cpp b/Engine/Plugins/Experimental/GeometryScripting/Source/GeometryScriptingCore/Private/PolygonFunctions.cpp
new file mode 100644
--- /dev/null
+++ b/Engine/Plugins/Experimental/GeometryScripting/Source/GeometryScriptingCore/Private/PolygonFunctions.cpp
@@ -0,0 +1,483 @@
+// Copyright Epic Games, Inc. All Rights Reserved.
+
+#include "GeometryScript/PolygonFunctions.h"
+
+#include "Curve/GeneralPolygon2.h"
+#include "Curve/PolygonIntersectionUtils.h"
+
+#include UE_INLINE_GENERATED_CPP_BY_NAME(PolygonFunctions)
+
+#define LOCTEXT_NAMESPACE "UGeometryScriptLibrary_PolygonFunctions"
+
+namespace
+{
+	// Helper to support looping the vertex index
+	static int32 PosMod(int32 InVertexIndex, int32 VertexCount)
+	{
+		int32 Index = InVertexIndex % VertexCount;
+		if (Index < 0)
+		{
+			Index += VertexCount;
+		}
+		return Index;
+	}
+
+	// Helper to safely get the polygon at the requested index, or null if no such polygon exists
+	static const UE::Geometry::FGeneralPolygon2d* GetGeneralPolygonHelper(const FGeometryScriptGeneralPolygonList& PolygonList, int PolygonIndex)
+	{
+		if (PolygonList.Polygons.IsValid())
+		{
+			const TArray<UE::Geometry::FGeneralPolygon2d>& Polygons = *(PolygonList.Polygons);
+			if (Polygons.IsValidIndex(PolygonIndex))
+			{
+				return &Polygons[PolygonIndex];
+			}
+		}
+
+		return nullptr;
+	}
+
+	static UE::Geometry::FGeneralPolygon2d* GetGeneralPolygonHelper(FGeometryScriptGeneralPolygonList& PolygonList, int PolygonIndex)
+	{
+		if (PolygonList.Polygons.IsValid())
+		{
+			TArray<UE::Geometry::FGeneralPolygon2d>& Polygons = *(PolygonList.Polygons);
+			if (Polygons.IsValidIndex(PolygonIndex))
+			{
+				return &Polygons[PolygonIndex];
+			}
+		}
+
+		return nullptr;
+	}
+
+	// Helper to safely get the polygon at the requested indices, or null if no such polygon exists
+	// if HoleIndex is -1, returns the outer polygon
+	static const UE::Geometry::FPolygon2d* GetSubPolygonHelper(const FGeometryScriptGeneralPolygonList& PolygonList, int PolygonIndex, int32 HoleIndex)
+	{
+		const UE::Geometry::FGeneralPolygon2d* GeneralPolygon = GetGeneralPolygonHelper(PolygonList, PolygonIndex);
+		if (GeneralPolygon)
+		{
+			if (HoleIndex == INDEX_NONE)
+			{
+				return &GeneralPolygon->GetOuter();
+			}
+			else
+			{
+				const TArray<UE::Geometry::FPolygon2d>& Holes = GeneralPolygon->GetHoles();
+				if (Holes.IsValidIndex(HoleIndex))
+				{
+					return &Holes[HoleIndex];
+				}
+			}
+		}
+
+		return nullptr;
+	}
+}
+
+
+int32 UGeometryScriptLibrary_SimplePolygonFunctions::GetPolygonVertexCount(FGeometryScriptSimplePolygon Polygon)
+{
+	if (Polygon.Vertices.IsValid())
+	{
+		return Polygon.Vertices->Num();
+	}
+	return 0;
+}
+
+FVector2D UGeometryScriptLibrary_SimplePolygonFunctions::GetPolygonVertex(FGeometryScriptSimplePolygon Polygon, int32 VertexIndex, bool& bPolygonIsEmpty)
+{
+	bPolygonIsEmpty = !Polygon.Vertices.IsValid() || Polygon.Vertices->IsEmpty();
+	if (!bPolygonIsEmpty)
+	{
+		return (*Polygon.Vertices)[PosMod(VertexIndex, Polygon.Vertices->Num())];
+	}
+	return FVector2D::ZeroVector;
+}
+
+void UGeometryScriptLibrary_SimplePolygonFunctions::SetPolygonVertex(FGeometryScriptSimplePolygon Polygon, int32 VertexIndex, FVector2D VertexPosition, bool& bPolygonIsEmpty)
+{
+	bPolygonIsEmpty = !Polygon.Vertices.IsValid() || Polygon.Vertices->IsEmpty();
+	if (!bPolygonIsEmpty)
+	{
+		(*Polygon.Vertices)[PosMod(VertexIndex, Polygon.Vertices->Num())] = VertexPosition;
+	}
+}
+
+int32 UGeometryScriptLibrary_SimplePolygonFunctions::AddPolygonVertex(FGeometryScriptSimplePolygon Polygon, FVector2D VertexPosition)
+{
+	if (!Polygon.Vertices.IsValid())
+	{
+		Polygon.Reset();
+	}
+	return Polygon.Vertices->Add(VertexPosition);
+}
+
+FVector2D UGeometryScriptLibrary_SimplePolygonFunctions::GetPolygonTangent(FGeometryScriptSimplePolygon Polygon, int32 VertexIndex, bool& bPolygonIsEmpty)
+{
+	bPolygonIsEmpty = !Polygon.Vertices.IsValid() || Polygon.Vertices->IsEmpty();
+	if (!bPolygonIsEmpty)
+	{
+		VertexIndex = PosMod(VertexIndex, Polygon.Vertices->Num());
+		return UE::Geometry::CurveUtil::Tangent<double, FVector2D>(*Polygon.Vertices, VertexIndex, true);
+	}
+	return FVector2D::ZeroVector;
+}
+
+double UGeometryScriptLibrary_SimplePolygonFunctions::GetPolygonArcLength(FGeometryScriptSimplePolygon Polygon)
+{
+	if (Polygon.Vertices.IsValid() && !Polygon.Vertices->IsEmpty())
+	{
+		UE::Geometry::CurveUtil::ArcLength<double, FVector2D>(*Polygon.Vertices, true);
+	}
+	return 0.0;
+}
+
+double UGeometryScriptLibrary_SimplePolygonFunctions::GetPolygonArea(FGeometryScriptSimplePolygon Polygon)
+{
+	if (Polygon.Vertices.IsValid() && !Polygon.Vertices->IsEmpty())
+	{
+		UE::Geometry::CurveUtil::SignedArea2<double, FVector2D>(*Polygon.Vertices);
+	}
+	return 0.0;
+}
+
+FBox2D UGeometryScriptLibrary_SimplePolygonFunctions::GetPolygonBounds(FGeometryScriptSimplePolygon Polygon)
+{
+	if (Polygon.Vertices.IsValid() && !Polygon.Vertices->IsEmpty())
+	{
+		return FBox2D(*Polygon.Vertices);
+	}
+	return FBox2D(ForceInit);
+}
+
+void UGeometryScriptLibrary_SimplePolygonFunctions::ConvertSplineToPolygon(const USplineComponent* Spline, FGeometryScriptSimplePolygon& Polygon, FGeometryScriptSplineSamplingOptions SamplingOptions, EGeometryScriptAxis DropAxis)
+{
+	Polygon.Reset();
+	if (Spline)
+	{
+		int32 Keep0 = int32(DropAxis == EGeometryScriptAxis::X);
+		int32 Keep1 = 1 + int32(DropAxis != EGeometryScriptAxis::Z);
+		auto ProjDropAxis = [Keep0, Keep1](const FVector& V3D) -> FVector2D
+		{
+			return FVector2D(V3D[Keep0], V3D[Keep1]);
+		};
+
+		bool bIsLoop = Spline->IsClosedLoop();
+		if (SamplingOptions.SampleSpacing == EGeometryScriptSampleSpacing::ErrorTolerance)
+		{
+			TArray<FVector> Path;
+			float SquaredErrorTolerance = FMath::Max(KINDA_SMALL_NUMBER, SamplingOptions.ErrorTolerance * SamplingOptions.ErrorTolerance);
+			Spline->ConvertSplineToPolyLine(SamplingOptions.CoordinateSpace, SquaredErrorTolerance, Path);
+			if (bIsLoop)
+			{
+				Path.Pop(); // delete the duplicate end-point for loops
+			}
+
+			// Copy to FVector2D array
+			Polygon.Vertices->SetNum(Path.Num());
+			for (int32 Idx = 0; Idx < Path.Num(); ++Idx)
+			{
+				(*Polygon.Vertices)[Idx] = ProjDropAxis(Path[Idx]);
+			}
+		}
+		else
+		{
+			float Duration = Spline->Duration;
+
+			bool bUseConstantVelocity = SamplingOptions.SampleSpacing == EGeometryScriptSampleSpacing::UniformDistance;
+			int32 UseSamples = FMath::Max(2, SamplingOptions.NumSamples); // Always use at least 2 samples
+			// In non-loops, we adjust DivNum so we exactly sample the end of the spline
+			// In loops we don't sample the endpoint, by convention, as it's the same as the start
+			float DivNum = float(UseSamples - (int32)!bIsLoop);
+			Polygon.Vertices->Reserve(UseSamples);
+			for (int32 Idx = 0; Idx < UseSamples; Idx++)
+			{
+				float Time = Duration * ((float)Idx / DivNum);
+				Polygon.Vertices->Add(ProjDropAxis(Spline->GetLocationAtTime(Time, SamplingOptions.CoordinateSpace, bUseConstantVelocity)));
+			}
+		}
+
+	}
+}
+
+TArray<FVector> UGeometryScriptLibrary_SimplePolygonFunctions::Conv_GeometryScriptSimplePolygonToArray(FGeometryScriptSimplePolygon Polygon)
+{
+	TArray<FVector> Points3d;
+	if (Polygon.Vertices.IsValid())
+	{
+		Points3d.Reserve(Polygon.Vertices->Num());
+		for (const FVector2D& V : *Polygon.Vertices)
+		{
+			Points3d.Emplace(V.X, V.Y, 0.0);
+		}
+	}
+	return Points3d;
+}
+
+TArray<FVector2D> UGeometryScriptLibrary_SimplePolygonFunctions::Conv_GeometryScriptSimplePolygonToArrayOfVector2D(FGeometryScriptSimplePolygon Polygon)
+{
+	if (Polygon.Vertices.IsValid())
+	{
+		return *Polygon.Vertices;
+	}
+	return TArray<FVector2D>();
+}
+
+FGeometryScriptSimplePolygon UGeometryScriptLibrary_SimplePolygonFunctions::Conv_ArrayToGeometryScriptSimplePolygon(const TArray<FVector>& PathVertices)
+{
+	FGeometryScriptSimplePolygon Polygon;
+	Polygon.Reset();
+	Polygon.Vertices->Reserve(PathVertices.Num());
+	for (FVector V : PathVertices)
+	{
+		Polygon.Vertices->Emplace(V.X, V.Y);
+	}
+	return Polygon;
+}
+
+FGeometryScriptSimplePolygon UGeometryScriptLibrary_SimplePolygonFunctions::Conv_ArrayOfVector2DToGeometryScriptSimplePolygon(const TArray<FVector2D>& PathVertices)
+{
+	FGeometryScriptSimplePolygon Polygon;
+	Polygon.Reset();
+	*Polygon.Vertices = PathVertices;
+	return Polygon;
+}
+
+int UGeometryScriptLibrary_PolygonListFunctions::GetPolygonCount(FGeometryScriptGeneralPolygonList PolygonList)
+{
+	return PolygonList.Polygons.IsValid() ? PolygonList.Polygons->Num() : 0;
+}
+
+int UGeometryScriptLibrary_PolygonListFunctions::GetPolygonVertexCount(FGeometryScriptGeneralPolygonList PolygonList, bool& bValidIndices, int PolygonIndex, int HoleIndex)
+{
+	if (const UE::Geometry::FPolygon2d* Polygon = GetSubPolygonHelper(PolygonList, PolygonIndex, HoleIndex))
+	{
+		bValidIndices = true;
+		return Polygon->VertexCount();
+	}
+	bValidIndices = false;
+	return 0;
+}
+
+int UGeometryScriptLibrary_PolygonListFunctions::GetPolygonHoleCount(FGeometryScriptGeneralPolygonList PolygonList, bool& bValidIndex, int PolygonIndex)
+{
+	if (const UE::Geometry::FGeneralPolygon2d* Polygon = GetGeneralPolygonHelper(PolygonList, PolygonIndex))
+	{
+		bValidIndex = true;
+		return Polygon->GetHoles().Num();
+	}
+	bValidIndex = false;
+	return 0;
+}
+
+void UGeometryScriptLibrary_PolygonListFunctions::GetPolygonVertices(FGeometryScriptGeneralPolygonList PolygonList, TArray<FVector2D>& OutVertices, bool& bValidIndices, int PolygonIndex, int HoleIndex)
+{
+	if (const UE::Geometry::FPolygon2d* Polygon = GetSubPolygonHelper(PolygonList, PolygonIndex, HoleIndex))
+	{
+		bValidIndices = true;
+		OutVertices = Polygon->GetVertices();
+	}
+	bValidIndices = false;
+}
+
+FGeometryScriptSimplePolygon UGeometryScriptLibrary_PolygonListFunctions::GetSimplePolygon(FGeometryScriptGeneralPolygonList PolygonList, bool& bValidIndices, int32 PolygonIndex, int32 HoleIndex)
+{
+	FGeometryScriptSimplePolygon ToRet;
+	ToRet.Reset();
+	if (const UE::Geometry::FPolygon2d* Polygon = GetSubPolygonHelper(PolygonList, PolygonIndex, HoleIndex))
+	{
+		bValidIndices = true;
+		*ToRet.Vertices = Polygon->GetVertices();
+	}
+	bValidIndices = false;
+	return ToRet;
+}
+
+FVector2D UGeometryScriptLibrary_PolygonListFunctions::GetPolygonVertex(FGeometryScriptGeneralPolygonList PolygonList, bool& bIsValidVertex, int32 VertexIndex, int PolygonIndex, int HoleIndex)
+{
+	if (const UE::Geometry::FPolygon2d* Polygon = GetSubPolygonHelper(PolygonList, PolygonIndex, HoleIndex))
+	{
+		const TArray<FVector2D>& Vertices = Polygon->GetVertices();
+		if (!Vertices.IsEmpty())
+		{
+			bIsValidVertex = true;
+			return Vertices[PosMod(VertexIndex, Vertices.Num())];
+		}
+	}
+	bIsValidVertex = false;
+	UE_LOG(LogGeometry, Warning, TEXT("GetPolygonVertex: No vertex found in Polygon List at Polygon Index: %d, Hole Index: %d, Vertex Index: %d"), PolygonIndex, HoleIndex, VertexIndex);
+	return FVector2D(0, 0);
+}
+
+double UGeometryScriptLibrary_PolygonListFunctions::GetPolygonArea(FGeometryScriptGeneralPolygonList PolygonList, bool& bValidIndex, int PolygonIndex)
+{
+	if (const UE::Geometry::FGeneralPolygon2d* Polygon = GetGeneralPolygonHelper(PolygonList, PolygonIndex))
+	{
+		bValidIndex = true;
+		return Polygon->SignedArea();
+	}
+	bValidIndex = false;
+	return 0;
+}
+
+double UGeometryScriptLibrary_PolygonListFunctions::GetPolygonListArea(FGeometryScriptGeneralPolygonList PolygonList)
+{
+	double AreaSum = 0.0;
+	if (PolygonList.Polygons.IsValid())
+	{
+		for (const UE::Geometry::FGeneralPolygon2d& Polygon : *PolygonList.Polygons)
+		{
+			AreaSum += Polygon.SignedArea();
+		}
+	}
+	return AreaSum;
+}
+
+FBox2D UGeometryScriptLibrary_PolygonListFunctions::GetPolygonBounds(FGeometryScriptGeneralPolygonList PolygonList, bool& bValidIndex, int PolygonIndex)
+{
+	if (const UE::Geometry::FGeneralPolygon2d* Polygon = GetGeneralPolygonHelper(PolygonList, PolygonIndex))
+	{
+		bValidIndex = true;
+		return (FBox2D)Polygon->Bounds();
+	}
+	bValidIndex = false;
+	return FBox2D(ForceInit);
+}
+
+FBox2D UGeometryScriptLibrary_PolygonListFunctions::GetPolygonListBounds(FGeometryScriptGeneralPolygonList PolygonList)
+{
+	UE::Geometry::FAxisAlignedBox2d Bounds;
+	if (PolygonList.Polygons.IsValid())
+	{
+		for (const UE::Geometry::FGeneralPolygon2d& Polygon : *PolygonList.Polygons)
+		{
+			Bounds.Contain(Polygon.Bounds());
+		}
+	}
+	return FBox2D(Bounds);
+}
+
+FGeometryScriptGeneralPolygonList UGeometryScriptLibrary_PolygonListFunctions::CreatePolygonListFromSinglePolygon(FGeometryScriptSimplePolygon OuterShape, const TArray<FGeometryScriptSimplePolygon>& Holes, bool bFixHoleOrientations)
+{
+	FGeometryScriptGeneralPolygonList PolygonList;
+	PolygonList.Reset();
+	AddPolygonToList(PolygonList, OuterShape, Holes, bFixHoleOrientations);
+	return PolygonList;
+}
+
+int32 UGeometryScriptLibrary_PolygonListFunctions::AddPolygonToList(FGeometryScriptGeneralPolygonList PolygonList, FGeometryScriptSimplePolygon OuterShape, const TArray<FGeometryScriptSimplePolygon>& Holes, bool bFixHoleOrientations)
+{
+	if (!PolygonList.Polygons.IsValid())
+	{
+		PolygonList.Polygons.Reset();
+	}
+
+	TArray<FVector2D>* UseOuterVertices = OuterShape.Vertices.Get();
+	TArray<FVector2D> EmptyVertices;
+	if (!OuterShape.Vertices.IsValid())
+	{
+		// Fall back to an empty outer shape
+		UseOuterVertices = &EmptyVertices;
+	}
+	int32 NewPolygonIndex = PolygonList.Polygons->Num();
+	UE::Geometry::FGeneralPolygon2d& AddedPolygon = PolygonList.Polygons->Emplace_GetRef(*UseOuterVertices);
+	for (const FGeometryScriptSimplePolygon& Hole : Holes)
+	{
+		if (!Hole.Vertices)
+		{
+			continue;
+		}
+		UE::Geometry::FPolygon2d HolePolygon(*Hole.Vertices);
+		if (bFixHoleOrientations && HolePolygon.IsClockwise() == AddedPolygon.OuterIsClockwise())
+		{
+			HolePolygon.Reverse();
+		}
+		AddedPolygon.AddHole(HolePolygon, false, false);
+	}
+	return NewPolygonIndex;
+}
+
+FGeometryScriptGeneralPolygonList UGeometryScriptLibrary_PolygonListFunctions::CreatePolygonListFromSimplePolygons(const TArray<FGeometryScriptSimplePolygon>& OuterPolygons)
+{
+	FGeometryScriptGeneralPolygonList PolygonList;
+	PolygonList.Reset();
+	for (const FGeometryScriptSimplePolygon& Polygon : OuterPolygons)
+	{
+		if (Polygon.Vertices.IsValid())
+		{
+			PolygonList.Polygons->Emplace(*Polygon.Vertices.Get());
+		}
+	}
+	return PolygonList;
+}
+
+void UGeometryScriptLibrary_PolygonListFunctions::AppendPolygonList(FGeometryScriptGeneralPolygonList PolygonList, FGeometryScriptGeneralPolygonList PolygonsToAppend)
+{
+	if (!PolygonList.Polygons)
+	{
+		PolygonList.Polygons.Reset();
+	}
+	if (!PolygonsToAppend.Polygons)
+	{
+		return;
+	}
+	PolygonList.Polygons->Append(*PolygonsToAppend.Polygons);
+}
+
+FGeometryScriptGeneralPolygonList UGeometryScriptLibrary_PolygonListFunctions::PolygonsUnion(FGeometryScriptGeneralPolygonList PolygonList, bool bCopyInputOnFailure)
+{
+	FGeometryScriptGeneralPolygonList PolygonListResult;
+	PolygonListResult.Reset();
+	if (!PolygonList.Polygons)
+	{
+		return PolygonListResult;
+	}
+	UE::Geometry::PolygonsUnion(*PolygonList.Polygons, *PolygonListResult.Polygons, bCopyInputOnFailure);
+	return PolygonListResult;
+}
+
+FGeometryScriptGeneralPolygonList UGeometryScriptLibrary_PolygonListFunctions::PolygonsDifference(FGeometryScriptGeneralPolygonList PolygonList, FGeometryScriptGeneralPolygonList PolygonsToSubtract)
+{
+	FGeometryScriptGeneralPolygonList PolygonListResult;
+	PolygonListResult.Reset();
+	if (!PolygonList.Polygons)
+	{
+		return PolygonListResult;
+	}
+	if (!PolygonsToSubtract.Polygons)
+	{
+		// Subtracting nothing
+		*PolygonListResult.Polygons = *PolygonList.Polygons;
+	}
+	UE::Geometry::PolygonsDifference(*PolygonList.Polygons, *PolygonsToSubtract.Polygons, *PolygonListResult.Polygons);
+	return PolygonListResult;
+}
+
+FGeometryScriptGeneralPolygonList UGeometryScriptLibrary_PolygonListFunctions::PolygonsIntersection(FGeometryScriptGeneralPolygonList PolygonList, FGeometryScriptGeneralPolygonList PolygonsToSubtract)
+{
+	FGeometryScriptGeneralPolygonList PolygonListResult;
+	PolygonListResult.Reset();
+	if (!PolygonList.Polygons || !PolygonsToSubtract.Polygons)
+	{
+		return PolygonListResult;
+	}
+	UE::Geometry::PolygonsIntersection(*PolygonList.Polygons, *PolygonsToSubtract.Polygons, *PolygonListResult.Polygons);
+	return PolygonListResult;
+}
+
+FGeometryScriptGeneralPolygonList UGeometryScriptLibrary_PolygonListFunctions::PolygonsExclusiveOr(FGeometryScriptGeneralPolygonList PolygonList, FGeometryScriptGeneralPolygonList PolygonsToSubtract)
+{
+	FGeometryScriptGeneralPolygonList PolygonListResult;
+	PolygonListResult.Reset();
+	if (!PolygonList.Polygons || !PolygonsToSubtract.Polygons)
+	{
+		return PolygonListResult;
+	}
+	UE::Geometry::PolygonsExclusiveOr(*PolygonList.Polygons, *PolygonsToSubtract.Polygons, *PolygonListResult.Polygons);
+	return PolygonListResult;
+}
+
+
+#undef LOCTEXT_NAMESPACE
diff --git a/Engine/Plugins/Experimental/GeometryScripting/Source/GeometryScriptingCore/Public/GeometryScript/GeometryScriptTypes.h b/Engine/Plugins/Experimental/GeometryScripting/Source/GeometryScriptingCore/Public/GeometryScript/GeometryScriptTypes.h
--- a/Engine/Plugins/Experimental/GeometryScripting/Source/GeometryScriptingCore/Public/GeometryScript/GeometryScriptTypes.h
+++ b/Engine/Plugins/Experimental/GeometryScripting/Source/GeometryScriptingCore/Public/GeometryScript/GeometryScriptTypes.h
@@ -10,6 +10,8 @@ PREDECLARE_GEOMETRY(class FDynamicMesh3);
 PREDECLARE_GEOMETRY(template<typename MeshType> class TMeshAABBTree3);
 PREDECLARE_GEOMETRY(template<typename MeshType> class TFastWindingTree);
 PREDECLARE_GEOMETRY(typedef TMeshAABBTree3<FDynamicMesh3> FDynamicMeshAABBTree3);
+PREDECLARE_GEOMETRY(template<typename RealType> class TGeneralPolygon2);
+PREDECLARE_GEOMETRY(typedef TGeneralPolygon2<double> FGeneralPolygon2d);
 
 
 UENUM(BlueprintType)
@@ -482,6 +484,64 @@ struct TStructOpsTypeTraits<FGeometryScriptPolyPath> : public TStructOpsTypeTrai
 	};
 };
 
+// A simple 2D Polygon with no holes
+USTRUCT(BlueprintType, meta = (DisplayName = "Simple Polygon"))
+struct GEOMETRYSCRIPTINGCORE_API FGeometryScriptSimplePolygon
+{
+	GENERATED_BODY()
+public:
+	TSharedPtr<TArray<FVector2D>> Vertices;
+
+	void Reset()
+	{
+		if (Vertices.IsValid() == false)
+		{
+			Vertices = MakeShared<TArray<FVector2D>>();
+		}
+		Vertices->Reset();
+	}
+
+	// Required by TStructOpsTypeTraits interface
+	bool operator==(const FGeometryScriptSimplePolygon& Other) const
+	{
+		return Vertices.Get() == Other.Vertices.Get();
+	}
+	bool operator!=(const FGeometryScriptSimplePolygon& Other) const
+	{
+		return Vertices.Get() != Other.Vertices.Get();
+	}
+};
+
+// A list of general polygons, which may have holes.
+USTRUCT(BlueprintType, meta = (DisplayName = "PolygonList"))
+struct GEOMETRYSCRIPTINGCORE_API FGeometryScriptGeneralPolygonList
+{
+	GENERATED_BODY()
+public:
+	TSharedPtr<TArray<UE::Geometry::FGeneralPolygon2d>> Polygons;
+
+	void Reset();
+
+	// Required by TStructOpsTypeTraits interface
+	bool operator==(const FGeometryScriptGeneralPolygonList& Other) const
+	{
+		return Polygons.Get() == Other.Polygons.Get();
+	}
+	bool operator!=(const FGeometryScriptGeneralPolygonList& Other) const
+	{
+		return !(*this == Other);
+	}
+};
+
+template<>
+struct TStructOpsTypeTraits<FGeometryScriptGeneralPolygonList> : public TStructOpsTypeTraitsBase2<FGeometryScriptGeneralPolygonList>
+{
+	enum
+	{
+		WithIdenticalViaEquality = true,
+	};
+};
+
 
 
 
diff --git a/Engine/Plugins/Experimental/GeometryScripting/Source/GeometryScriptingCore/Public/GeometryScript/MeshPrimitiveFunctions.h b/Engine/Plugins/Experimental/GeometryScripting/Source/GeometryScriptingCore/Public/GeometryScript/MeshPrimitiveFunctions.h
--- a/Engine/Plugins/Experimental/GeometryScripting/Source/GeometryScriptingCore/Public/GeometryScript/MeshPrimitiveFunctions.h
+++ b/Engine/Plugins/Experimental/GeometryScripting/Source/GeometryScriptingCore/Public/GeometryScript/MeshPrimitiveFunctions.h
@@ -528,6 +528,22 @@ public:
 		UGeometryScriptDebug* Debug = nullptr);
 
 
+	/**
+	* Generates a Delaunay Triangulation of the provided Polygon List, and appends it to the Target Mesh.
+	*/
+	UFUNCTION(BlueprintCallable, Category = "GeometryScript|Primitives")
+	static UPARAM(DisplayName = "Target Mesh") UDynamicMesh*
+	AppendPolygonListTriangulation(
+		UDynamicMesh* TargetMesh,
+		FGeometryScriptPrimitiveOptions PrimitiveOptions,
+		FTransform Transform,
+		FGeometryScriptGeneralPolygonList PolygonList,
+		UGeometryScriptDebug* Debug = nullptr);
+
+
+
+
+
 
 	//---------------------------------------------
 	// Backwards-Compatibility implementations
diff --git a/Engine/Plugins/Experimental/GeometryScripting/Source/GeometryScriptingCore/Public/GeometryScript/PolygonFunctions.h b/Engine/Plugins/Experimental/GeometryScripting/Source/GeometryScriptingCore/Public/GeometryScript/PolygonFunctions.h
new file mode 100644
--- /dev/null
+++ b/Engine/Plugins/Experimental/GeometryScripting/Source/GeometryScriptingCore/Public/GeometryScript/PolygonFunctions.h
@@ -0,0 +1,213 @@
+// Copyright Epic Games, Inc. All Rights Reserved.
+
+#pragma once
+
+#include "CoreMinimal.h"
+#include "Kismet/BlueprintFunctionLibrary.h"
+#include "GeometryScript/GeometryScriptTypes.h"
+#include "Components/SplineComponent.h"
+#include "PolyPathFunctions.h" // For spline sampling options
+#include "PolygonFunctions.generated.h"
+
+
+
+UCLASS(meta = (ScriptName = "GeometryScript_SimplePolygon"))
+class GEOMETRYSCRIPTINGCORE_API UGeometryScriptLibrary_SimplePolygonFunctions : public UBlueprintFunctionLibrary
+{
+	GENERATED_BODY()
+public:
+
+	/**
+	 * Returns the number of vertices in a Simple Polygon
+	 */
+	UFUNCTION(BlueprintCallable, Category = "GeometryScript|SimplePolygon", meta = (ScriptMethod))
+	static int32 GetPolygonVertexCount(FGeometryScriptSimplePolygon Polygon);
+
+	/**
+	 * Returns the specified vertex of a Simple Polygon. VertexIndex loops around, so e.g., -1 gives the last vertex in the polygon.
+	 * If Polygon has no vertices, returns the zero vector.
+	 */
+	UFUNCTION(BlueprintCallable, Category = "GeometryScript|SimplePolygon", meta = (ScriptMethod))
+	static FVector2D GetPolygonVertex(FGeometryScriptSimplePolygon Polygon, int32 VertexIndex, bool& bPolygonIsEmpty);
+
+	/**
+	 * Set the specified vertex of a Simple Polygon. VertexIndex loops around, so e.g., -1 gives the last vertex in the polygon.
+	 * Does nothing if Polygon has no vertices.
+	 */
+	UFUNCTION(BlueprintCallable, Category = "GeometryScript|SimplePolygon", meta = (ScriptMethod))
+	static void SetPolygonVertex(UPARAM(ref) FGeometryScriptSimplePolygon Polygon, int32 VertexIndex, FVector2D Position, bool& bPolygonIsEmpty);
+
+	/**
+	 * Set the specified vertex of a Simple Polygon. Returns the index of the added vertex.
+	 */
+	UFUNCTION(BlueprintCallable, Category = "GeometryScript|SimplePolygon", meta = (ScriptMethod))
+	static UPARAM(DisplayName = "Vertex Index") int32 AddPolygonVertex(UPARAM(ref) FGeometryScriptSimplePolygon Polygon, FVector2D Position);
+
+	/**
+	 * Returns a vertex's tangent of a Simple Polygon. VertexIndex loops around, so e.g., -1 gives the tangent of the last vertex in the polygon.
+	 * If Polygon has no vertices, returns the zero vector.
+	 */
+	UFUNCTION(BlueprintCallable, Category = "GeometryScript|SimplePolygon", meta = (ScriptMethod))
+	static FVector2D GetPolygonTangent(FGeometryScriptSimplePolygon Polygon, int32 VertexIndex, bool& bPolygonIsEmpty);
+
+	/**
+	 * Returns the arc length of a Simple Polygon
+	 */
+	UFUNCTION(BlueprintCallable, Category = "GeometryScript|SimplePolygon", meta = (ScriptMethod))
+	static double GetPolygonArcLength(FGeometryScriptSimplePolygon Polygon);
+
+	/**
+	 * Returns the area enclosed by a Simple Polygon
+	 */
+	UFUNCTION(BlueprintCallable, Category = "GeometryScript|SimplePolygon", meta = (ScriptMethod))
+	static double GetPolygonArea(FGeometryScriptSimplePolygon Polygon);
+
+	/**
+	 * Returns the bounding box of a Simple Polygon
+	 */
+	UFUNCTION(BlueprintCallable, Category = "GeometryScript|SimplePolygon", meta = (ScriptMethod))
+	static FBox2D GetPolygonBounds(FGeometryScriptSimplePolygon Polygon);
+
+	/**
+	 * Sample positions from a USplineComponent into a Simple Polyon, based on the given SamplingOptions
+	 */
+	UFUNCTION(BlueprintCallable, Category = "GeometryScript|SimplePolygon")
+	static void ConvertSplineToPolygon(const USplineComponent* Spline, FGeometryScriptSimplePolygon& Polygon, FGeometryScriptSplineSamplingOptions SamplingOptions, EGeometryScriptAxis DropAxis = EGeometryScriptAxis::Z);
+
+	/**
+	 * Returns an array of 3D vectors with the Polygon vertex locations, with Z coordinate set to zero.
+	 */
+	UFUNCTION(BlueprintPure, meta = (DisplayName = "Simple Polygon To Array Of Vector", CompactNodeTitle = "->", BlueprintAutocast), Category = "GeometryScript|SimplePolygon")
+	static TArray<FVector> Conv_GeometryScriptSimplePolygonToArray(FGeometryScriptSimplePolygon Polygon);
+
+	/**
+	 * Returns an array of 2D Vectors with the Polygon vertex locations.
+	 */
+	UFUNCTION(BlueprintPure, meta = (DisplayName = "Simple Polygon To Array Of Vector2D", CompactNodeTitle = "->", BlueprintAutocast), Category = "GeometryScript|SimplePolygon")
+	static TArray<FVector2D> Conv_GeometryScriptSimplePolygonToArrayOfVector2D(FGeometryScriptSimplePolygon Polygon);
+
+	/** 
+	 * Returns a Polygon created from an array of 3D position vectors, ignoring the Z coordinate.
+	 */
+	UFUNCTION(BlueprintPure, meta = (DisplayName = "Array Of Vector To Simple Polygon", CompactNodeTitle = "->", BlueprintAutocast), Category = "GeometryScript|SimplePolygon")
+	static FGeometryScriptSimplePolygon Conv_ArrayToGeometryScriptSimplePolygon(const TArray<FVector>& PathVertices);
+
+	/**
+	 * Returns a Polygon created from an array of 2D position vectors.
+	 */
+	UFUNCTION(BlueprintPure, meta = (DisplayName = "Array Of Vector2D To Simple Polygon", CompactNodeTitle = "->", BlueprintAutocast), Category = "GeometryScript|SimplePolygon")
+	static FGeometryScriptSimplePolygon Conv_ArrayOfVector2DToGeometryScriptSimplePolygon(const TArray<FVector2D>& PathVertices);
+};
+
+
+UCLASS(meta = (ScriptName = "GeometryScript_PolygonList"))
+class GEOMETRYSCRIPTINGCORE_API UGeometryScriptLibrary_PolygonListFunctions : public UBlueprintFunctionLibrary
+{
+	GENERATED_BODY()
+public:
+
+	/**
+	 * Returns the number of polygons in the Polygon List
+	 */
+	UFUNCTION(BlueprintCallable, Category = "GeometryScript|PolygonList", meta = (ScriptMethod))
+	static int32 GetPolygonCount(FGeometryScriptGeneralPolygonList PolygonList);
+
+	/**
+	 * Returns the number of holes in a Polygon. Returns zero for an invalid PolygonIndex.
+	 */
+	UFUNCTION(BlueprintCallable, Category = "GeometryScript|PolygonList", meta = (ScriptMethod))
+	static int32 GetPolygonHoleCount(FGeometryScriptGeneralPolygonList PolygonList, bool& bValidIndex, int32 PolygonIndex);
+
+	/**
+	 * Returns the number of vertices in a Polygon's outer shape, if HoleIndex is -1, or in the specified inner hole.
+	 * Returns 0 for invalid Polygon or Hole indices.
+	 */
+	UFUNCTION(BlueprintCallable, Category = "GeometryScript|PolygonList", meta=(ScriptMethod))
+	static int32 GetPolygonVertexCount(FGeometryScriptGeneralPolygonList PolygonList, bool& bValidIndices, int32 PolygonIndex, int32 HoleIndex = -1);
+
+	/**
+	 * Returns the vertices of a Polygon -- either of the outer polygon, if HoleIndex is -1, or specified inner hole.
+	 * OutVertices will be empty for invalid Polygon or Hole indices.
+	 */
+	UFUNCTION(BlueprintCallable, Category = "GeometryScript|PolygonList", meta = (ScriptMethod))
+	static void GetPolygonVertices(FGeometryScriptGeneralPolygonList PolygonList, TArray<FVector2D>& OutVertices, bool& bValidIndices, int32 PolygonIndex, int32 HoleIndex = -1);
+
+	/**
+	 * Returns a specified Simple Polygon from a Polygon List -- either the outer polygon, if HoleIndex is -1, or specified inner hole.
+	 * Polygon will be empty for invalid Polygon or Hole indices.
+	 */
+	UFUNCTION(BlueprintCallable, Category = "GeometryScript|PolygonList", meta = (ScriptMethod))
+	static FGeometryScriptSimplePolygon GetSimplePolygon(FGeometryScriptGeneralPolygonList PolygonList, bool& bValidIndices, int32 PolygonIndex, int32 HoleIndex = -1);
+
+	/**
+	 * Returns the specified vertex of a Polygon -- either of the outer polygon, if HoleIndex is -1, or specified inner hole.
+	 * Vertex will be the zero vector for invalid Polygon or Hole indices, or if the polygon is empty. VertexIndex will loop.
+	 */
+	UFUNCTION(BlueprintCallable, Category = "GeometryScript|PolygonList", meta = (ScriptMethod))
+	static FVector2D GetPolygonVertex(FGeometryScriptGeneralPolygonList PolygonList, bool& bIsValidVertex, int32 VertexIndex, int32 PolygonIndex, int32 HoleIndex = -1);
+
+	/**
+	 * Returns the area enclosed by a Polygon. Returns zero for an invalid PolygonIndex. 
+	 */
+	UFUNCTION(BlueprintCallable, Category = "GeometryScript|PolygonList", meta = (ScriptMethod))
+	static double GetPolygonArea(FGeometryScriptGeneralPolygonList PolygonList, bool& bValidIndex, int32 PolygonIndex);
+
+	/**
+	 * Returns the bounding box of a Polygon. Returns an empty, invalid box for an invalid PolygonIndex.
+	 */
+	UFUNCTION(BlueprintCallable, Category = "GeometryScript|PolygonList", meta = (ScriptMethod))
+	static FBox2D GetPolygonBounds(FGeometryScriptGeneralPolygonList PolygonList, bool& bValidIndex, int32 PolygonIndex);
+
+	/**
+	 * Returns the area enclosed by a Polygon
+	 */
+	UFUNCTION(BlueprintCallable, Category = "GeometryScript|PolygonList", meta = (ScriptMethod))
+	static double GetPolygonListArea(FGeometryScriptGeneralPolygonList PolygonList);
+
+	/**
+	 * Returns the bounding box of a Polygon
+	 */
+	UFUNCTION(BlueprintCallable, Category = "GeometryScript|PolygonList", meta = (ScriptMethod))
+	static FBox2D GetPolygonListBounds(FGeometryScriptGeneralPolygonList PolygonList);
+
+	/**
+	 * Create a Polygon List of a single Polygon, with optional holes
+	 */
+	UFUNCTION(BlueprintCallable, Category = "GeometryScript|PolygonList", meta = (ScriptMethod, AutoCreateRefTerm = "HolePolygons"))
+	static UPARAM(DisplayName = "Polygon List") FGeometryScriptGeneralPolygonList CreatePolygonListFromSinglePolygon(FGeometryScriptSimplePolygon OuterPolygon, const TArray<FGeometryScriptSimplePolygon>& HolePolygons, bool bFixHoleOrientations = true);
+
+	/**
+	 * Add Polygon to a Polygon List, with optional holes. Returns index of the added polygon.
+	 */
+	UFUNCTION(BlueprintCallable, Category = "GeometryScript|PolygonList", meta = (ScriptMethod, AutoCreateRefTerm = "HolePolygons"))
+	static UPARAM(DisplayName = "Polygon Index") int32 AddPolygonToList(UPARAM(ref) FGeometryScriptGeneralPolygonList PolygonList, FGeometryScriptSimplePolygon OuterPolygon, const TArray<FGeometryScriptSimplePolygon>& HolePolygons, bool bFixHoleOrientations = true);
+
+	/**
+	 * Create a Polygon List from an array of Simple Polygons
+	 */
+	UFUNCTION(BlueprintCallable, Category = "GeometryScript|PolygonList")
+	static UPARAM(DisplayName = "Polygon List") FGeometryScriptGeneralPolygonList CreatePolygonListFromSimplePolygons(const TArray<FGeometryScriptSimplePolygon>& OuterPolygons);
+
+	/**
+	 * Append the polygons in 'Polygons to Append' to Polygon List
+	 */
+	UFUNCTION(BlueprintCallable, Category = "GeometryScript|PolygonList", meta = (ScriptMethod))
+	static void AppendPolygonList(UPARAM(ref) FGeometryScriptGeneralPolygonList PolygonList, FGeometryScriptGeneralPolygonList PolygonsToAppend);
+	
+	/** Compute union of all polygons in Polygon List. Also resolves self-intersections within each polygon. */
+	UFUNCTION(BlueprintCallable, Category = "GeometryScript|PolygonList", meta = (ScriptMethod))
+	static FGeometryScriptGeneralPolygonList PolygonsUnion(FGeometryScriptGeneralPolygonList PolygonList, bool bCopyInputOnFailure = true);
+
+	/** Compute difference of Polygon List and Polygons to Subtract */
+	UFUNCTION(BlueprintCallable, Category = "GeometryScript|PolygonList", meta = (ScriptMethod))
+	static FGeometryScriptGeneralPolygonList PolygonsDifference(FGeometryScriptGeneralPolygonList PolygonList, FGeometryScriptGeneralPolygonList PolygonsToSubtract);
+	
+	/** Compute intersection of Polygon List and Polygons to Intersect */
+	UFUNCTION(BlueprintCallable, Category = "GeometryScript|PolygonList", meta = (ScriptMethod))
+	static FGeometryScriptGeneralPolygonList PolygonsIntersection(FGeometryScriptGeneralPolygonList PolygonList, FGeometryScriptGeneralPolygonList PolygonsToIntersect);
+	
+	/** Compute exclusive or of Polygon List and Polygons to Exclusive Or */
+	UFUNCTION(BlueprintCallable, Category = "GeometryScript|PolygonList", meta = (ScriptMethod))
+	static FGeometryScriptGeneralPolygonList PolygonsExclusiveOr(FGeometryScriptGeneralPolygonList PolygonList, FGeometryScriptGeneralPolygonList PolygonsToExclusiveOr);
+	
+};
