Geometry Scripting : Geodesic Functions, add ability to trace a "straight" line on a mesh surface.

diff --git a/Engine/Plugins/Experimental/GeometryScripting/Source/GeometryScriptingCore/Private/MeshGeodesicFunctions.cpp b/Engine/Plugins/Experimental/GeometryScripting/Source/GeometryScriptingCore/Private/MeshGeodesicFunctions.cpp
--- a/Engine/Plugins/Experimental/GeometryScripting/Source/GeometryScriptingCore/Private/MeshGeodesicFunctions.cpp
+++ b/Engine/Plugins/Experimental/GeometryScripting/Source/GeometryScriptingCore/Private/MeshGeodesicFunctions.cpp
@@ -5,6 +5,7 @@
 
 #include "DynamicMesh/DynamicMesh3.h"
 #include "Operations/GeodesicPath.h"
+#include "Operations/MeshGeodesicSurfaceTracer.h"
 #include "Math/UnrealMathUtility.h"
 #include "Parameterization/MeshDijkstra.h"
 
@@ -313,6 +314,7 @@ UDynamicMesh* UGeometryScriptLibrary_MeshGeodesicFunctions::GetShortestSurfacePa
 
 		// convert minimized deformable path to a PolyPath
 		{
+			ShortestPath.bClosedLoop = false;
 			TArray<FVector>& GeodesicPath = *ShortestPath.Path;
 			GeodesicPath.Empty();
 
@@ -337,4 +339,99 @@ UDynamicMesh* UGeometryScriptLibrary_MeshGeodesicFunctions::GetShortestSurfacePa
 
 }
 
+UDynamicMesh* UGeometryScriptLibrary_MeshGeodesicFunctions::CreateSurfacePath( UDynamicMesh* TargetMesh,
+																			   FVector StartDirection,
+																			   int32 StartTriangleID,
+																			   FVector StartBaryCoords,
+																			   float PathLength,
+																			   FGeometryScriptPolyPath& SurfacePath,
+																			   bool& bFoundErrors,
+																			   UGeometryScriptDebug* Debug)
+{
+	typedef UE::Geometry::FMeshGeodesicSurfaceTracer  FSurfaceTracer;
+
+	bFoundErrors = true;
+
+	SurfacePath.Reset();
+
+	if (TargetMesh == nullptr)
+	{
+		UE::Geometry::AppendError(Debug, EGeometryScriptErrorType::InvalidInputs, LOCTEXT("CreateSurfacePath_InvalidInput", "CreateSurfacePath: TargetMesh is Null"));
+		return TargetMesh;
+	}
+
+	auto MakeValidBaryCentric = [](const FVector& Vec)->FVector
+	{
+		FVector ResultVec = Vec;
+		
+		// correct for error in BC calculation that might have returned a very small number in place of zero.
+		constexpr double BCTol = 1.e-2;
+		for (int i = 0; i < 3; ++i)
+		{
+			if (ResultVec[i] < 0 && ResultVec[i] >= -BCTol)
+			{
+				ResultVec[i] = 0.;
+			}
+		}
+
+
+
+		// any negative value indicates invalid BC
+		if (ResultVec.X < 0 || ResultVec.Y < 0 || ResultVec.Z < 0)
+		{
+			ResultVec = FVector(1./3., 1./3., 1./3.);
+		}
+
+		const double Sum = Vec.X + Vec.Y + Vec.Z;
+		constexpr double Tol = 0.05;
+		if (FMath::Abs(Sum - 1.0)  > Tol) // allow for a small amount of error before we declare invalid BC 
+		{
+			ResultVec = FVector(1./3., 1./3., 1./3.);
+		}
+		else
+		{
+			ResultVec *= (1./Sum);
+		}
+
+		return ResultVec;
+	};
+
+	
+	const FVector Direction = StartDirection.GetSafeNormal(1.e-5 /* Tolerance */, FVector::UpVector);
+
+	const FVector StartBC = MakeValidBaryCentric(StartBaryCoords);
+
+	
+	TargetMesh->ProcessMesh([&](const FDynamicMesh3& ReadMesh)
+		{
+
+			// do the surface trace.
+			FSurfaceTracer SurfaceTracer(ReadMesh);
+			SurfaceTracer.TraceMeshFromBaryPoint(StartTriangleID, StartBC, Direction, PathLength);
+			const TArray<FSurfaceTracer::FTraceResult>& Trace = SurfaceTracer.GetTraceResults();
+			const int32 NumPoints = Trace.Num();
+			
+			// convert results to poly path
+			SurfacePath.bClosedLoop = false;
+			SurfacePath.Path->Reserve(NumPoints);
+			for (int32 i = 0; i < NumPoints; ++i)
+			{
+				const FSurfaceTracer::FTraceResult& TracePoint = Trace[i];
+
+				if (TracePoint.TriID == IndexConstants::InvalidID) break;
+				
+				if (ReadMesh.IsTriangle(TracePoint.TriID))
+				{
+					FVector Pos = ReadMesh.GetTriBaryPoint(TracePoint.TriID, TracePoint.Barycentric[0], TracePoint.Barycentric[1], TracePoint.Barycentric[2]);
+					SurfacePath.Path->Add(Pos);
+				}
+				
+			}
+		});
+
+	bFoundErrors = false;
+	return TargetMesh;
+
+}
+
 #undef LOCTEXT_NAMESPACE
\ No newline at end of file
diff --git a/Engine/Plugins/Experimental/GeometryScripting/Source/GeometryScriptingCore/Public/GeometryScript/MeshGeodesicFunctions.h b/Engine/Plugins/Experimental/GeometryScripting/Source/GeometryScriptingCore/Public/GeometryScript/MeshGeodesicFunctions.h
--- a/Engine/Plugins/Experimental/GeometryScripting/Source/GeometryScriptingCore/Public/GeometryScript/MeshGeodesicFunctions.h
+++ b/Engine/Plugins/Experimental/GeometryScripting/Source/GeometryScriptingCore/Public/GeometryScript/MeshGeodesicFunctions.h
@@ -62,5 +62,30 @@ public:
 		bool& bFoundErrors,
 		UGeometryScriptDebug* Debug = nullptr);
 
-	
+	/**
+	* Computes a PolyPath that represents a "straight" surface path starting at the prescribed point on the mesh, and continuing 
+	* in the indicated direction until reaching the requested path length or encountering a mesh boundary, whichever comes first.
+	* @param TargetMesh defines the surface where the path is computed.
+	* @param Direction is a three-dimensional vector that is projected onto the mesh surface to determine the path direction. 
+	* @param StartTriangleID the ID of mesh Triangle that contains the start point of the path.
+	* @param StartBaryCoords indicates the location of start point within the start triangle, in terms of barycentric coordinates.
+	* @param MaxPathLength sets the maximal length of the path, but the actual path may be shorter as it automatically terminates when encountering a mesh boundary edge. 
+	* @param SurfacePath holds on return a PolyPath that forms a "straight" path along the mesh surface from the start position.
+	* @param bFoundErrors, will be false on success.  
+	* Note, Barycentric coordinates are of the form (a,b,c) where each entry is positive and a + b + c = 1.  
+	*       If the provided coordinate is invalid, the value (1/3, 1/3, 1/3) will be used.
+	*       Also, if the direction vector is nearly zero, the up-vector will be used. 
+	*/
+	UFUNCTION(BlueprintCallable, Category = "GeometryScript|MeshGeodesic", meta=(ScriptMethod))
+	static UPARAM(DisplayName = "Target Mesh") UDynamicMesh* 
+	CreateSurfacePath(
+		UDynamicMesh* TargetMesh,
+		FVector Direction,
+		int32 StartTriangleID,
+		FVector StartBaryCoords,
+		float MaxPathLength,
+		FGeometryScriptPolyPath& SurfacePath,
+		bool& bFoundErrors,
+		UGeometryScriptDebug* Debug = nullptr);
+
 };
\ No newline at end of file
diff --git a/Engine/Plugins/Runtime/GeometryProcessing/Source/DynamicMesh/Private/Operations/MeshGeodesicSurfaceTracer.cpp b/Engine/Plugins/Runtime/GeometryProcessing/Source/DynamicMesh/Private/Operations/MeshGeodesicSurfaceTracer.cpp
--- a/Engine/Plugins/Runtime/GeometryProcessing/Source/DynamicMesh/Private/Operations/MeshGeodesicSurfaceTracer.cpp
+++ b/Engine/Plugins/Runtime/GeometryProcessing/Source/DynamicMesh/Private/Operations/MeshGeodesicSurfaceTracer.cpp
@@ -772,6 +772,10 @@ double FMeshGeodesicSurfaceTracer::TraceMeshFromBaryPoint(const int32 TriID, con
 
 	// project the ray onto the X,Y plane ( i.e. the face of the triangle )
 	FVector2d ProjectedRay(FVector3d::DotProduct(XDir, RayDir3), FVector3d::DotProduct(YDir, RayDir3));
+	if (!ProjectedRay.Normalize())
+	{
+		ProjectedRay = FVector2d::UnitX();
+	}
 
 	// lower level code expects the 2d ray to be encoded relative to the "direction" of the first edge, edge(v0,v1)
 	const int32 EID = Mesh->GetTriEdge(TriID, 0);
diff --git a/Engine/Plugins/Runtime/GeometryProcessing/Source/DynamicMesh/Public/Operations/MeshGeodesicSurfaceTracer.h b/Engine/Plugins/Runtime/GeometryProcessing/Source/DynamicMesh/Public/Operations/MeshGeodesicSurfaceTracer.h
--- a/Engine/Plugins/Runtime/GeometryProcessing/Source/DynamicMesh/Public/Operations/MeshGeodesicSurfaceTracer.h
+++ b/Engine/Plugins/Runtime/GeometryProcessing/Source/DynamicMesh/Public/Operations/MeshGeodesicSurfaceTracer.h
@@ -199,7 +199,7 @@ public:
 * in terms of barycentric coordinates, triangle references and local path direction relative 
 * to a local frame (in the form of one of the edges of the triangle).
 */
-class FMeshGeodesicSurfaceTracer
+class DYNAMICMESH_API FMeshGeodesicSurfaceTracer
 {
 public:
 
