- Fixed a bug where bRecomputeNormals and  bRecomputeTangents were not set on the affected LOD in CopyMeshToSkeleletalMesh geometry script funciton.
- Fixed a bug where bHasNormals and bHasTangents flags were not set when converting FMeshDescription to FSkeletalMeshImportData.

#rnx

diff --git a/Engine/Plugins/Experimental/GeometryScripting/Source/GeometryScriptingCore/Private/MeshAssetFunctions.cpp b/Engine/Plugins/Experimental/GeometryScripting/Source/GeometryScriptingCore/Private/MeshAssetFunctions.cpp
--- a/Engine/Plugins/Experimental/GeometryScripting/Source/GeometryScriptingCore/Private/MeshAssetFunctions.cpp
+++ b/Engine/Plugins/Experimental/GeometryScripting/Source/GeometryScriptingCore/Private/MeshAssetFunctions.cpp
@@ -599,6 +599,10 @@ UDynamicMesh* UGeometryScriptLibrary_StaticMeshFunctions::CopyMeshToSkeletalMesh
 		Converter.Convert(&ReadMesh, MeshDescription, !Options.bEnableRecomputeTangents);
 	});
 
+	FSkeletalMeshLODInfo* SkeletalLODInfo = ToSkeletalMeshAsset->GetLODInfo(TargetLOD.LODIndex);
+	SkeletalLODInfo->BuildSettings.bRecomputeNormals = Options.bEnableRecomputeNormals;
+	SkeletalLODInfo->BuildSettings.bRecomputeTangents = Options.bEnableRecomputeTangents;
+
 	FSkeletalMeshImportData SkeletalMeshImportData = 
 		FSkeletalMeshImportData::CreateFromMeshDescription(MeshDescription);
 
diff --git a/Engine/Source/Runtime/Engine/Private/Rendering/SkeletalMeshLODImporterData.cpp b/Engine/Source/Runtime/Engine/Private/Rendering/SkeletalMeshLODImporterData.cpp
--- a/Engine/Source/Runtime/Engine/Private/Rendering/SkeletalMeshLODImporterData.cpp
+++ b/Engine/Source/Runtime/Engine/Private/Rendering/SkeletalMeshLODImporterData.cpp
@@ -2005,6 +2005,16 @@ FSkeletalMeshImportData FSkeletalMeshImportData::CreateFromMeshDescription(const
 			Face.TangentZ[Corner] = VertexInstanceNormals[VertexInstanceID];
 			Face.TangentY[Corner] = FVector3f::CrossProduct(VertexInstanceNormals[VertexInstanceID], VertexInstanceTangents[VertexInstanceID]).GetSafeNormal() * VertexInstanceBiNormalSigns[VertexInstanceID];
 
+			if (Face.TangentX[Corner] != FVector3f::ZeroVector)
+			{
+				SkelMeshImportData.bHasTangents = true;
+			}
+			
+			if (Face.TangentZ[Corner] != FVector3f::ZeroVector)
+			{
+				SkelMeshImportData.bHasNormals = true;
+			}
+
 			Face.WedgeIndex[Corner] = SkelMeshImportData.Wedges.Add(Wedge);
 		}
 	}
diff --git a/Engine/Source/Runtime/Engine/Public/Rendering/SkeletalMeshLODImporterData.h b/Engine/Source/Runtime/Engine/Public/Rendering/SkeletalMeshLODImporterData.h
--- a/Engine/Source/Runtime/Engine/Public/Rendering/SkeletalMeshLODImporterData.h
+++ b/Engine/Source/Runtime/Engine/Public/Rendering/SkeletalMeshLODImporterData.h
@@ -441,6 +441,11 @@ public:
 	 */
 	bool GetMeshDescription(FMeshDescription &OutMeshDescription) const;
 
+	/**
+	 * @note MeshDescription always contains color, normal and tangent data by default. Therefore, while iterating
+	 * over the vertices, we check if at least one normal/tangent vector is not a zero vector and the color is
+	 * not white, then we set the corresponding bHasNormals/bHasTangent/bHasVertexColors flags to true.
+	 */
 	static FSkeletalMeshImportData CreateFromMeshDescription(const FMeshDescription &InMeshDescription);
 
 private:
