GeometryScript: expose texture filtering options in BakeTexture node

diff --git a/Engine/Plugins/Experimental/GeometryScripting/Source/GeometryScriptingCore/Private/MeshBakeFunctions.cpp b/Engine/Plugins/Experimental/GeometryScripting/Source/GeometryScriptingCore/Private/MeshBakeFunctions.cpp
--- a/Engine/Plugins/Experimental/GeometryScripting/Source/GeometryScriptingCore/Private/MeshBakeFunctions.cpp
+++ b/Engine/Plugins/Experimental/GeometryScripting/Source/GeometryScriptingCore/Private/MeshBakeFunctions.cpp
@@ -93,6 +93,19 @@ namespace GeometryScriptBakeLocals
 		return Samples;
 	}
 
+	FMeshMapBaker::EBakeFilterType GetFilter(EGeometryScriptBakeFilteringType FilteringType)
+	{
+		switch (FilteringType)
+		{
+		case EGeometryScriptBakeFilteringType::Box:
+			return FMeshMapBaker::EBakeFilterType::Box;
+		case EGeometryScriptBakeFilteringType::BSpline:
+		default:
+			return FMeshMapBaker::EBakeFilterType::BSpline;
+		}
+	}
+
+
 	FTexture2DBuilder::ETextureType GetTextureType(const FMeshMapEvaluator* Evaluator, const EGeometryScriptBakeBitDepth MapFormat)
 	{
 		FTexture2DBuilder::ETextureType TexType = FTexture2DBuilder::ETextureType::Color;
@@ -520,6 +533,7 @@ namespace GeometryScriptBakeLocals
 		Baker.SetDimensions(BakeDimensions);
 		Baker.SetProjectionDistance(BakeOptions.ProjectionDistance);
 		Baker.SetSamplesPerPixel(GetSamplesPerPixel(BakeOptions.SamplesPerPixel));
+		Baker.SetFilter(GetFilter(BakeOptions.FilteringType));
 		TUniquePtr<TImageBuilder<FVector4f>> SampleFilterMask = GetSampleFilterMask(BakeOptions, BakeTexturePrefix, Debug);
 		if (SampleFilterMask)
 		{
diff --git a/Engine/Plugins/Experimental/GeometryScripting/Source/GeometryScriptingCore/Public/GeometryScript/MeshBakeFunctions.h b/Engine/Plugins/Experimental/GeometryScripting/Source/GeometryScriptingCore/Public/GeometryScript/MeshBakeFunctions.h
--- a/Engine/Plugins/Experimental/GeometryScripting/Source/GeometryScriptingCore/Public/GeometryScript/MeshBakeFunctions.h
+++ b/Engine/Plugins/Experimental/GeometryScripting/Source/GeometryScriptingCore/Public/GeometryScript/MeshBakeFunctions.h
@@ -41,6 +41,13 @@ enum class EGeometryScriptBakeSamplesPerPixel : uint8
 	Samples256 UMETA(DisplayName = "256")
 };
 
+UENUM(BlueprintType)
+enum class EGeometryScriptBakeFilteringType : uint8
+{
+	BSpline UMETA(DisplayName = "B-Spline"),
+	Box UMETA(DisplayName = "Box")
+};
+
 UENUM(BlueprintType)
 enum class EGeometryScriptBakeTypes : uint8
 {
@@ -212,6 +219,10 @@ struct GEOMETRYSCRIPTINGCORE_API FGeometryScriptBakeTextureOptions
 	UPROPERTY(EditAnywhere, BlueprintReadWrite, Category = Options)
 	TObjectPtr<UTexture2D> SampleFilterMask = nullptr;
 
+	/** Filtering Type to perform on samples */
+	UPROPERTY(EditAnywhere, BlueprintReadWrite, Category = Options)
+	EGeometryScriptBakeFilteringType FilteringType = EGeometryScriptBakeFilteringType::BSpline;
+
 	/** Maximum allowed distance for the projection from target mesh to source mesh for the sample to be considered valid.
 	 * This is only relevant if a separate source mesh is provided. */
 	UPROPERTY(EditAnywhere, BlueprintReadWrite, Category = Options)
