GeometryScript: add primitive function AppendBoundingBox, which takes an FBox instead of box parameters. Add MarkForMeshRebuild() function to GeneratedDynamicMeshActor, allows other BP code/etc to force a rebuild w/o re-running construction script.

diff --git a/Engine/Plugins/Experimental/GeometryScripting/Source/GeometryScriptingCore/Private/MeshPrimitiveFunctions.cpp b/Engine/Plugins/Experimental/GeometryScripting/Source/GeometryScriptingCore/Private/MeshPrimitiveFunctions.cpp
--- a/Engine/Plugins/Experimental/GeometryScripting/Source/GeometryScriptingCore/Private/MeshPrimitiveFunctions.cpp
+++ b/Engine/Plugins/Experimental/GeometryScripting/Source/GeometryScriptingCore/Private/MeshPrimitiveFunctions.cpp
@@ -225,6 +225,36 @@ UDynamicMesh* UGeometryScriptLibrary_MeshPrimitiveFunctions::AppendBox(
 }
 
 
+UDynamicMesh* UGeometryScriptLibrary_MeshPrimitiveFunctions::AppendBoundingBox(
+	UDynamicMesh* TargetMesh,
+	FGeometryScriptPrimitiveOptions PrimitiveOptions,
+	FTransform Transform,
+	FBox Box,
+	int32 StepsX,
+	int32 StepsY,
+	int32 StepsZ,
+	UGeometryScriptDebug* Debug)
+{
+	if (TargetMesh == nullptr)
+	{
+		UE::Geometry::AppendError(Debug, EGeometryScriptErrorType::InvalidInputs, LOCTEXT("PrimitiveFunctions_AppendBoundingBox", "AppendBoundingBox: TargetMesh is Null"));
+		return TargetMesh;
+	}
+
+	// todo: if Steps X/Y/Z are zero, can use trivial box generator
+
+	FGridBoxMeshGenerator GridBoxGenerator;
+	GridBoxGenerator.Box = UE::Geometry::FOrientedBox3d( UE::Geometry::FAxisAlignedBox3d(Box) );
+	GridBoxGenerator.EdgeVertices = FIndex3i(FMath::Max(0, StepsX), FMath::Max(0, StepsY), FMath::Max(0, StepsZ));
+	GridBoxGenerator.bPolygroupPerQuad = (PrimitiveOptions.PolygroupMode == EGeometryScriptPrimitivePolygroupMode::PerQuad);
+	GridBoxGenerator.Generate();
+
+	AppendPrimitive(TargetMesh, &GridBoxGenerator, Transform, PrimitiveOptions, FVector3d::Zero());
+
+	return TargetMesh;
+}
+
+
 
 UDynamicMesh* UGeometryScriptLibrary_MeshPrimitiveFunctions::AppendSphereLatLong(
 	UDynamicMesh* TargetMesh,
diff --git a/Engine/Plugins/Experimental/GeometryScripting/Source/GeometryScriptingCore/Public/GeometryScript/MeshPrimitiveFunctions.h b/Engine/Plugins/Experimental/GeometryScripting/Source/GeometryScriptingCore/Public/GeometryScript/MeshPrimitiveFunctions.h
--- a/Engine/Plugins/Experimental/GeometryScripting/Source/GeometryScriptingCore/Public/GeometryScript/MeshPrimitiveFunctions.h
+++ b/Engine/Plugins/Experimental/GeometryScripting/Source/GeometryScriptingCore/Public/GeometryScript/MeshPrimitiveFunctions.h
@@ -172,6 +172,22 @@ public:
 		EGeometryScriptPrimitiveOriginMode Origin = EGeometryScriptPrimitiveOriginMode::Base,
 		UGeometryScriptDebug* Debug = nullptr);
 
+	/**
+	 * Appends a 3D box to the Target Mesh with dimensions and origin taken from the input Box
+	 */
+	UFUNCTION(BlueprintCallable, Category = "GeometryScript|Primitives", meta=(ScriptMethod))
+	static UPARAM(DisplayName = "Target Mesh") UDynamicMesh* 
+	AppendBoundingBox( 
+		UDynamicMesh* TargetMesh, 
+		FGeometryScriptPrimitiveOptions PrimitiveOptions,
+		FTransform Transform,
+		FBox Box,
+		int32 StepsX = 0,
+		int32 StepsY = 0,
+		int32 StepsZ = 0,
+		UGeometryScriptDebug* Debug = nullptr);
+
+
 	/**
 	* Appends a 3D Sphere triangulated using latitude/longitude topology to the Target Mesh.
 	*/ 
@@ -326,7 +342,7 @@ public:
 	 * @param PolylineVertices vertices of the open 2D path that will be swept along the SweepPath
 	 * @param SweepPath defines the 3D sweep path curve as a 3D poly-path, with rotation and scaling at each polypath vertex taken from the Transform
 	 * @param PolylineTexParamU defines U coordinate value for each element in PolylineVertices. Must be same length as PolylineVertices (ignored if length=0).
-	 * @param SweepPathTexParamV defines V coordinate value for each element in SweepPath. Must be same length as PolylineVertices if bLoop=false, length+1 if bLoop=true, and ignored if length=0.
+	 * @param SweepPathTexParamV defines V coordinate value for each element in SweepPath. Must be same length as SweepPath if bLoop=false, length+1 if bLoop=true, and ignored if length=0.
 	 * @param bLoop if true, SweepPath is considered to be a Loop and a section connecting the end and start of the path is added (bCapped is ignored)
 	 * @param StartScale uniform scaling applied to the 2D polygon at the start of the path. Interpolated via arc length to EndScale at the end of the path.
 	 * @param EndScale uniform scaling applied to the 2D polygon at the end of the path
diff --git a/Engine/Plugins/Experimental/GeometryScripting/Source/GeometryScriptingEditor/Private/GeometryActors/GeneratedDynamicMeshActor.cpp b/Engine/Plugins/Experimental/GeometryScripting/Source/GeometryScriptingEditor/Private/GeometryActors/GeneratedDynamicMeshActor.cpp
--- a/Engine/Plugins/Experimental/GeometryScripting/Source/GeometryScriptingEditor/Private/GeometryActors/GeneratedDynamicMeshActor.cpp
+++ b/Engine/Plugins/Experimental/GeometryScripting/Source/GeometryScriptingEditor/Private/GeometryActors/GeneratedDynamicMeshActor.cpp
@@ -195,6 +195,28 @@ void AGeneratedDynamicMeshActor::ExecuteRebuildGeneratedMeshIfPending()
 	}
 }
 
+
+void AGeneratedDynamicMeshActor::MarkForMeshRebuild(bool bImmediate, bool bImmediateEvenIfFrozen)
+{
+	if (bImmediate)
+	{
+		bool bRestoreFrozenState = bFrozen;
+		if (bImmediateEvenIfFrozen && bFrozen)
+		{
+			bFrozen = false;
+		}
+
+		bGeneratedMeshRebuildPending = true;
+		ExecuteRebuildGeneratedMeshIfPending();
+
+		bFrozen = bRestoreFrozenState;
+	}
+	else
+	{
+		bGeneratedMeshRebuildPending = true;
+	}
+}
+
 void AGeneratedDynamicMeshActor::IncrementProgress(int NumSteps, FString Message)
 {
 	if (ActiveSlowTask)
diff --git a/Engine/Plugins/Experimental/GeometryScripting/Source/GeometryScriptingEditor/Public/GeometryActors/GeneratedDynamicMeshActor.h b/Engine/Plugins/Experimental/GeometryScripting/Source/GeometryScriptingEditor/Public/GeometryActors/GeneratedDynamicMeshActor.h
--- a/Engine/Plugins/Experimental/GeometryScripting/Source/GeometryScriptingEditor/Public/GeometryActors/GeneratedDynamicMeshActor.h
+++ b/Engine/Plugins/Experimental/GeometryScripting/Source/GeometryScriptingEditor/Public/GeometryActors/GeneratedDynamicMeshActor.h
@@ -47,6 +47,12 @@ public:
 	UFUNCTION(BlueprintImplementableEvent, CallInEditor, Category = "Events")
 	void OnRebuildGeneratedMesh(UDynamicMesh* TargetMesh);
 
+	/**
+	 * Mark this Actor as modified so that OnRebuildGeneratedMesh runs 
+	 */
+	UFUNCTION(BlueprintCallable, Category = "DynamicMeshActor")
+	void MarkForMeshRebuild(bool bImmediate = false, bool bImmediateEvenIfFrozen = false);
+
 
 	/**
 	 * This function will fire the OnRebuildGeneratedMesh event if the actor has been
