Add support for ISMC's to CopyMeshFromComponent

diff --git a/Engine/Plugins/Experimental/GeometryScripting/Source/GeometryScriptingCore/Private/SceneUtilityFunctions.cpp b/Engine/Plugins/Experimental/GeometryScripting/Source/GeometryScriptingCore/Private/SceneUtilityFunctions.cpp
--- a/Engine/Plugins/Experimental/GeometryScripting/Source/GeometryScriptingCore/Private/SceneUtilityFunctions.cpp
+++ b/Engine/Plugins/Experimental/GeometryScripting/Source/GeometryScriptingCore/Private/SceneUtilityFunctions.cpp
@@ -7,10 +7,12 @@
 #include "GeometryScript/MeshAssetFunctions.h"
 #include "DynamicMesh/MeshTransforms.h"
 #include "DynamicMesh/MeshNormals.h"
+#include "DynamicMeshEditor.h"
 
 #include "Engine/StaticMesh.h"
 #include "Engine/SkeletalMesh.h"
 #include "Components/StaticMeshComponent.h"
+#include "Components/InstancedStaticMeshComponent.h"
 #include "Components/SkeletalMeshComponent.h"
 #include "Components/SkinnedMeshComponent.h"
 #include "Components/BrushComponent.h"
@@ -105,6 +107,38 @@ UDynamicMesh* UGeometryScriptLibrary_SceneUtilityFunctions::CopyMeshFromComponen
 			AssetOptions.bRequestTangents = Options.bWantTangents;
 			UGeometryScriptLibrary_StaticMeshFunctions::CopyMeshFromStaticMesh(
 				StaticMesh, ToDynamicMesh, AssetOptions, Options.RequestedLOD, Outcome, Debug);	// will set Outcome pin
+
+			// if we have an ISMC, append instances
+			if (UInstancedStaticMeshComponent* ISMComponent = Cast<UInstancedStaticMeshComponent>(StaticMeshComponent))
+			{
+				FDynamicMesh3 InstancedMesh;
+				ToDynamicMesh->EditMesh([&](FDynamicMesh3& EditMesh) { InstancedMesh = MoveTemp(EditMesh); EditMesh.Clear(); },
+					EDynamicMeshChangeType::MeshChange, EDynamicMeshAttributeChangeFlags::Unknown, true);
+				
+				FDynamicMesh3 AccumMesh;
+				AccumMesh.EnableMatchingAttributes(InstancedMesh);
+				FDynamicMeshEditor Editor(&AccumMesh);
+				FMeshIndexMappings Mappings;
+
+				int32 NumInstances = ISMComponent->GetInstanceCount();
+				for (int32 InstanceIdx = 0; InstanceIdx < NumInstances; ++InstanceIdx)
+				{
+					if (ISMComponent->IsValidInstance(InstanceIdx))
+					{
+						FTransform InstanceTransform;
+						ISMComponent->GetInstanceTransform(InstanceIdx, InstanceTransform, /*bWorldSpace=*/false);
+						FTransformSRT3d XForm(InstanceTransform);
+
+						Mappings.Reset();
+						Editor.AppendMesh(&InstancedMesh, Mappings,
+							[&](int, const FVector3d& Position) { return XForm.TransformPosition(Position); },
+							[&](int, const FVector3d& Normal) { return XForm.TransformNormal(Normal); });
+					}
+				}
+
+				ToDynamicMesh->EditMesh([&](FDynamicMesh3& EditMesh) { EditMesh = MoveTemp(AccumMesh); },
+					EDynamicMeshChangeType::MeshChange, EDynamicMeshAttributeChangeFlags::Unknown, true);
+			}
 		}
 		else
 		{
