Added new mesh bone weight functions for getting information about individual bones or the skeleton.
Added InpaintWeights option to the TransferBoneWeightsFromMesh function

diff --git a/Engine/Plugins/Experimental/GeometryScripting/Source/GeometryScriptingCore/Private/MeshBoneWeightFunctions.cpp b/Engine/Plugins/Experimental/GeometryScripting/Source/GeometryScriptingCore/Private/MeshBoneWeightFunctions.cpp
--- a/Engine/Plugins/Experimental/GeometryScripting/Source/GeometryScriptingCore/Private/MeshBoneWeightFunctions.cpp
+++ b/Engine/Plugins/Experimental/GeometryScripting/Source/GeometryScriptingCore/Private/MeshBoneWeightFunctions.cpp
@@ -11,7 +11,8 @@
 #include "SkinningOps/SkinBindingOp.h"
 #include "UDynamicMesh.h"
 #include "Operations/TransferBoneWeights.h"
-
+#include "MathUtil.h"
+#include "Containers/Queue.h"
 
 #include UE_INLINE_GENERATED_CPP_BY_NAME(MeshBoneWeightFunctions)
 
@@ -380,6 +381,16 @@ UDynamicMesh* UGeometryScriptLibrary_MeshBoneWeightFunctions::TransferBoneWeight
 
 		FTransferBoneWeights TransferBoneWeights(&ReadMesh, Options.SourceProfile.GetProfileName());
 		TransferBoneWeights.TransferMethod = static_cast<FTransferBoneWeights::ETransferBoneWeightsMethod>(Options.TransferMethod);
+		TransferBoneWeights.bUseParallel = true;
+		if (Options.TransferMethod == ETransferBoneWeightsMethod::InpaintWeights)
+		{
+			TransferBoneWeights.NormalThreshold = FMathd::DegToRad * Options.NormalThreshold;
+			TransferBoneWeights.SearchRadius = Options.RadiusPercentage * ReadMesh.GetBounds().DiagonalLength();
+			TransferBoneWeights.NumSmoothingIterations = Options.NumSmoothingIterations;
+			TransferBoneWeights.SmoothingStrength = Options.SmoothingStrength;
+			TransferBoneWeights.LayeredMeshSupport = Options.LayeredMeshSupport;
+		}
+
 		if (TransferBoneWeights.Validate() != EOperationValidationResult::Ok)
 		{
 			UE::Geometry::AppendError(Debug, EGeometryScriptErrorType::OperationFailed, LOCTEXT("TransferBoneWeightsFromMesh_ValidationFailed", "TransferBoneWeightsFromMesh: Invalid parameters were set for the transfer weight operator"));
@@ -511,15 +522,274 @@ UDynamicMesh* UGeometryScriptLibrary_MeshBoneWeightFunctions::GetBoneIndex(
 	{	
 		if (!EditMesh.HasAttributes() || !EditMesh.Attributes()->HasBones())
 		{
-			AppendError(Debug, EGeometryScriptErrorType::InvalidInputs, LOCTEXT("GetBoneIndex_TargetMeshHasNoBones", "TargetMesh has no bone attributes"));
+			AppendError(Debug, EGeometryScriptErrorType::InvalidInputs, LOCTEXT("GetBoneIndex_TargetMeshHasNoBones", "GetBoneIndex: TargetMesh has no bone attributes"));
 			return;
 		}
 
 		// INDEX_NONE if BoneName doesn't exist in the bone names attribute
 		BoneIndex = EditMesh.Attributes()->GetBoneNames()->GetAttribValues().Find(BoneName);
 		
-		bIsValidBoneName = BoneIndex == INDEX_NONE;
+		bIsValidBoneName = BoneIndex != INDEX_NONE;
+
+	});
+
+	return TargetMesh;
+}
+
+UDynamicMesh* UGeometryScriptLibrary_MeshBoneWeightFunctions::GetRootBoneName(
+	UDynamicMesh* TargetMesh,
+	FName& BoneName,
+	UGeometryScriptDebug* Debug)
+{
+	if (TargetMesh == nullptr)
+	{
+		AppendError(Debug, EGeometryScriptErrorType::InvalidInputs, LOCTEXT("GetRootBoneName_InvalidTargetMesh", "GetRootBoneName: TargetMesh is Null"));
+		return TargetMesh;
+	}
+
+	TargetMesh->ProcessMesh([&](const FDynamicMesh3& EditMesh)
+	{	
+		if (!EditMesh.HasAttributes() || !EditMesh.Attributes()->HasBones())
+		{
+			AppendError(Debug, EGeometryScriptErrorType::InvalidInputs, LOCTEXT("GetRootBoneName_TargetMeshHasNoBones", "GetRootBoneName: TargetMesh has no bone attributes"));
+			return;
+		}
+
+		if (EditMesh.Attributes()->GetBoneNames()->GetAttribValues().IsEmpty())
+		{
+			AppendError(Debug, EGeometryScriptErrorType::InvalidInputs, LOCTEXT("GetRootBoneName_TargetMeshHasEmptySkeleton", "GetRootBoneName: TargetMesh has bone attributes set, but they are empty (doesn't contain a single bone)"));
+			return;
+		}
+
+		// Root bone's parent will be INDEX_NONE
+		const int32 BoneIdx = EditMesh.Attributes()->GetBoneParentIndices()->GetAttribValues().Find(INDEX_NONE);
+
+		if (BoneIdx == INDEX_NONE)
+		{
+			AppendError(Debug, EGeometryScriptErrorType::InvalidInputs, LOCTEXT("GetRootBoneName_TargetMeshHasNoRootBone", "GetRootBoneName: TargetMesh has no root bone"));
+			return;
+		}
+
+		BoneName = EditMesh.Attributes()->GetBoneNames()->GetAttribValues()[BoneIdx];
+	});
+
+	return TargetMesh;
+}
+
+UDynamicMesh* UGeometryScriptLibrary_MeshBoneWeightFunctions::GetBoneChildren(
+	UDynamicMesh* TargetMesh,
+	FName BoneName,
+	bool bRecursive,
+	bool& bIsValidBoneName,
+	TArray<FGeometryScriptBoneInfo>& ChildrenInfo,
+	UGeometryScriptDebug* Debug)
+{
+	if (TargetMesh == nullptr)
+	{
+		AppendError(Debug, EGeometryScriptErrorType::InvalidInputs, LOCTEXT("GetBoneChildren_InvalidTargetMesh", "GetBoneChildren: TargetMesh is Null"));
+		return TargetMesh;
+	}
+
+	TargetMesh->ProcessMesh([&](const FDynamicMesh3& EditMesh)
+	{	
+		if (!EditMesh.HasAttributes() || !EditMesh.Attributes()->HasBones())
+		{
+			AppendError(Debug, EGeometryScriptErrorType::InvalidInputs, LOCTEXT("GetBoneChildren_TargetMeshHasNoBones", "GetBoneChildren: TargetMesh has no bone attributes"));
+			return;
+		}
+		
+		const int32 NumBones = EditMesh.Attributes()->GetNumBones();
+		const TArray<FName>& NamesAttrib = EditMesh.Attributes()->GetBoneNames()->GetAttribValues();
+		const TArray<int32>& ParentsAttrib = EditMesh.Attributes()->GetBoneParentIndices()->GetAttribValues();
+		const TArray<FTransform>& TransformsAttrib = EditMesh.Attributes()->GetBonePoses()->GetAttribValues();
+		const TArray<FVector4f>& ColorsAttrib = EditMesh.Attributes()->GetBoneColors()->GetAttribValues();
+
+		// INDEX_NONE if BoneName doesn't exist in the bone names attribute
+		const int32 BoneIndex = NamesAttrib.Find(BoneName);
+		
+		bIsValidBoneName = BoneIndex != INDEX_NONE;
+
+		if (bIsValidBoneName)
+		{	
+			ChildrenInfo.Reset();
 
+			TArray<int32> ChildrenIndices;
+
+			// add direct children only
+			for (int32 Idx = 0; Idx < NumBones; ++Idx)
+			{
+				if (ParentsAttrib[Idx] == BoneIndex)
+				{
+					ChildrenIndices.Add(Idx);
+				}
+			}
+			
+			// add all grandchildren as well
+			if (bRecursive) 
+			{
+				// Map each bone index to its direct children bone indices
+				TMap<int32, TArray<int32>> ChildrenMap;
+				for (int32 Idx = 0; Idx < NumBones; ++Idx)
+				{
+					const int32 ParentIdx = ParentsAttrib[Idx];
+					if (ChildrenMap.Contains(ParentsAttrib[Idx]))
+					{
+						ChildrenMap[ParentIdx].Add(Idx);
+					}
+					else 
+					{
+						ChildrenMap.Add(ParentIdx, {Idx});
+					}
+				}
+
+				// Do breadth first search
+				TQueue<int32> BFS;
+				for (const int32 ChildIdx : ChildrenIndices)
+				{
+					BFS.Enqueue(ChildIdx);
+				}
+
+				while (!BFS.IsEmpty())
+				{
+					int32 ChildIdx; 
+					BFS.Dequeue(ChildIdx);
+
+					if (ChildrenMap.Contains(ChildIdx))
+					{
+						for (const int32 GrandChildIdx : ChildrenMap[ChildIdx])
+						{
+							BFS.Enqueue(GrandChildIdx);
+							checkSlow(ChildrenIndices.Find(GrandChildIdx) == INDEX_NONE);
+							ChildrenIndices.Add(GrandChildIdx);
+						}
+					}
+				}
+			}
+			
+			// Get all information about the children
+			ChildrenInfo.SetNum(NumBones);
+			for (const int ChildIdx : ChildrenIndices)
+			{
+				FGeometryScriptBoneInfo& Entry = ChildrenInfo[ChildIdx];
+				Entry.Index = ChildIdx;
+				Entry.Name = NamesAttrib[ChildIdx];
+				Entry.ParentIndex = ParentsAttrib[ChildIdx];
+				Entry.LocalTransform = TransformsAttrib[ChildIdx];
+				Entry.WorldTransform = Entry.LocalTransform;
+				Entry.Color = ColorsAttrib[ChildIdx];
+
+				int32 CurParentIndex = ParentsAttrib[ChildIdx];
+				while (CurParentIndex != INDEX_NONE) // until we didn't reach the root bone
+				{
+					Entry.WorldTransform = Entry.WorldTransform * TransformsAttrib[CurParentIndex];
+					CurParentIndex = ParentsAttrib[CurParentIndex];
+				}
+			}
+		}
+	});
+
+	return TargetMesh;
+}
+
+UDynamicMesh* UGeometryScriptLibrary_MeshBoneWeightFunctions::GetBoneInfo(
+		UDynamicMesh* TargetMesh,
+		FName BoneName,
+		bool& bIsValidBoneName,
+		FGeometryScriptBoneInfo& BoneInfo,
+		UGeometryScriptDebug* Debug)
+{
+	if (TargetMesh == nullptr)
+	{
+		AppendError(Debug, EGeometryScriptErrorType::InvalidInputs, LOCTEXT("GetBoneInfo_InvalidTargetMesh", "GetBoneInfo: TargetMesh is Null"));
+		return TargetMesh;
+	}
+
+	TargetMesh->ProcessMesh([&](const FDynamicMesh3& EditMesh)
+	{	
+		if (!EditMesh.HasAttributes() || !EditMesh.Attributes()->HasBones())
+		{
+			AppendError(Debug, EGeometryScriptErrorType::InvalidInputs, LOCTEXT("GetBoneInfo_TargetMeshHasNoBones", "GetBoneInfo: TargetMesh has no bone attributes"));
+			return;
+		}
+		
+		
+		// INDEX_NONE if BoneName doesn't exist in the bone names attribute
+		int32 BoneIndex = EditMesh.Attributes()->GetBoneNames()->GetAttribValues().Find(BoneName);
+
+		bIsValidBoneName = BoneIndex != INDEX_NONE;
+
+		if (bIsValidBoneName)
+		{	
+			const TArray<FName>& NamesAttrib = EditMesh.Attributes()->GetBoneNames()->GetAttribValues();
+			const TArray<int32>& ParentsAttrib = EditMesh.Attributes()->GetBoneParentIndices()->GetAttribValues();
+			const TArray<FTransform>& TransformsAttrib = EditMesh.Attributes()->GetBonePoses()->GetAttribValues();
+			const TArray<FVector4f>& ColorsAttrib = EditMesh.Attributes()->GetBoneColors()->GetAttribValues();
+
+			BoneInfo.Index = BoneIndex;
+			BoneInfo.Name = NamesAttrib[BoneIndex];
+			BoneInfo.ParentIndex = ParentsAttrib[BoneIndex];
+			BoneInfo.LocalTransform = TransformsAttrib[BoneIndex];
+			BoneInfo.WorldTransform = BoneInfo.LocalTransform;
+			BoneInfo.Color = ColorsAttrib[BoneIndex];
+
+			int32 CurParentIndex = ParentsAttrib[BoneIndex];
+			while (CurParentIndex != INDEX_NONE) // until we didn't reach the root bone
+			{
+				BoneInfo.WorldTransform = BoneInfo.WorldTransform * TransformsAttrib[CurParentIndex];
+				CurParentIndex = ParentsAttrib[CurParentIndex];
+			}
+		}
+	});
+
+	return TargetMesh;
+}
+
+UDynamicMesh* UGeometryScriptLibrary_MeshBoneWeightFunctions::GetAllBonesInfo(
+		UDynamicMesh* TargetMesh,
+		TArray<FGeometryScriptBoneInfo>& BonesInfo,
+		UGeometryScriptDebug* Debug)
+{
+	if (TargetMesh == nullptr)
+	{
+		AppendError(Debug, EGeometryScriptErrorType::InvalidInputs, LOCTEXT("GetAllBonesInfo_InvalidTargetMesh", "GetAllBonesInfo: TargetMesh is Null"));
+		return TargetMesh;
+	}
+
+	TargetMesh->ProcessMesh([&](const FDynamicMesh3& EditMesh)
+	{	
+		if (!EditMesh.HasAttributes() || !EditMesh.Attributes()->HasBones())
+		{
+			AppendError(Debug, EGeometryScriptErrorType::InvalidInputs, LOCTEXT("GetAllBonesInfo_TargetMeshHasNoBones", "GetAllBonesInfo: TargetMesh has no bone attributes"));
+			return;
+		}
+		
+		const int32 NumBones = EditMesh.Attributes()->GetNumBones();
+
+		BonesInfo.SetNum(NumBones);
+
+		const TArray<FName>& NamesAttrib = EditMesh.Attributes()->GetBoneNames()->GetAttribValues();
+		const TArray<int32>& ParentsAttrib = EditMesh.Attributes()->GetBoneParentIndices()->GetAttribValues();
+		const TArray<FTransform>& TransformsAttrib = EditMesh.Attributes()->GetBonePoses()->GetAttribValues();
+		const TArray<FVector4f>& ColorsAttrib = EditMesh.Attributes()->GetBoneColors()->GetAttribValues();
+
+		for (int32 BoneIndex = 0; BoneIndex < NumBones; ++BoneIndex)
+		{
+			FGeometryScriptBoneInfo& Entry = BonesInfo[BoneIndex];
+			
+			Entry.Index = BoneIndex;
+			Entry.Name = NamesAttrib[BoneIndex];
+			Entry.ParentIndex = ParentsAttrib[BoneIndex];
+			Entry.LocalTransform = TransformsAttrib[BoneIndex];
+			Entry.WorldTransform = Entry.LocalTransform;
+			Entry.Color = ColorsAttrib[BoneIndex];
+
+			int32 CurParentIndex = ParentsAttrib[BoneIndex];
+			while (CurParentIndex != INDEX_NONE) // until we didn't reach the root bone
+			{
+				Entry.WorldTransform = Entry.WorldTransform * TransformsAttrib[CurParentIndex];
+				CurParentIndex = ParentsAttrib[CurParentIndex];
+			}
+		}
 	});
 
 	return TargetMesh;
diff --git a/Engine/Plugins/Experimental/GeometryScripting/Source/GeometryScriptingCore/Public/GeometryScript/MeshBoneWeightFunctions.h b/Engine/Plugins/Experimental/GeometryScripting/Source/GeometryScriptingCore/Public/GeometryScript/MeshBoneWeightFunctions.h
--- a/Engine/Plugins/Experimental/GeometryScripting/Source/GeometryScriptingCore/Public/GeometryScript/MeshBoneWeightFunctions.h
+++ b/Engine/Plugins/Experimental/GeometryScripting/Source/GeometryScriptingCore/Public/GeometryScript/MeshBoneWeightFunctions.h
@@ -80,7 +80,15 @@ enum class ETransferBoneWeightsMethod : uint8
 	 * bone weights from it. This is usually a point on the SourceMesh triangle where the bone weights are computed via 
 	 * interpolation of the bone weights at the vertices of the triangle via barycentric coordinates.
 	 */
-	ClosestPointOnSurface = 0
+	ClosestPointOnSurface = 0,
+
+	/** For every vertex on the target mesh, find the closest point on the surface of the source mesh. If that point 
+	 * is within the search radius (controlled via SearchPercentage), and their normals differ by less than the 
+	 * NormalThreshold, then we directly copy the weights from the source point to the target mesh vertex 
+	 * (same as the ClosestPointOnSurface method). For all the vertices we didn't copy the weights directly, 
+	 * automatically compute the smooth weights. 
+	 */
+	InpaintWeights = 1
 };
 
 UENUM(BlueprintType)
@@ -103,7 +111,7 @@ enum class EOutputTargetMeshBones : uint8
 USTRUCT(BlueprintType)
 struct GEOMETRYSCRIPTINGCORE_API FGeometryScriptTransferBoneWeightsOptions
 {
-	GENERATED_BODY();
+	GENERATED_BODY()
 	
 	/** The type of algorithm to use for transferring the bone weights. */
 	UPROPERTY(EditAnywhere, BlueprintReadWrite, Category = Options)
@@ -120,6 +128,60 @@ struct GEOMETRYSCRIPTINGCORE_API FGeometryScriptTransferBoneWeightsOptions
 	/** The identifier for the source bone/skin weight profile used to transfer the weights to. */
 	UPROPERTY(EditAnywhere, BlueprintReadWrite, Category = Options)
 	FGeometryScriptBoneWeightProfile TargetProfile = FGeometryScriptBoneWeightProfile();
+
+	/** Defines the search radius as the RadiusPercentage * (input mesh bounding box diagonal). All points not within the search
+	  * radius will be ignored. If negative, all points are considered. Only used in the InpaintWeights algorithm.*/
+	UPROPERTY(EditAnywhere, BlueprintReadWrite, Category = Options, Meta = (UIMin = -1, UIMax = 2, ClampMin = -1, ClampMax = 2, EditCondition = "TransferMethod == ETransferBoneWeightsMethod::InpaintWeights"))
+	double RadiusPercentage = -1;
+
+	/** Maximum angle (in degrees) difference between the target and the source point normals to be considred a match. 
+	 * If negative, normals are ignored. Only used in the InpaintWeights algorithm.*/
+	UPROPERTY(EditAnywhere, BlueprintReadWrite, Category = Options, Meta = (UIMin = -1, UIMax = 180, ClampMin = -1, ClampMax = 180, EditCondition="TransferMethod == ETransferBoneWeightsMethod::InpaintWeights"))
+	double NormalThreshold = -1;
+
+	/** If true, when the closest point doesn't pass the normal threshold test, will try again with a flipped normal. 
+	 * This helps with layered meshes where the "inner" and "outer" layers are close to each other but whose normals 
+	 * are pointing in the opposite directions. */
+	UPROPERTY(EditAnywhere, BlueprintReadWrite, Category = Options, Meta = (EditCondition = "TransferMethod == ETransferBoneWeightsMethod::InpaintWeights"))
+	bool LayeredMeshSupport = true;
+
+	/** The number of optional post-processing smoothing iterations applied to the vertices without the match. */
+	UPROPERTY(EditAnywhere, BlueprintReadWrite, Category = Options, meta = (UIMin = 0, UIMax = 100, ClampMin = 0, ClampMax = 100, EditCondition = "TransferMethod == ETransferBoneWeightsMethod::InpaintWeights"))
+	int32 NumSmoothingIterations = 0; 
+
+	/** The strength of each post-processing smoothing iteration. */
+	UPROPERTY(EditAnywhere, BlueprintReadWrite, Category = Options, meta = (UIMin = 0, UIMax = 1, ClampMin = 0, ClampMax = 1, EditCondition = "TransferMethod == ETransferBoneWeightsMethod::InpaintWeights"))
+	float SmoothingStrength = 0.0f;
+};
+
+USTRUCT(BlueprintType, meta = (DisplayName = "Bone Info"))
+struct GEOMETRYSCRIPTINGCORE_API FGeometryScriptBoneInfo
+{
+	GENERATED_BODY()
+
+	/** Index of the bone in the skeletal hierarchy. */
+	UPROPERTY(BlueprintReadWrite, Category = Bone)
+	int Index = INDEX_NONE;
+
+	/** Bone name. */
+	UPROPERTY(BlueprintReadWrite, Category = Bone)
+	FName Name = NAME_None;
+	
+	/** Parent bone index. */
+	UPROPERTY(BlueprintReadWrite, Category = Bone)
+	int ParentIndex = INDEX_NONE;
+
+	/** Local/bone space reference transform. */
+	UPROPERTY(BlueprintReadWrite, Category = Bone)
+	FTransform LocalTransform = FTransform::Identity;
+
+	/** Global/world space reference transform. */
+	UPROPERTY(BlueprintReadWrite, Category = Bone)
+	FTransform WorldTransform = FTransform::Identity;
+
+	/** Bone color. */
+	UPROPERTY(BlueprintReadWrite, Category = Bone)
+	FLinearColor Color = FLinearColor::White;
 };
 
 UCLASS(meta = (ScriptName = "GeometryScript_BoneWeights"))
@@ -289,6 +351,7 @@ public:
 	 * @param bIsValidBoneName Set to true if the TargetMesh contains a bone with the given name, false otherwise.
 	 * @param BoneIndex The index of the bone with the given name. Will be set to -1 if bIsValidBoneName is false.
 	 */
+	UE_DEPRECATED(5.3, "Use GetBoneInfo instead.")
 	UFUNCTION(BlueprintCallable, Category = "GeometryScript|MeshQueries|BoneWeights", meta=(ScriptMethod))
 	static UPARAM(DisplayName = "Target Mesh") UDynamicMesh* 
 	GetBoneIndex(
@@ -297,4 +360,66 @@ public:
 		bool& bIsValidBoneName,
 		int& BoneIndex,
 		UGeometryScriptDebug* Debug = nullptr);
+
+	/**
+	 * Get the name of the root bone.
+	 * 
+	 * @param TargetMesh Mesh containing the bone attributes.
+	 * @param BoneName The name of the root bone.
+	 */
+	UFUNCTION(BlueprintCallable, Category = "GeometryScript|MeshQueries|BoneWeights", meta=(ScriptMethod))
+	static UPARAM(DisplayName = "Target Mesh") UDynamicMesh* 
+	GetRootBoneName(
+		UDynamicMesh* TargetMesh,
+		FName& BoneName,
+		UGeometryScriptDebug* Debug = nullptr);
+
+	/**
+	 * Get the information about the children of the bone.
+	 * 
+	 * @param TargetMesh Mesh containing the bone attributes.
+	 * @param BoneName The name of bone.
+	 * @param bRecursive If set to true, grandchildren will also be added recursively
+	 * @param bIsValidBoneName Set to true if the TargetMesh contains a bone with the given name, false otherwise.
+	 * @param ChildrenInfo The informattion of the children.
+	 */
+	UFUNCTION(BlueprintCallable, Category = "GeometryScript|MeshQueries|BoneWeights", meta=(ScriptMethod))
+	static UPARAM(DisplayName = "Target Mesh") UDynamicMesh* 
+	GetBoneChildren(
+		UDynamicMesh* TargetMesh,
+		FName BoneName,
+		bool bRecursive,
+		bool& bIsValidBoneName,
+		TArray<FGeometryScriptBoneInfo>& ChildrenInfo,
+		UGeometryScriptDebug* Debug = nullptr);
+
+	/**
+	 * Get the bone information.
+	 * 
+	 * @param TargetMesh Mesh containing the bone attributes.
+	 * @param BoneName The name of bone.
+	 * @param bIsValidBoneName Set to true if the TargetMesh contains a bone with the given name, false otherwise.
+	 * @param BoneInfo The information about the bone.
+	 */
+	UFUNCTION(BlueprintCallable, Category = "GeometryScript|MeshQueries|BoneWeights", meta=(ScriptMethod))
+	static UPARAM(DisplayName = "Target Mesh") UDynamicMesh* 
+	GetBoneInfo(
+		UDynamicMesh* TargetMesh,
+		FName BoneName,
+		bool& bIsValidBoneName,
+		FGeometryScriptBoneInfo& BoneInfo,
+		UGeometryScriptDebug* Debug = nullptr);
+
+	/**
+	 * Get an array of bones representing the skeleton. Each entry contains information about the bone.
+	 * 
+	 * @param TargetMesh Mesh containing the bone attributes.
+	 * @param BonesInfo Skeleton information.
+	 */
+	UFUNCTION(BlueprintCallable, Category = "GeometryScript|MeshQueries|BoneWeights", meta=(ScriptMethod))
+	static UPARAM(DisplayName = "Target Mesh") UDynamicMesh* 
+	GetAllBonesInfo(
+		UDynamicMesh* TargetMesh,
+		TArray<FGeometryScriptBoneInfo>& BonesInfo,
+		UGeometryScriptDebug* Debug = nullptr);
 };
\ No newline at end of file
