GeometryScript: add SelectMeshElementsInBoxWithBVH, allows efficient range query on large meshes. Fix issue in FGeometryScriptMeshSelection::ConvertToMeshIndexArray where if converting from vertex to triangle selection, it was emitting a list with duplicate triangles.

diff --git a/Engine/Plugins/Experimental/GeometryScripting/Source/GeometryScriptingCore/Private/GeometryScriptSelectionTypes.cpp b/Engine/Plugins/Experimental/GeometryScripting/Source/GeometryScriptingCore/Private/GeometryScriptSelectionTypes.cpp
--- a/Engine/Plugins/Experimental/GeometryScripting/Source/GeometryScriptingCore/Private/GeometryScriptSelectionTypes.cpp
+++ b/Engine/Plugins/Experimental/GeometryScripting/Source/GeometryScriptingCore/Private/GeometryScriptSelectionTypes.cpp
@@ -271,11 +271,17 @@ EGeometryScriptIndexType FGeometryScriptMeshSelection::ConvertToMeshIndexArray(
 	}
 	else if (ConvertToType == EGeometryScriptIndexType::Triangle)
 	{
+		TSet<int32> UniqueTriangles;
+
 		if ( (GeoSelection->TopologyType == EGeometryTopologyType::Triangle) &&
 			 (GeoSelection->ElementType == EGeometryElementType::Vertex) )
 		{
 			UE::Geometry::EnumerateSelectionTriangles(*GeoSelection, Mesh,
-				[&](int32 TriangleID) { IndexListOut.Add(TriangleID); });
+				[&](int32 TriangleID) { UniqueTriangles.Add(TriangleID); });
+			for (int32 tid : UniqueTriangles)
+			{
+				IndexListOut.Add(tid);
+			}
 			return EGeometryScriptIndexType::Triangle;
 		}
 		else if (GeoSelection->TopologyType == EGeometryTopologyType::Polygroup && GeoSelection->ElementType == EGeometryElementType::Face)
diff --git a/Engine/Plugins/Experimental/GeometryScripting/Source/GeometryScriptingCore/Private/MeshSpatialFunctions.cpp b/Engine/Plugins/Experimental/GeometryScripting/Source/GeometryScriptingCore/Private/MeshSpatialFunctions.cpp
--- a/Engine/Plugins/Experimental/GeometryScripting/Source/GeometryScriptingCore/Private/MeshSpatialFunctions.cpp
+++ b/Engine/Plugins/Experimental/GeometryScripting/Source/GeometryScriptingCore/Private/MeshSpatialFunctions.cpp
@@ -2,12 +2,12 @@
 
 #include "GeometryScript/MeshSpatialFunctions.h"
 #include "UDynamicMesh.h"
-
 #include "DynamicMesh/DynamicMesh3.h"
 #include "DynamicMesh/DynamicMeshAttributeSet.h"
 #include "DynamicMesh/DynamicMeshAABBTree3.h"
 #include "MeshQueries.h"
 #include "Spatial/FastWinding.h"
+#include "Selections/GeometrySelection.h"
 
 #include UE_INLINE_GENERATED_CPP_BY_NAME(MeshSpatialFunctions)
 
@@ -274,6 +274,126 @@ UDynamicMesh* UGeometryScriptLibrary_MeshSpatial::IsPointInsideMesh(
 
 
 
+UDynamicMesh* UGeometryScriptLibrary_MeshSpatial::SelectMeshElementsInBoxWithBVH(
+	UDynamicMesh* TargetMesh,
+	UPARAM(ref) const FGeometryScriptDynamicMeshBVH& QueryBVH,
+	FBox QueryBox,
+	FGeometryScriptSpatialQueryOptions Options,
+	FGeometryScriptMeshSelection& SelectionOut,
+	EGeometryScriptMeshSelectionType SelectionType,
+	int MinNumTrianglePoints,
+	UGeometryScriptDebug* Debug)
+{
+	MinNumTrianglePoints = FMath::Clamp(MinNumTrianglePoints, 1, 3);
+
+	if (TargetMesh == nullptr)
+	{
+		UE::Geometry::AppendError(Debug, EGeometryScriptErrorType::InvalidInputs, LOCTEXT("SelectMeshElementsInBoxWithBVH_InvalidInput", "SelectMeshElementsInBoxWithBVH: TargetMesh is Null"));
+		return TargetMesh;
+	}
+
+	FAxisAlignedBox3d QueryBounds(QueryBox);
+	IMeshSpatial::FQueryOptions QueryOptions;
+	QueryOptions.MaxDistance = (Options.MaxDistance == 0) ? TNumericLimits<float>::Max() : Options.MaxDistance;
+
+	FGeometrySelection GeoSelection;
+
+	auto TriangleIntersectionTest = [](const FDynamicMesh3& Mesh, int32 TriangleID, const FAxisAlignedBox3d& Box, int MinNumTrianglePoints) -> bool
+	{
+		FIndex3i Tri = Mesh.GetTriangle(TriangleID);
+		int NumContained =
+			(Box.Contains(Mesh.GetVertex(Tri.A)) ? 1 : 0) +
+			(Box.Contains(Mesh.GetVertex(Tri.B)) ? 1 : 0) +
+			(Box.Contains(Mesh.GetVertex(Tri.C)) ? 1 : 0);
+		return (NumContained >= MinNumTrianglePoints);
+	};
+
+	bool bHaveValidBVH = false;
+	TargetMesh->ProcessMesh([&](const FDynamicMesh3& ReadMesh)
+	{
+		if (QueryBVH.Spatial.IsValid())
+		{
+			TargetMesh->ProcessMesh([&](const FDynamicMesh3& ReadMesh)
+			{
+				if (QueryBVH.Spatial->GetMesh() == &ReadMesh && QueryBVH.Spatial->IsValid(Options.bAllowUnsafeModifiedQueries))
+				{
+					if (SelectionType == EGeometryScriptMeshSelectionType::Vertices)
+					{
+						GeoSelection.InitializeTypes(EGeometryElementType::Vertex, EGeometryTopologyType::Triangle);
+
+						FDynamicMeshAABBTree3::FTreeTraversal Traversal;
+						Traversal.NextBoxF = [QueryBounds](const FAxisAlignedBox3d& Box, int Depth) { return Box.Intersects(QueryBounds); };
+						Traversal.NextTriangleF = [&ReadMesh,&GeoSelection,QueryBounds](int TriangleID) 
+						{
+							FIndex3i Tri = ReadMesh.GetTriangle(TriangleID);
+							for (int32 k = 0; k < 3; ++k)
+							{
+								if ( QueryBounds.Contains( ReadMesh.GetVertex(Tri[k])) )
+								{
+									GeoSelection.Selection.Add(FGeoSelectionID::MeshVertex(Tri[k]).Encoded());
+								}
+							}
+						};
+						QueryBVH.Spatial->DoTraversal(Traversal, QueryOptions);
+					}
+					else if (SelectionType == EGeometryScriptMeshSelectionType::Triangles)
+					{
+						GeoSelection.InitializeTypes(EGeometryElementType::Face, EGeometryTopologyType::Triangle);
+						FDynamicMeshAABBTree3::FTreeTraversal Traversal;
+						Traversal.NextBoxF = [QueryBounds](const FAxisAlignedBox3d& Box, int Depth) { return Box.Intersects(QueryBounds); };
+						Traversal.NextTriangleF = [&](int TriangleID) 
+						{
+							if (TriangleIntersectionTest(ReadMesh, TriangleID, QueryBounds, MinNumTrianglePoints))
+							{
+								GeoSelection.Selection.Add(FGeoSelectionID::MeshTriangle(TriangleID).Encoded());
+							}
+						};
+						QueryBVH.Spatial->DoTraversal(Traversal, QueryOptions);
+					}
+					else if (SelectionType == EGeometryScriptMeshSelectionType::Polygroups)
+					{
+						TArray<int32> UniqueGroups;
+						TArray<FIndex2i> GroupTriPairs;
+						FDynamicMeshAABBTree3::FTreeTraversal Traversal;
+						Traversal.NextBoxF = [QueryBounds](const FAxisAlignedBox3d& Box, int Depth) { return Box.Intersects(QueryBounds); };
+						Traversal.NextTriangleF = [&](int TriangleID) 
+						{
+							if (TriangleIntersectionTest(ReadMesh, TriangleID, QueryBounds, MinNumTrianglePoints))
+							{
+								int32 GroupID = ReadMesh.GetTriangleGroup(TriangleID);
+								if (UniqueGroups.Contains(GroupID) == false)
+								{
+									UniqueGroups.Add(GroupID);
+									GroupTriPairs.Add(FIndex2i(TriangleID, GroupID));
+								}
+							}
+						};
+						QueryBVH.Spatial->DoTraversal(Traversal, QueryOptions);
+
+						GeoSelection.InitializeTypes(EGeometryElementType::Face, EGeometryTopologyType::Polygroup);
+						for (FIndex2i Pair : GroupTriPairs)
+						{
+							GeoSelection.Selection.Add(FGeoSelectionID::GroupFace(Pair.A, Pair.B).Encoded());
+						}
+					}
+				}
+			});
+		}
+
+	});
+
+	if (!bHaveValidBVH)
+	{
+		UE::Geometry::AppendError(Debug, EGeometryScriptErrorType::InvalidInputs, LOCTEXT("IsPointInsideMesh_InvalidBVH", "IsPointInsideMesh: QueryBVH is Invalid for this TargetMesh"));
+	}
+
+	SelectionOut.SetSelection(MoveTemp(GeoSelection));
+
+
+	return TargetMesh;
+}
+
+
 
 #undef LOCTEXT_NAMESPACE
 
diff --git a/Engine/Plugins/Experimental/GeometryScripting/Source/GeometryScriptingCore/Public/GeometryScript/MeshSpatialFunctions.h b/Engine/Plugins/Experimental/GeometryScripting/Source/GeometryScriptingCore/Public/GeometryScript/MeshSpatialFunctions.h
--- a/Engine/Plugins/Experimental/GeometryScripting/Source/GeometryScriptingCore/Public/GeometryScript/MeshSpatialFunctions.h
+++ b/Engine/Plugins/Experimental/GeometryScripting/Source/GeometryScriptingCore/Public/GeometryScript/MeshSpatialFunctions.h
@@ -5,6 +5,7 @@
 #include "CoreMinimal.h"
 #include "Kismet/BlueprintFunctionLibrary.h"
 #include "GeometryScript/GeometryScriptTypes.h"
+#include "GeometryScript/GeometryScriptSelectionTypes.h"
 #include "MeshSpatialFunctions.generated.h"
 
 class UDynamicMesh;
@@ -150,4 +151,28 @@ public:
 		EGeometryScriptContainmentOutcomePins& Outcome,
 		UGeometryScriptDebug* Debug = nullptr );
 
+
+	/**
+	* Create Mesh Selection of mesh elements in TargetMesh contained by QueryBox, using QueryBVH to accellerate the computation.
+	* Triangles are selected if MinNumTrianglePoints or more vertices are inside the box. 
+	* PolyGroups are selected if any of their triangles are inside the box
+	* 
+	* @param QueryBVH is an acceleration structure previously built with TargetMesh.
+	* @param QueryPoint the point in the mesh's 3D local space.
+	* @param Options control the fast winding number threshold 
+	*/
+	UFUNCTION(BlueprintCallable, Category = "GeometryScript|Spatial", meta=(ScriptMethod, AdvancedDisplay = "MinNumTrianglePoints"))
+	static UPARAM(DisplayName = "Target Mesh") UDynamicMesh* 
+	SelectMeshElementsInBoxWithBVH(
+		UDynamicMesh* TargetMesh,
+		UPARAM(ref) const FGeometryScriptDynamicMeshBVH& QueryBVH,
+		FBox QueryBox,
+		FGeometryScriptSpatialQueryOptions Options,
+		FGeometryScriptMeshSelection& Selection,
+		EGeometryScriptMeshSelectionType SelectionType = EGeometryScriptMeshSelectionType::Vertices,
+		int MinNumTrianglePoints = 3,
+		UGeometryScriptDebug* Debug = nullptr );
+
+
+
 };
\ No newline at end of file
