Fixed a bug where the GetBoneChildren will return the whole skeleton and not just the children.

diff --git a/Engine/Plugins/Experimental/GeometryScripting/Source/GeometryScriptingCore/Private/MeshBoneWeightFunctions.cpp b/Engine/Plugins/Experimental/GeometryScripting/Source/GeometryScriptingCore/Private/MeshBoneWeightFunctions.cpp
--- a/Engine/Plugins/Experimental/GeometryScripting/Source/GeometryScriptingCore/Private/MeshBoneWeightFunctions.cpp
+++ b/Engine/Plugins/Experimental/GeometryScripting/Source/GeometryScriptingCore/Private/MeshBoneWeightFunctions.cpp
@@ -667,10 +667,11 @@ UDynamicMesh* UGeometryScriptLibrary_MeshBoneWeightFunctions::GetBoneChildren(
 			}
 			
 			// Get all information about the children
-			ChildrenInfo.SetNum(NumBones);
-			for (const int ChildIdx : ChildrenIndices)
+			ChildrenInfo.SetNum(ChildrenIndices.Num());
+			for (int32 Idx = 0; Idx < ChildrenIndices.Num(); ++Idx)
 			{
-				FGeometryScriptBoneInfo& Entry = ChildrenInfo[ChildIdx];
+				const int32 ChildIdx = ChildrenIndices[Idx];
+				FGeometryScriptBoneInfo& Entry = ChildrenInfo[Idx];
 				Entry.Index = ChildIdx;
 				Entry.Name = NamesAttrib[ChildIdx];
 				Entry.ParentIndex = ParentsAttrib[ChildIdx];
