GeometryScript: add GetVertexConnectedTriangles, GetVertexConnectedVertices, GetAllSplitUVsAtVertex

diff --git a/Engine/Plugins/Experimental/GeometryScripting/Source/GeometryScriptingCore/Private/MeshQueryFunctions.cpp b/Engine/Plugins/Experimental/GeometryScripting/Source/GeometryScriptingCore/Private/MeshQueryFunctions.cpp
--- a/Engine/Plugins/Experimental/GeometryScripting/Source/GeometryScriptingCore/Private/MeshQueryFunctions.cpp
+++ b/Engine/Plugins/Experimental/GeometryScripting/Source/GeometryScriptingCore/Private/MeshQueryFunctions.cpp
@@ -371,7 +371,24 @@ UDynamicMesh* UGeometryScriptLibrary_MeshQueryFunctions::GetAllVertexPositions(U
 }
 
 
+UDynamicMesh* UGeometryScriptLibrary_MeshQueryFunctions::GetVertexConnectedTriangles(UDynamicMesh* TargetMesh, int32 VertexID, TArray<int32>& Triangles)
+{
+	SimpleMeshQuery<int32>(TargetMesh, 0, [&](const FDynamicMesh3& Mesh) {
+		Mesh.GetVtxTriangles(VertexID, Triangles);
+		return 0;
+	});
+	return TargetMesh;
+}
+
 
+UDynamicMesh* UGeometryScriptLibrary_MeshQueryFunctions::GetVertexConnectedVertices(UDynamicMesh* TargetMesh, int32 VertexID, TArray<int32>& Vertices)
+{
+	SimpleMeshQuery<int32>(TargetMesh, 0, [&](const FDynamicMesh3& Mesh) {
+		Mesh.EnumerateVertexVertices(VertexID, [&](int32 vid) { Vertices.Add(vid); });
+		return 0;
+	});
+	return TargetMesh;
+}
 
 
 int32 UGeometryScriptLibrary_MeshQueryFunctions::GetNumUVSets( UDynamicMesh* TargetMesh )
@@ -423,6 +440,28 @@ void UGeometryScriptLibrary_MeshQueryFunctions::GetTriangleUVs(UDynamicMesh* Tar
 	});
 }
 
+UDynamicMesh* UGeometryScriptLibrary_MeshQueryFunctions::GetAllSplitUVsAtVertex(UDynamicMesh* TargetMesh, int32 UVSetIndex, int32 VertexID,
+	TArray<int32>& ElementIDs, TArray<FVector2D>& ElementUVs, bool& bHaveValidUVs)
+{
+	bHaveValidUVs = SimpleMeshUVSetQuery<bool>(TargetMesh, UVSetIndex, bHaveValidUVs, false, 
+		[&](const FDynamicMesh3& Mesh, const FDynamicMeshUVOverlay& UVSet) 
+	{
+		if (Mesh.IsVertex(VertexID))
+		{
+			UVSet.EnumerateVertexElements(VertexID, [&](int TriangleID, int ElementID, const FVector2f& UVValue)
+			{
+				ElementIDs.Add(ElementID);
+				ElementUVs.Add((FVector2D)UVValue);
+				return true;
+			});
+			return true;
+		}
+		return false;
+	});
+	return TargetMesh;
+}
+
+
 UDynamicMesh* UGeometryScriptLibrary_MeshQueryFunctions::GetInterpolatedTriangleUV(UDynamicMesh* TargetMesh, UPARAM(DisplayName = "UV Channel") int32 UVSetIndex, int32 TriangleID, FVector BarycentricCoords, bool& bHaveValidUVs, FVector2D& InterpolatedUV)
 {
 	InterpolatedUV = FVector2D::Zero();
diff --git a/Engine/Plugins/Experimental/GeometryScripting/Source/GeometryScriptingCore/Public/GeometryScript/MeshQueryFunctions.h b/Engine/Plugins/Experimental/GeometryScripting/Source/GeometryScriptingCore/Public/GeometryScript/MeshQueryFunctions.h
--- a/Engine/Plugins/Experimental/GeometryScripting/Source/GeometryScriptingCore/Public/GeometryScript/MeshQueryFunctions.h
+++ b/Engine/Plugins/Experimental/GeometryScripting/Source/GeometryScriptingCore/Public/GeometryScript/MeshQueryFunctions.h
@@ -218,6 +218,20 @@ public:
 	GetAllVertexPositions( UDynamicMesh* TargetMesh, FGeometryScriptVectorList& PositionList, bool bSkipGaps, bool& bHasVertexIDGaps );
 
 
+	/**
+	 * Return array of Triangle IDs connected to the given VertexID, ie the triangle one-ring
+	 */
+	UFUNCTION(BlueprintCallable, Category = "GeometryScript|MeshQueries", meta=(ScriptMethod))
+	static UPARAM(DisplayName = "Target Mesh") UDynamicMesh* 
+	GetVertexConnectedTriangles( UDynamicMesh* TargetMesh, int32 VertexID, TArray<int32>& Triangles);
+
+	/**
+	 * Return array of Vertex IDs connected via a mesh edge to the given VertexID, ie the vertex one-ring
+	 */
+	UFUNCTION(BlueprintCallable, Category = "GeometryScript|MeshQueries", meta=(ScriptMethod))
+	static UPARAM(DisplayName = "Target Mesh") UDynamicMesh* 
+	GetVertexConnectedVertices( UDynamicMesh* TargetMesh, int32 VertexID, TArray<int32>& Vertices);
+
 	//
 	// UV Queries
 	//
@@ -243,6 +257,15 @@ public:
 	UFUNCTION(BlueprintPure, Category = "GeometryScript|MeshQueries", meta=(ScriptMethod))
 	static void GetTriangleUVs( UDynamicMesh* TargetMesh, UPARAM(DisplayName = "UV Channel") int32 UVSetIndex, int32 TriangleID, FVector2D& UV1, FVector2D& UV2, FVector2D& UV3, bool& bHaveValidUVs );
 
+
+	/**
+	* Returns the unique UV element IDs and values associated with the mesh vertex, in the specified UV Channel.
+	* If the Vertex or UV channel does not exist, the arrays will be empty and bHaveValidUVs will be set to false.
+	*/
+	UFUNCTION(BlueprintCallable, Category = "GeometryScript|MeshQueries", meta=(ScriptMethod))
+	static UPARAM(DisplayName = "Target Mesh") UDynamicMesh* 
+	GetAllSplitUVsAtVertex( UDynamicMesh* TargetMesh, UPARAM(DisplayName = "UV Channel") int32 UVSetIndex, int32 VertexID, TArray<int32>& ElementIDs, TArray<FVector2D>& ElementUVs, bool& bHaveValidUVs);
+
 	/**
 	 * Compute the interpolated UV (A*UV1+ B*UV2+ C*UV3), where (A,B,C)=BarycentricCoords and the UV positions are taken
 	 * from the specified TriangleID in the specified UVSet of the TargetMesh. 
