Change the Geometry Script plugin from Experimental to Beta

diff --git a/Engine/Plugins/Experimental/GeometryScripting/GeometryScripting.uplugin b/Engine/Plugins/Experimental/GeometryScripting/GeometryScripting.uplugin
--- a/Engine/Plugins/Experimental/GeometryScripting/GeometryScripting.uplugin
+++ b/Engine/Plugins/Experimental/GeometryScripting/GeometryScripting.uplugin
@@ -11,7 +11,8 @@
 	"MarketplaceURL": "",
 	"SupportURL": "",
 	"CanContainContent": true,
-	"IsExperimentalVersion": true,
+	"IsExperimentalVersion": false,
+	"IsBetaVersion":  true,
 	"Installed": false,
 	"Modules": [
 		{
