GeometryScript: support reading from HiRes LOD in CopyMeshFromStaticMesh

diff --git a/Engine/Plugins/Experimental/GeometryScripting/Source/GeometryScriptingCore/Private/MeshAssetFunctions.cpp b/Engine/Plugins/Experimental/GeometryScripting/Source/GeometryScriptingCore/Private/MeshAssetFunctions.cpp
--- a/Engine/Plugins/Experimental/GeometryScripting/Source/GeometryScriptingCore/Private/MeshAssetFunctions.cpp
+++ b/Engine/Plugins/Experimental/GeometryScripting/Source/GeometryScriptingCore/Private/MeshAssetFunctions.cpp
@@ -47,27 +47,45 @@ static UDynamicMesh* CopyMeshFromStaticMesh_SourceData(
 	UGeometryScriptDebug* Debug
 )
 {
-	if (RequestedLOD.LODType != EGeometryScriptLODType::MaxAvailable && RequestedLOD.LODType != EGeometryScriptLODType::SourceModel)
+	if (RequestedLOD.LODType != EGeometryScriptLODType::MaxAvailable && RequestedLOD.LODType != EGeometryScriptLODType::SourceModel && RequestedLOD.LODType != EGeometryScriptLODType::HiResSourceModel)
 	{
 		UE::Geometry::AppendError(Debug, EGeometryScriptErrorType::InvalidInputs, LOCTEXT("CopyMeshFromStaticMesh_LODNotAvailable", "CopyMeshFromStaticMesh: Requested LOD Type is not available"));
 		return ToDynamicMesh;
 	}
 
 #if WITH_EDITOR
-	int32 UseLODIndex = FMath::Clamp(RequestedLOD.LODIndex, 0, FromStaticMeshAsset->GetNumSourceModels() - 1);
+	if (RequestedLOD.LODType == EGeometryScriptLODType::HiResSourceModel && FromStaticMeshAsset->IsHiResMeshDescriptionValid() == false)
+	{
+		UE::Geometry::AppendError(Debug, EGeometryScriptErrorType::InvalidInputs, LOCTEXT("CopyMeshFromStaticMesh_HiResLODNotAvailable", "CopyMeshFromStaticMesh: HiResSourceModel LOD Type is not available"));
+		return ToDynamicMesh;
+	}
+
+	const FMeshDescription* SourceMesh = nullptr;
+	const FMeshBuildSettings* BuildSettings = nullptr;
+
+	if ((RequestedLOD.LODType == EGeometryScriptLODType::HiResSourceModel) ||
+		(RequestedLOD.LODType == EGeometryScriptLODType::MaxAvailable && FromStaticMeshAsset->IsHiResMeshDescriptionValid()))
+	{
+		SourceMesh = FromStaticMeshAsset->GetHiResMeshDescription();
+		const FStaticMeshSourceModel& SourceModel = FromStaticMeshAsset->GetHiResSourceModel();
+		BuildSettings = &SourceModel.BuildSettings;
+	}
+	else
+	{
+		int32 UseLODIndex = FMath::Clamp(RequestedLOD.LODIndex, 0, FromStaticMeshAsset->GetNumSourceModels() - 1);
+		SourceMesh = FromStaticMeshAsset->GetMeshDescription(UseLODIndex);
+		const FStaticMeshSourceModel& SourceModel = FromStaticMeshAsset->GetSourceModel(UseLODIndex);
+		BuildSettings = &SourceModel.BuildSettings;
+	}
 
-	const FMeshDescription* SourceMesh = FromStaticMeshAsset->GetMeshDescription(UseLODIndex);
 	if (SourceMesh == nullptr)
 	{
 		UE::Geometry::AppendError(Debug, EGeometryScriptErrorType::InvalidInputs, LOCTEXT("CopyMeshFromStaticMesh_SourceLODIsNull", "CopyMeshFromStaticMesh: Requested SourceModel LOD is null, only RenderData Mesh is available"));
 		return ToDynamicMesh;
 	}
 
-	const FStaticMeshSourceModel& SourceModel = FromStaticMeshAsset->GetSourceModel(UseLODIndex);
-	const FMeshBuildSettings& BuildSettings = SourceModel.BuildSettings;
-
-	bool bHasDirtyBuildSettings = BuildSettings.bRecomputeNormals
-		|| (BuildSettings.bRecomputeTangents && AssetOptions.bRequestTangents);
+	bool bHasDirtyBuildSettings = BuildSettings->bRecomputeNormals
+		|| (BuildSettings->bRecomputeTangents && AssetOptions.bRequestTangents);
 
 	FMeshDescription LocalSourceMeshCopy;
 	if (AssetOptions.bApplyBuildSettings && bHasDirtyBuildSettings )
@@ -82,16 +100,16 @@ static UDynamicMesh* CopyMeshFromStaticMesh_SourceData(
 		}
 
 		EComputeNTBsFlags ComputeNTBsOptions = EComputeNTBsFlags::BlendOverlappingNormals;
-		ComputeNTBsOptions |= BuildSettings.bRecomputeNormals ? EComputeNTBsFlags::Normals : EComputeNTBsFlags::None;
+		ComputeNTBsOptions |= BuildSettings->bRecomputeNormals ? EComputeNTBsFlags::Normals : EComputeNTBsFlags::None;
 		if (AssetOptions.bRequestTangents)
 		{
-			ComputeNTBsOptions |= BuildSettings.bRecomputeTangents ? EComputeNTBsFlags::Tangents : EComputeNTBsFlags::None;
-			ComputeNTBsOptions |= BuildSettings.bUseMikkTSpace ? EComputeNTBsFlags::UseMikkTSpace : EComputeNTBsFlags::None;
+			ComputeNTBsOptions |= BuildSettings->bRecomputeTangents ? EComputeNTBsFlags::Tangents : EComputeNTBsFlags::None;
+			ComputeNTBsOptions |= BuildSettings->bUseMikkTSpace ? EComputeNTBsFlags::UseMikkTSpace : EComputeNTBsFlags::None;
 		}
-		ComputeNTBsOptions |= BuildSettings.bComputeWeightedNormals ? EComputeNTBsFlags::WeightedNTBs : EComputeNTBsFlags::None;
+		ComputeNTBsOptions |= BuildSettings->bComputeWeightedNormals ? EComputeNTBsFlags::WeightedNTBs : EComputeNTBsFlags::None;
 		if (AssetOptions.bIgnoreRemoveDegenerates == false)
 		{
-			ComputeNTBsOptions |= BuildSettings.bRemoveDegenerates ? EComputeNTBsFlags::IgnoreDegenerateTriangles : EComputeNTBsFlags::None;
+			ComputeNTBsOptions |= BuildSettings->bRemoveDegenerates ? EComputeNTBsFlags::IgnoreDegenerateTriangles : EComputeNTBsFlags::None;
 		}
 
 		FStaticMeshOperations::ComputeTangentsAndNormals(LocalSourceMeshCopy, ComputeNTBsOptions);
