GeometryScript: add new functions AddUVElementToMesh, SetMeshTriangleUVElementIDs, GetMeshTriangleUVElementIDs, GetMeshUVElementPosition, SetMeshUVElementPosition

diff --git a/Engine/Plugins/Experimental/GeometryScripting/Source/GeometryScriptingCore/Private/MeshUVFunctions.cpp b/Engine/Plugins/Experimental/GeometryScripting/Source/GeometryScriptingCore/Private/MeshUVFunctions.cpp
--- a/Engine/Plugins/Experimental/GeometryScripting/Source/GeometryScriptingCore/Private/MeshUVFunctions.cpp
+++ b/Engine/Plugins/Experimental/GeometryScripting/Source/GeometryScriptingCore/Private/MeshUVFunctions.cpp
@@ -127,6 +127,164 @@ UDynamicMesh* UGeometryScriptLibrary_MeshUVFunctions::SetMeshTriangleUVs(
 
 
 
+UDynamicMesh* UGeometryScriptLibrary_MeshUVFunctions::AddUVElementToMesh(
+	UDynamicMesh* TargetMesh,
+	int UVSetIndex,
+	FVector2D NewUVPosition,
+	int& NewUVElementID,
+	bool& bIsValidUVSet,
+	bool bDeferChangeNotifications)
+{
+	bIsValidUVSet = false;
+	NewUVElementID = -1;
+	if (TargetMesh)
+	{
+		TargetMesh->EditMesh([&](FDynamicMesh3& EditMesh)
+		{
+			if (EditMesh.HasAttributes() && UVSetIndex < EditMesh.Attributes()->NumUVLayers() )
+			{
+				FDynamicMeshUVOverlay* UVOverlay = EditMesh.Attributes()->GetUVLayer(UVSetIndex);
+				if (UVOverlay != nullptr)
+				{
+					bIsValidUVSet = true;
+					NewUVElementID = UVOverlay->AppendElement((FVector2f)NewUVPosition);
+				}
+			}
+		}, EDynamicMeshChangeType::GeneralEdit, EDynamicMeshAttributeChangeFlags::Unknown, bDeferChangeNotifications);
+	}
+	return TargetMesh;	
+}
+
+
+
+
+UDynamicMesh* UGeometryScriptLibrary_MeshUVFunctions::SetMeshTriangleUVElementIDs(
+	UDynamicMesh* TargetMesh,
+	int UVSetIndex,
+	int TriangleID, 
+	FIntVector TriangleUVElements,
+	bool& bIsValidTriangle,
+	bool bDeferChangeNotifications)
+{
+	bIsValidTriangle = false;
+	if (TargetMesh)
+	{
+		TargetMesh->EditMesh([&](FDynamicMesh3& EditMesh)
+		{
+			if (EditMesh.IsTriangle(TriangleID) && EditMesh.HasAttributes() && UVSetIndex < EditMesh.Attributes()->NumUVLayers() )
+			{
+				FDynamicMeshUVOverlay* UVOverlay = EditMesh.Attributes()->GetUVLayer(UVSetIndex);
+				if (UVOverlay != nullptr)
+				{
+					// sanity check here because SetTriangle does not
+					FIndex3i MeshTri = EditMesh.GetTriangle(TriangleID);
+					for (int32 j = 0; j < 3; ++j)
+					{
+						int32 VertexID = MeshTri[j];
+						int32 ElemID = TriangleUVElements[j];
+						int32 ParentVertexID = UVOverlay->GetParentVertex(ElemID);
+						if (ParentVertexID != IndexConstants::InvalidID && ParentVertexID != VertexID)
+						{
+							return;		// would create broken topology
+						}
+						
+					}
+
+					if (UVOverlay->SetTriangle(TriangleID, (FIndex3i)TriangleUVElements, false) == EMeshResult::Ok)
+					{
+						bIsValidTriangle = true;
+					}
+				}
+			}
+		}, EDynamicMeshChangeType::GeneralEdit, EDynamicMeshAttributeChangeFlags::Unknown, bDeferChangeNotifications);
+	}
+	return TargetMesh;	
+}
+
+
+
+UDynamicMesh* UGeometryScriptLibrary_MeshUVFunctions::GetMeshTriangleUVElementIDs(
+	UDynamicMesh* TargetMesh,
+	int UVSetIndex,
+	int TriangleID,
+	FIntVector& TriangleUVElements,
+	bool& bHaveValidUVs)
+{
+	bHaveValidUVs = false;
+	if (TargetMesh)
+	{
+		TargetMesh->ProcessMesh([&](const FDynamicMesh3& EditMesh)
+		{
+			if (EditMesh.IsTriangle(TriangleID) && EditMesh.HasAttributes() && UVSetIndex < EditMesh.Attributes()->NumUVLayers() )
+			{
+				const FDynamicMeshUVOverlay* UVOverlay = EditMesh.Attributes()->GetUVLayer(UVSetIndex);
+				if (UVOverlay != nullptr && UVOverlay->IsSetTriangle(TriangleID))
+				{
+					bHaveValidUVs = true;
+					TriangleUVElements = UVOverlay->GetTriangle(TriangleID);
+				}
+			}
+		});
+	}
+	return TargetMesh;	
+}
+
+
+
+UDynamicMesh* UGeometryScriptLibrary_MeshUVFunctions::GetMeshUVElementPosition(
+	UDynamicMesh* TargetMesh,
+	int UVSetIndex,
+	int ElementID,
+	FVector2D& UVPosition,
+	bool& bIsValidElementID)
+{
+	bIsValidElementID = false;
+	if (TargetMesh)
+	{
+		TargetMesh->ProcessMesh([&](const FDynamicMesh3& EditMesh)
+		{
+			if (EditMesh.HasAttributes() && UVSetIndex < EditMesh.Attributes()->NumUVLayers() )
+			{
+				const FDynamicMeshUVOverlay* UVOverlay = EditMesh.Attributes()->GetUVLayer(UVSetIndex);
+				if (UVOverlay != nullptr && UVOverlay->IsElement(ElementID))
+				{
+					bIsValidElementID = true;
+					UVPosition = (FVector2D)UVOverlay->GetElement(ElementID);
+				}
+			}
+		});
+	}
+	return TargetMesh;	
+}
+
+
+
+UDynamicMesh* UGeometryScriptLibrary_MeshUVFunctions::SetMeshUVElementPosition(
+	UDynamicMesh* TargetMesh,
+	int UVSetIndex,
+	int ElementID,
+	FVector2D NewUVPosition,
+	bool& bIsValidElementID,
+	bool bDeferChangeNotifications)
+{
+	bIsValidElementID = false;
+	if (TargetMesh)
+	{
+		TargetMesh->EditMesh([&](FDynamicMesh3& EditMesh)
+		{
+			if (EditMesh.HasAttributes() && UVSetIndex < EditMesh.Attributes()->NumUVLayers() )
+			{
+				FDynamicMeshUVOverlay* UVOverlay = EditMesh.Attributes()->GetUVLayer(UVSetIndex);
+				if (UVOverlay != nullptr && UVOverlay->IsElement(ElementID))
+				{
+					bIsValidElementID = true;
+					UVOverlay->SetElement(ElementID, (FVector2f)NewUVPosition);
+				}
+			}
+		}, EDynamicMeshChangeType::GeneralEdit, EDynamicMeshAttributeChangeFlags::Unknown, bDeferChangeNotifications);
+	}
+	return TargetMesh;	
+}
 
 
 void ApplyMeshUVEditorOperation(UDynamicMesh* TargetMesh, int32 UVSetIndex, bool& bHasUVSet, UGeometryScriptDebug* Debug,
diff --git a/Engine/Plugins/Experimental/GeometryScripting/Source/GeometryScriptingCore/Public/GeometryScript/MeshUVFunctions.h b/Engine/Plugins/Experimental/GeometryScripting/Source/GeometryScriptingCore/Public/GeometryScript/MeshUVFunctions.h
--- a/Engine/Plugins/Experimental/GeometryScripting/Source/GeometryScriptingCore/Public/GeometryScript/MeshUVFunctions.h
+++ b/Engine/Plugins/Experimental/GeometryScripting/Source/GeometryScriptingCore/Public/GeometryScript/MeshUVFunctions.h
@@ -168,8 +168,10 @@ public:
 		UGeometryScriptDebug* Debug = nullptr );
 
 	/**
-	* Sets the UVs of a mesh triangle in the given UV Channel.
-	*/
+	 * Sets the UVs of a mesh triangle in the given UV Channel. 
+	 * This function will create new UV elements for each vertex of the triangle, meaning that
+	 * the triangle will become an isolated UV island.
+	 */
 	UFUNCTION(BlueprintCallable, Category = "GeometryScript|UVs", meta=(ScriptMethod))
 	static UPARAM(DisplayName = "Target Mesh") UDynamicMesh* 
 	SetMeshTriangleUVs( 
@@ -180,6 +182,79 @@ public:
 		bool& bIsValidTriangle, 
 		bool bDeferChangeNotifications = false );
 
+
+	/**
+	 * Adds a new UV Element to the specified UV Channel of the Mesh and returns a new UV Element ID.
+	 */
+	UFUNCTION(BlueprintCallable, Category = "GeometryScript|MeshEdits", meta=(ScriptMethod))
+	static UPARAM(DisplayName = "Target Mesh") UDynamicMesh* 
+	AddUVElementToMesh( 
+		UDynamicMesh* TargetMesh, 
+		UPARAM(DisplayName = "UV Channel") int32 UVSetIndex,
+		FVector2D NewUVPosition, 
+		UPARAM(DisplayName = "New UV Element ID") int& NewUVElementID,
+		bool& bIsValidUVSet,
+		bool bDeferChangeNotifications = false );
+
+	/**
+	 * Sets the UV Element IDs for a given Triangle in the specified UV Channel, ie the "UV Triangle" indices.
+	 * This function does not create new UVs, the provided UV Elements must already.
+	 * The UV Triangle can only be set if the resulting topology would be valid, ie the Elements cannot be shared
+	 * between different base Mesh Vertices, so they must either be unused by any other triangles, or already associated
+	 * with the same mesh vertex in other UV triangles. 
+	 * If any conditions are not met, bIsValidTriangle will be returned as false.
+	 */
+	UFUNCTION(BlueprintCallable, Category = "GeometryScript|UVs", meta=(ScriptMethod))
+	static UPARAM(DisplayName = "Target Mesh") UDynamicMesh* 
+	SetMeshTriangleUVElementIDs( 
+		UDynamicMesh* TargetMesh, 
+		UPARAM(DisplayName = "UV Channel") int UVSetIndex,
+		int TriangleID, 
+		FIntVector TriangleUVElements,
+		bool& bIsValidTriangle, 
+		bool bDeferChangeNotifications = false );
+
+
+	/**
+	 * Returns the UV Element IDs associated with the three vertices of the triangle in the specified UV Channel.
+	 * If the Triangle does not exist in the mesh or if no UVs are set in the specified UV Channel for the triangle, bHaveValidUVs will be returned as false.
+	 */
+	UFUNCTION(BlueprintCallable, Category = "GeometryScript|UVs", meta = (ScriptMethod))
+	static UPARAM(DisplayName = "Target Mesh") UDynamicMesh* 
+	GetMeshTriangleUVElementIDs(
+		UDynamicMesh* TargetMesh, 
+		UPARAM(DisplayName = "UV Channel") int32 UVSetIndex, 
+		int32 TriangleID, 
+		FIntVector& TriangleUVElements,
+		bool& bHaveValidUVs);
+
+	/**
+	 * Returns the UV Position for a given UV Element ID in the specified UV Channel.
+	 * If the UV Set or Element ID does not exist, bIsValidElementID will be returned as false.
+	 */
+	UFUNCTION(BlueprintCallable, Category = "GeometryScript|UVs", meta = (ScriptMethod))
+	static UPARAM(DisplayName = "Target Mesh") UDynamicMesh* 
+	GetMeshUVElementPosition(
+		UDynamicMesh* TargetMesh, 
+		UPARAM(DisplayName = "UV Channel") int32 UVSetIndex, 
+		int32 ElementID, 
+		FVector2D& UVPosition,
+		bool& bIsValidElementID);
+
+	/**
+	 * Sets the UV position of a specific ElementID in the given UV Set/Channel
+	 * If the UV Set or Element ID does not exist, bIsValidElementID will be returned as false.
+	 */
+	UFUNCTION(BlueprintCallable, Category = "GeometryScript|UVs", meta=(ScriptMethod))
+	static UPARAM(DisplayName = "Target Mesh") UDynamicMesh* 
+	SetMeshUVElementPosition( 
+		UDynamicMesh* TargetMesh, 
+		UPARAM(DisplayName = "UV Channel") int UVSetIndex,
+		int ElementID, 
+		FVector2D NewUVPosition,
+		bool& bIsValidElementID, 
+		bool bDeferChangeNotifications = false );
+
 	/**
 	* Update all selected UV values in the specified UV Channel by adding the Translation value to each.
 	* If the provided Selection is empty, the Translation is applied to the entire UV Channel.
