Fixed a bug where we didn't include the inpaint mask name in the transfer node.

#rnx

diff --git a/Engine/Plugins/Experimental/GeometryScripting/Source/GeometryScriptingCore/Private/MeshBoneWeightFunctions.cpp b/Engine/Plugins/Experimental/GeometryScripting/Source/GeometryScriptingCore/Private/MeshBoneWeightFunctions.cpp
--- a/Engine/Plugins/Experimental/GeometryScripting/Source/GeometryScriptingCore/Private/MeshBoneWeightFunctions.cpp
+++ b/Engine/Plugins/Experimental/GeometryScripting/Source/GeometryScriptingCore/Private/MeshBoneWeightFunctions.cpp
@@ -421,6 +421,7 @@ UDynamicMesh* UGeometryScriptLibrary_MeshBoneWeightFunctions::TransferBoneWeight
 				TransferBoneWeights.NumSmoothingIterations = Options.NumSmoothingIterations;
 				TransferBoneWeights.SmoothingStrength = Options.SmoothingStrength;
 				TransferBoneWeights.LayeredMeshSupport = Options.LayeredMeshSupport;
+				TransferBoneWeights.ForceInpaintWeightMapName = Options.InpaintMask;
 			}
 
 			if (TransferBoneWeights.Validate() != EOperationValidationResult::Ok)
diff --git a/Engine/Plugins/Experimental/GeometryScripting/Source/GeometryScriptingCore/Public/GeometryScript/MeshBoneWeightFunctions.h b/Engine/Plugins/Experimental/GeometryScripting/Source/GeometryScriptingCore/Public/GeometryScript/MeshBoneWeightFunctions.h
--- a/Engine/Plugins/Experimental/GeometryScripting/Source/GeometryScriptingCore/Public/GeometryScript/MeshBoneWeightFunctions.h
+++ b/Engine/Plugins/Experimental/GeometryScripting/Source/GeometryScriptingCore/Public/GeometryScript/MeshBoneWeightFunctions.h
@@ -152,6 +152,10 @@ struct GEOMETRYSCRIPTINGCORE_API FGeometryScriptTransferBoneWeightsOptions
 	/** The strength of each post-processing smoothing iteration. */
 	UPROPERTY(EditAnywhere, BlueprintReadWrite, Category = Options, meta = (UIMin = 0, UIMax = 1, ClampMin = 0, ClampMax = 1, EditCondition = "TransferMethod == ETransferBoneWeightsMethod::InpaintWeights"))
 	float SmoothingStrength = 0.0f;
+
+    /** Optional weight attribute name where a non-zero value indicates that we want the skinning weights for the vertex to be computed automatically instead of it being copied over from the source mesh. */
+	UPROPERTY(EditAnywhere, BlueprintReadWrite, Category = Options, meta = (EditCondition = "TransferMethod == ETransferBoneWeightsMethod::InpaintWeights"))
+	FName InpaintMask;
 };
 
 USTRUCT(BlueprintType, meta = (DisplayName = "Bone Info"))
