Fix crash in BakeTextureFromRenderCaptures Geometry Script node

diff --git a/Engine/Plugins/Experimental/GeometryScripting/Source/GeometryScriptingCore/Private/MeshBakeFunctions.cpp b/Engine/Plugins/Experimental/GeometryScripting/Source/GeometryScriptingCore/Private/MeshBakeFunctions.cpp
--- a/Engine/Plugins/Experimental/GeometryScripting/Source/GeometryScriptingCore/Private/MeshBakeFunctions.cpp
+++ b/Engine/Plugins/Experimental/GeometryScripting/Source/GeometryScriptingCore/Private/MeshBakeFunctions.cpp
@@ -1114,6 +1114,7 @@ FGeometryScriptRenderCaptureTextures UGeometryScriptLibrary_MeshBakeFunctions::B
 	const FSceneCaptureConfig Options = GeometryScriptBakeLocals::GetSceneCaptureConfig(BakeOptions);
 	TUniquePtr<FSceneCapturePhotoSet> SceneCapture = MakeUnique<FSceneCapturePhotoSet>();
 	ConfigureSceneCapture(SceneCapture, ValidSourceActors, Options, false);
+	SceneCapture->Compute();
 
 	const FDynamicMeshAABBTree3 TargetMeshSpatial(TargetMesh->GetMeshPtr());
 	TSharedPtr<FMeshTangentsd, ESPMode::ThreadSafe> TargetMeshTangents = MakeShared<FMeshTangentsd>(TargetMesh->GetMeshPtr());
