Fix PostEditUndo not called when undo is used on a BP derived from GeneratedDynamicMeshActor that has added components.

diff --git a/Engine/Plugins/Experimental/GeometryScripting/Source/GeometryScriptingEditor/Private/GeometryActors/GeneratedDynamicMeshActor.cpp b/Engine/Plugins/Experimental/GeometryScripting/Source/GeometryScriptingEditor/Private/GeometryActors/GeneratedDynamicMeshActor.cpp
--- a/Engine/Plugins/Experimental/GeometryScripting/Source/GeometryScriptingEditor/Private/GeometryActors/GeneratedDynamicMeshActor.cpp
+++ b/Engine/Plugins/Experimental/GeometryScripting/Source/GeometryScriptingEditor/Private/GeometryActors/GeneratedDynamicMeshActor.cpp
@@ -108,6 +108,23 @@ void AGeneratedDynamicMeshActor::PostEditUndo()
 	}
 }
 
+void AGeneratedDynamicMeshActor::PostEditUndo(TSharedPtr<ITransactionObjectAnnotation> TransactionAnnotation)
+{
+	Super::PostEditUndo(TransactionAnnotation);
+
+	// There is no direct signal that an Actor is being created or destroyed due to undo/redo.
+	// Currently (5.1) the checks below will tell us if the undo/redo has destroyed the
+	// Actor, and we assume otherwise it was created
+
+	if (IsActorBeingDestroyed() || !IsValid(this))	// equivalent to AActor::IsPendingKillPending()
+	{
+		UnregisterWithGenerationManager();
+	}
+	else
+	{
+		RegisterWithGenerationManager();
+	}
+}
 #endif
 
 
diff --git a/Engine/Plugins/Experimental/GeometryScripting/Source/GeometryScriptingEditor/Public/GeometryActors/GeneratedDynamicMeshActor.h b/Engine/Plugins/Experimental/GeometryScripting/Source/GeometryScriptingEditor/Public/GeometryActors/GeneratedDynamicMeshActor.h
--- a/Engine/Plugins/Experimental/GeometryScripting/Source/GeometryScriptingEditor/Public/GeometryActors/GeneratedDynamicMeshActor.h
+++ b/Engine/Plugins/Experimental/GeometryScripting/Source/GeometryScriptingEditor/Public/GeometryActors/GeneratedDynamicMeshActor.h
@@ -95,6 +95,7 @@ public:
 
 #if WITH_EDITOR
 	virtual void PostEditUndo() override;
+	virtual void PostEditUndo(TSharedPtr<ITransactionObjectAnnotation> TransactionAnnotation) override;
 #endif
 
 
