Scene sampling using a Render Capture approach

diff --git a/Engine/Plugins/Experimental/GeometryScripting/Source/GeometryScriptingCore/Private/MeshSamplingFunctions.cpp b/Engine/Plugins/Experimental/GeometryScripting/Source/GeometryScriptingCore/Private/MeshSamplingFunctions.cpp
--- a/Engine/Plugins/Experimental/GeometryScripting/Source/GeometryScriptingCore/Private/MeshSamplingFunctions.cpp
+++ b/Engine/Plugins/Experimental/GeometryScripting/Source/GeometryScriptingCore/Private/MeshSamplingFunctions.cpp
@@ -4,6 +4,9 @@
 #include "UDynamicMesh.h"
 #include "DynamicMesh/DynamicMesh3.h"
 #include "Sampling/MeshSurfacePointSampling.h"
+#include "Scene/SceneCapturePhotoSet.h"
+#include "GameFramework/Actor.h"
+#include "Async/ParallelFor.h"
 
 #include UE_INLINE_GENERATED_CPP_BY_NAME(MeshSamplingFunctions)
 
@@ -164,5 +167,169 @@ UDynamicMesh* UGeometryScriptLibrary_MeshSamplingFunctions::ComputeVertexWeighte
 }
 
 
+// TODO Move this function into SceneCapturePhotoSet.h and use it in place of ComputeExteriorSpatialPhotoParameters
+void UGeometryScriptLibrary_MeshSamplingFunctions::ComputeRenderCaptureCamerasForBox(
+	TArray<FGeometryScriptRenderCaptureCamera>& Cameras,
+	FBox Box,
+	const FGeometryScriptRenderCaptureCamerasForBoxOptions& Options,
+	UGeometryScriptDebug* Debug)
+{
+	TArray<FVector> Directions;
+
+	if (Options.bViewFromBoxFaces)
+	{
+		Directions.Add( FVector3d::UnitX());
+		Directions.Add(-FVector3d::UnitX());
+		Directions.Add( FVector3d::UnitY());
+		Directions.Add(-FVector3d::UnitY());
+		Directions.Add( FVector3d::UnitZ());
+		Directions.Add(-FVector3d::UnitZ());
+	}
+
+	if (Options.bViewFromUpperCorners)
+	{
+		Directions.Add(Normalized(FVector3d( 1,  1, -1)));
+		Directions.Add(Normalized(FVector3d(-1,  1, -1)));
+		Directions.Add(Normalized(FVector3d( 1, -1, -1)));
+		Directions.Add(Normalized(FVector3d(-1, -1, -1)));
+	}
+
+	if (Options.bViewFromLowerCorners)
+	{
+		Directions.Add(Normalized(FVector3d( 1,  1, 1)));
+		Directions.Add(Normalized(FVector3d(-1,  1, 1)));
+		Directions.Add(Normalized(FVector3d( 1, -1, 1)));
+		Directions.Add(Normalized(FVector3d(-1, -1, 1)));
+	}
+
+	if (Options.bViewFromUpperEdges)
+	{
+		Directions.Add(Normalized(FVector3d(-1,  0, -1)));
+		Directions.Add(Normalized(FVector3d( 1,  0, -1)));
+		Directions.Add(Normalized(FVector3d( 0, -1, -1)));
+		Directions.Add(Normalized(FVector3d( 0,  1, -1)));
+	}
+
+	if (Options.bViewFromLowerEdges)
+	{
+		Directions.Add(Normalized(FVector3d(-1,  0, 1)));
+		Directions.Add(Normalized(FVector3d( 1,  0, 1)));
+		Directions.Add(Normalized(FVector3d( 0, -1, 1)));
+		Directions.Add(Normalized(FVector3d( 0,  1, 1)));
+	}
+
+	if (Options.bViewFromSideEdges)
+	{
+		Directions.Add(Normalized(FVector3d( 1,  1, 0)));
+		Directions.Add(Normalized(FVector3d(-1,  1, 0)));
+		Directions.Add(Normalized(FVector3d( 1, -1, 0)));
+		Directions.Add(Normalized(FVector3d(-1, -1, 0)));
+	}
+
+	for (FVector Position : Options.ExtraViewFromPositions)
+	{
+		if (Position.Normalize(TMathUtil<double>::ZeroTolerance))
+		{
+			Directions.Add(-Position);
+		}
+		else
+		{
+			const FText Message = FText::Format(
+				LOCTEXT("ComputeRenderCaptureCamerasForBox", "ComputeRenderCaptureCamerasForBox: Unable to normalize view direction for position ({0}, {1}, {2})"),
+				Position.X,
+				Position.Y,
+				Position.Z);
+			UE::Geometry::AppendWarning(Debug, EGeometryScriptErrorType::InvalidInputs, Message);
+		}
+	}
+
+	FSphere Sphere = FBoxSphereBounds(Box).GetSphere();
+
+	const float HalfFOVRadians = FMath::DegreesToRadians<float>(Options.FieldOfViewDegrees) * .5f;
+	const float CameraToCenterDist = Sphere.W / FMath::Tan(HalfFOVRadians);
+
+	for (FVector ViewDirection : Directions)
+	{
+		ensure(IsNormalized(ViewDirection));
+
+		FGeometryScriptRenderCaptureCamera Camera;
+		Camera.NearPlaneDist = CameraToCenterDist - Sphere.W;
+		Camera.FieldOfViewDegrees = Options.FieldOfViewDegrees;
+		Camera.Resolution = FMath::Max(1, Options.Resolution);
+		Camera.ViewDirection = ViewDirection;
+		Camera.ViewPosition = Box.GetCenter() - CameraToCenterDist * Camera.ViewDirection;
+
+		Cameras.Add(Camera);
+	}
+}
+
+
+
+void UGeometryScriptLibrary_MeshSamplingFunctions::ComputeRenderCapturePointSampling(
+	TArray<FTransform>& Samples,
+	const TArray<AActor*>& Actors,
+	const TArray<FGeometryScriptRenderCaptureCamera>& Cameras,
+	UGeometryScriptDebug* Debug)
+{
+	TRACE_CPUPROFILER_EVENT_SCOPE(GS_ComputeRenderCapturePointSampling);
+
+	// It's possible to pass nullptrs in Actors so we filter these out here
+	const TArray<AActor*> ValidActors = Actors.FilterByPredicate([](AActor* Actor) { return Actor != nullptr; });
+	if (ValidActors.IsEmpty())
+	{
+		return;
+	}
+
+	// Setup and configure SceneCapture
+	TUniquePtr<FSceneCapturePhotoSet> SceneCapture = MakeUnique<FSceneCapturePhotoSet>();
+	{
+		ForEachCaptureType([&SceneCapture](ERenderCaptureType CaptureType)
+		{
+			const bool bCaptureTypeEnabled =
+				CaptureType == ERenderCaptureType::DeviceDepth ||
+				CaptureType == ERenderCaptureType::WorldNormal;
+			SceneCapture->SetCaptureTypeEnabled(CaptureType, bCaptureTypeEnabled);
+
+			FRenderCaptureConfig Config;
+			Config.bAntiAliasing = false;
+			SceneCapture->SetCaptureConfig(CaptureType, Config);
+		});
+
+		SceneCapture->SetCaptureSceneActors(ValidActors[0]->GetWorld(), ValidActors);
+
+		TArray<FSpatialPhotoParams> SpatialParams;
+		for (const FGeometryScriptRenderCaptureCamera& Camera : Cameras)
+		{
+
+			FSpatialPhotoParams Params;
+			Params.NearPlaneDist = Camera.NearPlaneDist;
+			Params.HorzFOVDegrees = Camera.FieldOfViewDegrees;
+			Params.Dimensions = FImageDimensions(FMath::Max(1, Camera.Resolution), FMath::Max(1, Camera.Resolution));
+			Params.Frame.AlignAxis(0, Camera.ViewDirection);
+			Params.Frame.ConstrainedAlignAxis(2, FVector3d::UnitZ(), Params.Frame.X());
+			Params.Frame.Origin = Camera.ViewPosition;
+			SpatialParams.Add(Params);
+		}
+		SceneCapture->SetSpatialPhotoParams(SpatialParams);
+	}
+
+	SceneCapture->Compute();
+
+	// Extract the Samples
+	{
+		FSceneCapturePhotoSet::FSceneSamples SceneSamples;
+		TArray<FFrame3f> WorldOrientedPoints;
+		SceneSamples.WorldOrientedPoints = &WorldOrientedPoints;
+		SceneCapture->GetSceneSamples(SceneSamples);
+
+		Samples.SetNum(WorldOrientedPoints.Num());
+		ParallelFor(WorldOrientedPoints.Num(), [&Samples, &WorldOrientedPoints](int32 Index)
+		{
+			Samples[Index] = WorldOrientedPoints[Index].ToFTransform();
+		});
+	}
+}
+
+
 
-#undef LOCTEXT_NAMESPACE
\ No newline at end of file
+#undef LOCTEXT_NAMESPACE
diff --git a/Engine/Plugins/Experimental/GeometryScripting/Source/GeometryScriptingCore/Public/GeometryScript/MeshSamplingFunctions.h b/Engine/Plugins/Experimental/GeometryScripting/Source/GeometryScriptingCore/Public/GeometryScript/MeshSamplingFunctions.h
--- a/Engine/Plugins/Experimental/GeometryScripting/Source/GeometryScriptingCore/Public/GeometryScript/MeshSamplingFunctions.h
+++ b/Engine/Plugins/Experimental/GeometryScripting/Source/GeometryScriptingCore/Public/GeometryScript/MeshSamplingFunctions.h
@@ -93,6 +93,72 @@ public:
 };
 
 
+USTRUCT(BlueprintType)
+struct GEOMETRYSCRIPTINGCORE_API FGeometryScriptRenderCaptureCamera
+{
+	GENERATED_BODY()
+
+	/** The pixel resolution of render capture photo set, this value is used for width and height */
+	UPROPERTY(EditAnywhere, BlueprintReadWrite, Category = Options)
+	int Resolution = 256;
+
+	UPROPERTY(EditAnywhere, BlueprintReadWrite, Category = Options)
+	double FieldOfViewDegrees = 45.0;
+
+	UPROPERTY(EditAnywhere, BlueprintReadWrite, Category = Options)
+	FVector ViewPosition;
+
+	UPROPERTY(EditAnywhere, BlueprintReadWrite, Category = Options)
+	FVector ViewDirection;
+
+	UPROPERTY(EditAnywhere, BlueprintReadWrite, Category = Options)
+	double NearPlaneDist = 1.0;
+};
+
+
+
+USTRUCT(BlueprintType)
+struct GEOMETRYSCRIPTINGCORE_API FGeometryScriptRenderCaptureCamerasForBoxOptions
+{
+	GENERATED_BODY()
+
+	/** The pixel resolution of render capture photos, this value is used for width and height */
+	UPROPERTY(EditAnywhere, BlueprintReadWrite, Category = Options)
+	int Resolution = 256;
+
+	UPROPERTY(EditAnywhere, BlueprintReadWrite, Category = Options)
+	double FieldOfViewDegrees = 45.;
+
+	/** Enable 6 directions corresponding to views from box face centers to the box center */
+	UPROPERTY(EditAnywhere, BlueprintReadWrite, Category = Options)
+	bool bViewFromBoxFaces = true;
+
+	/** Enable 4 directions corresponding to views from box upper corners to the box center */
+	UPROPERTY(EditAnywhere, BlueprintReadWrite, Category = Options)
+	bool bViewFromUpperCorners = false;
+
+	/** Enable 4 directions corresponding to views from box lower corners to the box center */
+	UPROPERTY(EditAnywhere, BlueprintReadWrite, Category = Options)
+	bool bViewFromLowerCorners = false;
+
+	/** Enable 4 directions corresponding to views from box upper edges centers to the box center */
+	UPROPERTY(EditAnywhere, BlueprintReadWrite, Category = Options)
+	bool bViewFromUpperEdges = false;
+
+	/** Enable 4 directions corresponding to views from box lower edges centers to the box center */
+	UPROPERTY(EditAnywhere, BlueprintReadWrite, Category = Options)
+	bool bViewFromLowerEdges = false;
+
+	/** Enable 4 directions corresponding to views from box side edges centers to the box center */
+	UPROPERTY(EditAnywhere, BlueprintReadWrite, Category = Options)
+	bool bViewFromSideEdges = false;
+
+	/** Extra positions from which to deduce view directions on the box center (located at (0,0,0)) */
+	UPROPERTY(EditAnywhere, BlueprintReadWrite, Category = Options)
+	TArray<FVector> ExtraViewFromPositions;
+};
+
+
 UCLASS(meta = (ScriptName = "GeometryScript_MeshSampling"))
 class GEOMETRYSCRIPTINGCORE_API UGeometryScriptLibrary_MeshSamplingFunctions : public UBlueprintFunctionLibrary
 {
@@ -149,4 +215,30 @@ public:
 		FGeometryScriptIndexList& TriangleIDs,
 		UGeometryScriptDebug* Debug = nullptr);
 
+	/**
+	 * Compute a set of Render Capture Cameras to capture a scene within the given Box
+	 * @param Cameras Output Cameras with view frustums that contain the Box while maintaining the desired FOV
+	 * @param Box     Bounding Box containing the scene to be captured
+	 * @param Options Defines the Camera viewing directions into the box and other Camera parameters
+	 */
+	UFUNCTION(BlueprintCallable, Category = "GeometryScript|MeshSampling")
+	static void
+	ComputeRenderCaptureCamerasForBox(
+		TArray<FGeometryScriptRenderCaptureCamera>& Cameras,
+		FBox Box,
+		const FGeometryScriptRenderCaptureCamerasForBoxOptions& Options,
+		UGeometryScriptDebug* Debug = nullptr);
+
+	/**
+	 * Compute oriented sample points on the visible surfaces of the given Actors
+	 * The Samples are computed using Render Capture from the given virtual Cameras
+	 */
+	UFUNCTION(BlueprintCallable, Category = "GeometryScript|MeshSampling")
+	static void
+	ComputeRenderCapturePointSampling(
+		TArray<FTransform>& Samples,
+		const TArray<AActor*>& Actors,
+		const TArray<FGeometryScriptRenderCaptureCamera>& Cameras,
+		UGeometryScriptDebug* Debug = nullptr);
+
 };
\ No newline at end of file
diff --git a/Engine/Plugins/Runtime/MeshModelingToolset/Source/ModelingComponents/Private/Scene/SceneCapturePhotoSet.cpp b/Engine/Plugins/Runtime/MeshModelingToolset/Source/ModelingComponents/Private/Scene/SceneCapturePhotoSet.cpp
--- a/Engine/Plugins/Runtime/MeshModelingToolset/Source/ModelingComponents/Private/Scene/SceneCapturePhotoSet.cpp
+++ b/Engine/Plugins/Runtime/MeshModelingToolset/Source/ModelingComponents/Private/Scene/SceneCapturePhotoSet.cpp
@@ -155,7 +155,7 @@ void FSceneCapturePhotoSet::Compute()
 		return;
 	}
 
-	TRACE_CPUPROFILER_EVENT_SCOPE(CapturePhotoSet);
+	TRACE_CPUPROFILER_EVENT_SCOPE(SceneCapturePhotoSet::Compute);
 	FScopedSlowTask Progress(static_cast<float>(NumPending), LOCTEXT("CapturingScene", "Capturing Scene..."));
 	Progress.MakeDialog(bAllowCancel);
 
@@ -524,7 +524,7 @@ FSceneCapturePhotoSet::FSceneSample::FSceneSample()
 	Opacity = 0.0f;
 	SubsurfaceColor = FVector3f(0, 0, 0);
 	WorldNormal = FVector3f(0, 0, 1);
-	DeviceDepth = 0.0f;
+	DeviceDepth = 0.0f; // Since we use reverse z perspective projection the far plane corresponds to DeviceDepth == 0
 }
 
 FVector3f FSceneCapturePhotoSet::FSceneSample::GetValue3f(ERenderCaptureType CaptureType) const
@@ -766,6 +766,126 @@ bool FSceneCapturePhotoSet::ComputeSample(
 	return true;
 }
 
+void FSceneCapturePhotoSet::GetSceneSamples(FSceneSamples& OutSamples)
+{
+	if (!ensure(PhotoSetStatus.DeviceDepth == ECaptureTypeStatus::Computed))
+	{
+		return;
+	}
+
+	TRACE_CPUPROFILER_EVENT_SCOPE(SceneCapturePhotoSet::GetSceneSamples);
+
+	for (int32 PhotoIndex = 0; PhotoIndex < GetSpatialPhotoParams().Num(); ++PhotoIndex)
+	{
+		const FSpatialPhotoParams& Params = GetSpatialPhotoParams()[PhotoIndex];
+
+		for (int64 PixelLinearIndex = 0; PixelLinearIndex < Params.Dimensions.Num(); ++PixelLinearIndex)
+		{
+			const FSpatialPhoto1f& DeviceDepthPhoto = GetDeviceDepthPhotoSet().Get(PhotoIndex);
+
+			const float DeviceZ = DeviceDepthPhoto.Image.GetPixel(PixelLinearIndex);
+			const bool bValidDeviceZ = (DeviceZ > 0) && (DeviceZ < 1);
+			if (bValidDeviceZ) // Skip samples on the near/far plane
+			{
+				bool bComputePointWorld  = OutSamples.WorldPoint != nullptr;
+				bool bComputeNormalWorld = false;
+				if (PhotoSetStatus.WorldNormal == ECaptureTypeStatus::Computed)
+				{
+					bComputePointWorld  =  bComputePointWorld || (OutSamples.WorldOrientedPoints != nullptr);
+					bComputeNormalWorld = (OutSamples.WorldNormal != nullptr) || (OutSamples.WorldOrientedPoints != nullptr);
+				}
+
+				FVector3f PointWorld;
+				if (bComputePointWorld)
+				{
+					// Map from pixel space to normalized device coordinates
+					FVector2d DeviceXY = FRenderCaptureCoordinateConverter2D::PixelToDevice(
+						Params.Dimensions.GetCoords(PixelLinearIndex),
+						Params.Dimensions.GetWidth(),
+						Params.Dimensions.GetHeight());
+
+					// Map from normalized device coordinates to world coordinates
+					FVector3d PointDevice(DeviceXY, DeviceZ);
+					const FMatrix& InvViewProjectionMatrix = PhotoViewMatricies[PhotoIndex].GetInvViewProjectionMatrix();
+					const FVector4d PointWorld4 = InvViewProjectionMatrix.TransformPosition(PointDevice);
+
+					// Convert from Homogenous to Cartesian coordinates
+					PointWorld.X = static_cast<float>(PointWorld4.X / PointWorld4.W);
+					PointWorld.Y = static_cast<float>(PointWorld4.Y / PointWorld4.W);
+					PointWorld.Z = static_cast<float>(PointWorld4.Z / PointWorld4.W);
+				}
+
+				FVector3f NormalWorld;
+				if (bComputeNormalWorld)
+				{
+					// Map normal from color components [0,1] to world components [-1,1]
+					const FSpatialPhoto3f& NormalPhoto = GetWorldNormalPhotoSet().Get(PhotoIndex);
+					const FVector4f NormalColor = NormalPhoto.Image.GetPixel(PixelLinearIndex);
+					NormalWorld = 2.f * (FVector3f(NormalColor) - FVector3f(.5f));
+				}
+
+				if (OutSamples.WorldPoint)
+				{
+					OutSamples.WorldPoint->Add(PointWorld);
+				}
+
+				if (PhotoSetStatus.WorldNormal == ECaptureTypeStatus::Computed)
+				{
+					if (OutSamples.WorldNormal)
+					{
+						OutSamples.WorldNormal->Add(NormalWorld);
+					}
+
+					if (OutSamples.WorldOrientedPoints)
+					{
+						OutSamples.WorldOrientedPoints->Add(FFrame3f(PointWorld, NormalWorld));
+					}
+				}
+
+				if (PhotoSetStatus.Metallic == ECaptureTypeStatus::Computed && OutSamples.Metallic)
+				{
+					OutSamples.Metallic->Add(GetMetallicPhotoSet().Get(PhotoIndex).Image.GetPixel(PixelLinearIndex));
+				}
+
+				if (PhotoSetStatus.Roughness == ECaptureTypeStatus::Computed && OutSamples.Roughness)
+				{
+					OutSamples.Roughness->Add(GetRoughnessPhotoSet().Get(PhotoIndex).Image.GetPixel(PixelLinearIndex));
+				}
+
+				if (PhotoSetStatus.Specular == ECaptureTypeStatus::Computed && OutSamples.Specular)
+				{
+					OutSamples.Specular->Add(GetSpecularPhotoSet().Get(PhotoIndex).Image.GetPixel(PixelLinearIndex));
+				}
+
+				if (PhotoSetStatus.CombinedMRS == ECaptureTypeStatus::Computed && OutSamples.PackedMRS)
+				{
+					OutSamples.PackedMRS->Add(GetPackedMRSPhotoSet().Get(PhotoIndex).Image.GetPixel(PixelLinearIndex));
+				}
+
+				if (PhotoSetStatus.Emissive == ECaptureTypeStatus::Computed && OutSamples.Emissive)
+				{
+					OutSamples.Emissive->Add(GetEmissivePhotoSet().Get(PhotoIndex).Image.GetPixel(PixelLinearIndex));
+				}
+
+				if (PhotoSetStatus.BaseColor == ECaptureTypeStatus::Computed && OutSamples.BaseColor)
+				{
+					OutSamples.BaseColor->Add(GetBaseColorPhotoSet().Get(PhotoIndex).Image.GetPixel(PixelLinearIndex));
+				}
+
+				if (PhotoSetStatus.SubsurfaceColor == ECaptureTypeStatus::Computed && OutSamples.SubsurfaceColor)
+				{
+					OutSamples.SubsurfaceColor->Add(GetSubsurfaceColorPhotoSet().Get(PhotoIndex).Image.GetPixel(PixelLinearIndex));
+				}
+
+				if (PhotoSetStatus.Opacity == ECaptureTypeStatus::Computed && OutSamples.Opacity)
+				{
+					OutSamples.Opacity->Add(GetOpacityPhotoSet().Get(PhotoIndex).Image.GetPixel(PixelLinearIndex));
+				}
+			} // bValidDeviceZ
+		} // PixelLinearIndex
+	} // PhotoIndex
+}
+
 
 void FSceneCapturePhotoSet::SetEnableVisibilityByUnregisterMode(bool bEnable)
 {
diff --git a/Engine/Plugins/Runtime/MeshModelingToolset/Source/ModelingComponents/Private/Scene/WorldRenderCapture.cpp b/Engine/Plugins/Runtime/MeshModelingToolset/Source/ModelingComponents/Private/Scene/WorldRenderCapture.cpp
--- a/Engine/Plugins/Runtime/MeshModelingToolset/Source/ModelingComponents/Private/Scene/WorldRenderCapture.cpp
+++ b/Engine/Plugins/Runtime/MeshModelingToolset/Source/ModelingComponents/Private/Scene/WorldRenderCapture.cpp
@@ -833,7 +833,6 @@ bool FWorldRenderCapture::CaptureDeviceDepthFromPosition(
 	ShowFlags.SetMotionBlur(false);
 	ShowFlags.SetBloom(false);
 	ShowFlags.SetSceneColorFringe(false);
-	ShowFlags.SetNaniteMeshes(false);
 	ShowFlags.SetAtmosphere(false);
 	ShowFlags.SetLighting(false);
 	ShowFlags.SetScreenPercentage(false);
@@ -1057,6 +1056,8 @@ bool FWorldRenderCapture::CaptureFromPosition(
 	FImageAdapter& ResultImageOut,
 	const FRenderCaptureConfig& Config)
 {
+	TRACE_CPUPROFILER_EVENT_SCOPE(WorldRenderCapture::CaptureFromPosition);
+
 	RenderCaptureInterface::FScopedCapture RenderCapture(RenderCaptureDraws > 0, TEXT("RenderCaptureFromPosition"));
 	RenderCaptureDraws--;
 
diff --git a/Engine/Plugins/Runtime/MeshModelingToolset/Source/ModelingComponents/Public/Scene/SceneCapturePhotoSet.h b/Engine/Plugins/Runtime/MeshModelingToolset/Source/ModelingComponents/Public/Scene/SceneCapturePhotoSet.h
--- a/Engine/Plugins/Runtime/MeshModelingToolset/Source/ModelingComponents/Public/Scene/SceneCapturePhotoSet.h
+++ b/Engine/Plugins/Runtime/MeshModelingToolset/Source/ModelingComponents/Public/Scene/SceneCapturePhotoSet.h
@@ -51,6 +51,7 @@ public:
 	/**
 	 * Set the target World and set of Actors/Components
 	 * If World != this->World or Actors != this->Actors all existing photo sets are cleared
+	 * Note: The caller must ensure Actors does not contain nullptr
 	 */
 	void SetCaptureSceneActors(UWorld* World, const TArray<AActor*>& Actors);
 	void SetCaptureSceneComponents(UWorld* World, const TArray<UActorComponent*>& Components);
@@ -212,6 +213,36 @@ public:
 		const FVector2d& PhotoCoords,
 		const FSceneSample& DefaultSample) const;
 
+	/**
+	 * FSceneSamples stores samples corresponding to pixels in the DeviceDepth photoset where the values are strictly
+	 * within the viewing frustum, which corresponds to a camera ray intersecting an actor in VisibleActors. All non-null
+	 * containers with Computed capture status will be filled by GetSceneSamples() and will have the same .Num() counts
+	 */
+	struct FSceneSamples
+	{
+		// These are world-space point, normal or oriented point samples
+		// The points are computed from the DeviceDepth channel and normals from the WorldNormal channel
+		// Coordinate frames are located at the world points with Z axis aligned with world normals and arbitrary X/Y axes
+		TArray<FVector3f>* WorldPoint = nullptr;
+		TArray<FVector3f>* WorldNormal = nullptr;
+		TArray<FFrame3f>*  WorldOrientedPoints = nullptr;
+
+		// These are samples of the corresponding capture channels
+		TArray<float>* Metallic = nullptr;
+		TArray<float>* Roughness = nullptr;
+		TArray<float>* Specular = nullptr;
+		TArray<FVector3f>* PackedMRS = nullptr;
+		TArray<FVector3f>* Emissive = nullptr;
+		TArray<FVector3f>* BaseColor = nullptr;
+		TArray<FVector3f>* SubsurfaceColor = nullptr;
+		TArray<float>* Opacity = nullptr;
+	};
+
+	/**
+	 * Fills in the non-null arrays in OutSamples if the status of the needed captures is Computed
+	 */
+	void GetSceneSamples(FSceneSamples& OutSamples);
+
 	const FSpatialPhotoSet3f& GetBaseColorPhotoSet() { return BaseColorPhotoSet; }
 	const FSpatialPhotoSet1f& GetRoughnessPhotoSet() { return RoughnessPhotoSet; }
 	const FSpatialPhotoSet1f& GetSpecularPhotoSet() { return SpecularPhotoSet; }
diff --git a/Engine/Source/Runtime/GeometryCore/Public/Image/SpatialPhotoSet.h b/Engine/Source/Runtime/GeometryCore/Public/Image/SpatialPhotoSet.h
--- a/Engine/Source/Runtime/GeometryCore/Public/Image/SpatialPhotoSet.h
+++ b/Engine/Source/Runtime/GeometryCore/Public/Image/SpatialPhotoSet.h
@@ -21,7 +21,12 @@ struct FSpatialPhotoParams
 	/** Coordinate system of the view camera - X() is forward, Z() is up */
 	FFrame3d Frame;
 
-	/** Near-plane distance for the camera, image pixels lie on this plane */
+	/**
+	 * Near-plane distance of the camera. Image pixels lie on this plane.
+	 *
+	 * Note: For best depth precision, the near plane should be placed as close as possible to the scene being captured.
+	 * For more on depth precision, see: https://www.reedbeta.com/blog/depth-precision-visualized/
+	 */	
 	double NearPlaneDist = 1.0;
 
 	/** Horizontal Field-of-View of the camera in degrees (full FOV, so generally calculations will use half this value) */
