add more point set functions for Geometry Script:

1. OffsetTransforms to offset the location of all transforms in an array, e.g. to move them all along their local Z if you want to move mesh samples along their normal
2. DownsamplePoints to downsample points to a fixed number of samples, optionally with priority weighting and approximate uniform spacing

diff --git a/Engine/Plugins/Experimental/GeometryScripting/Source/GeometryScriptingCore/Private/PointSetFunctions.cpp b/Engine/Plugins/Experimental/GeometryScripting/Source/GeometryScriptingCore/Private/PointSetFunctions.cpp
--- a/Engine/Plugins/Experimental/GeometryScripting/Source/GeometryScriptingCore/Private/PointSetFunctions.cpp
+++ b/Engine/Plugins/Experimental/GeometryScripting/Source/GeometryScriptingCore/Private/PointSetFunctions.cpp
@@ -3,6 +3,7 @@
 #include "GeometryScript/PointSetFunctions.h"
 
 #include "Clustering/KMeans.h"
+#include "Spatial/PriorityOrderPoints.h"
 
 #include UE_INLINE_GENERATED_CPP_BY_NAME(PointSetFunctions)
 
@@ -84,6 +85,47 @@ void UGeometryScriptLibrary_PointSetSamplingFunctions::TransformsToPoints(
 	}
 }
 
+void UGeometryScriptLibrary_PointSetSamplingFunctions::DownsamplePoints(
+	const TArray<FVector>& Points,
+	const FGeometryScriptPointPriorityOptions& Options,
+	FGeometryScriptIndexList& DownsampledIndices,
+	int32 KeepNumPoints,
+	UGeometryScriptDebug* Debug)
+{
+	DownsampledIndices.Reset(EGeometryScriptIndexType::Any);
+	int32 NumWeights = Options.OptionalPriorityWeights.Num();
+	bool bHasImportanceWeights = NumWeights >= Points.Num();
+	if (NumWeights > 0 && NumWeights < Points.Num())
+	{
+		UE::Geometry::AppendWarning(Debug, EGeometryScriptErrorType::InvalidInputs, 
+			FText::Format(LOCTEXT("DownsamplePointsMissingWeights", "DownsamplePoints: Found {0} points but only {1} PriorityWeights"), Points.Num(), NumWeights));
+	}
+	KeepNumPoints = FMath::Min(Points.Num(), KeepNumPoints);
+	if (Options.bUniformSpacing || bHasImportanceWeights)
+	{
+		FPriorityOrderPoints Order;
+		if (Options.bUniformSpacing)
+		{
+			Order.ComputeUniformSpaced(Points, bHasImportanceWeights ? Options.OptionalPriorityWeights : TArrayView<const float>(), KeepNumPoints);
+		}
+		else // bHasImportanceWeights
+		{
+			Order.ComputeDescendingImportance(Options.OptionalPriorityWeights);
+		}
+		Order.Order.SetNum(KeepNumPoints);
+		(*DownsampledIndices.List) = MoveTemp(Order.Order);
+		
+	}
+	else // no criteria to prioritize points; just take the first indices
+	{
+		DownsampledIndices.List->SetNumUninitialized(KeepNumPoints);
+		for (int32 Idx = 0; Idx < KeepNumPoints; ++Idx)
+		{
+			(*DownsampledIndices.List)[Idx] = Idx;
+		}
+	}
+}
+
 void UGeometryScriptLibrary_PointSetSamplingFunctions::GetPointsFromIndexList(
 	const TArray<FVector>& AllPoints,
 	const FGeometryScriptIndexList& Indices,
@@ -102,6 +144,25 @@ void UGeometryScriptLibrary_PointSetSamplingFunctions::GetPointsFromIndexList(
 	}
 }
 
+void UGeometryScriptLibrary_PointSetSamplingFunctions::OffsetTransforms(
+	TArray<FTransform>& Transforms,
+	double Offset,
+	FVector Direction,
+	EGeometryScriptCoordinateSpace Space
+)
+{
+	Direction.Normalize();
+	for (FTransform& Transform : Transforms)
+	{
+		FVector OffsetVec = Direction * Offset;
+		if (Space == EGeometryScriptCoordinateSpace::Local)
+		{
+			OffsetVec = Transform.TransformVector(OffsetVec);
+		}
+		Transform.SetLocation(OffsetVec + Transform.GetLocation());
+	}
+}
+
 void UGeometryScriptLibrary_PointSetSamplingFunctions::FlattenPoints(
 	const TArray<FVector>& PointsIn3D,
 	TArray<FVector2D>& PointsIn2D,
diff --git a/Engine/Plugins/Experimental/GeometryScripting/Source/GeometryScriptingCore/Public/GeometryScript/GeometryScriptTypes.h b/Engine/Plugins/Experimental/GeometryScripting/Source/GeometryScriptingCore/Public/GeometryScript/GeometryScriptTypes.h
--- a/Engine/Plugins/Experimental/GeometryScripting/Source/GeometryScriptingCore/Public/GeometryScript/GeometryScriptTypes.h
+++ b/Engine/Plugins/Experimental/GeometryScripting/Source/GeometryScriptingCore/Public/GeometryScript/GeometryScriptTypes.h
@@ -55,6 +55,14 @@ enum class EGeometryScriptAxis : uint8
 };
 
 
+UENUM(BlueprintType)
+enum class EGeometryScriptCoordinateSpace : uint8
+{
+	Local = 0,
+	World = 1
+};
+
+
 USTRUCT(BlueprintType)
 struct GEOMETRYSCRIPTINGCORE_API FGeometryScriptMeshReadLOD
 {
diff --git a/Engine/Plugins/Experimental/GeometryScripting/Source/GeometryScriptingCore/Public/GeometryScript/PointSetFunctions.h b/Engine/Plugins/Experimental/GeometryScripting/Source/GeometryScriptingCore/Public/GeometryScript/PointSetFunctions.h
--- a/Engine/Plugins/Experimental/GeometryScripting/Source/GeometryScriptingCore/Public/GeometryScript/PointSetFunctions.h
+++ b/Engine/Plugins/Experimental/GeometryScripting/Source/GeometryScriptingCore/Public/GeometryScript/PointSetFunctions.h
@@ -42,6 +42,22 @@ public:
 	int32 MaxIterations = 500;
 };
 
+USTRUCT(BlueprintType)
+struct GEOMETRYSCRIPTINGCORE_API FGeometryScriptPointPriorityOptions
+{
+	GENERATED_BODY()
+public:
+
+	/** If not empty, will be used to order the points so that higher-priority points are kept. */
+	UPROPERTY(EditAnywhere, BlueprintReadWrite, Category = Options)
+	TArray<float> OptionalPriorityWeights;
+
+	/** Whether to ensure the kept points are approximately uniformly spaced */
+	UPROPERTY(EditAnywhere, BlueprintReadWrite, Category = Options)
+	bool bUniformSpacing = true;
+
+};
+
 
 USTRUCT(BlueprintType)
 struct GEOMETRYSCRIPTINGCORE_API FGeometryScriptPointFlatteningOptions
@@ -86,6 +102,22 @@ public:
 		const FGeometryScriptPointClusteringOptions& Options,
 		TArray<FGeometryScriptIndexList>& ClusterIDToLists);
 
+	/**
+	 * Find a subset of the given Points of a specified size.
+	 * Can optionally specify a priorty weighting and/or request uniform spacing for the downsampled points.
+	 * Note: Ordering of the result will balance:
+	 * (1) if weights are provided, higher weight points come earlier and
+	 * (2) if uniform spacing is requested, points will be ordered to have an octree-style coverage --
+	 *     so the first 8 points will cover the 8 octants (where samples are available) and the subsequent points will progressively fill in the space
+	 */
+	UFUNCTION(BlueprintCallable, Category = "GeometryScript|PointSet")
+	static void DownsamplePoints(
+		const TArray<FVector>& Points,
+		const FGeometryScriptPointPriorityOptions& Options,
+		FGeometryScriptIndexList& DownsampledIndices,
+		int32 KeepNumPoints = 100,
+		UGeometryScriptDebug* Debug = nullptr);
+
 	/**
 	 * Create an array of the positions of the input Transforms
 	 */
@@ -95,6 +127,18 @@ public:
 		TArray<FVector>& Points
 	);
 
+	/**
+	 * Offset the location of all Transforms by Offset in the given Direction, either locally in the space of the transform or in world space.
+	 * For example, this can offset mesh surface samples along the surface normal direction.
+	 */
+	UFUNCTION(BlueprintCallable, Category = "GeometryScript|PointSet")
+	static void OffsetTransforms(
+		UPARAM(ref) TArray<FTransform>& Transforms,
+		double Offset,
+		FVector Direction = FVector::UpVector,
+		EGeometryScriptCoordinateSpace Space = EGeometryScriptCoordinateSpace::Local
+	);
+
 	/**
 	 * Convert an array of points from 3D to 2D, by transforming into the given ReferenceFrame and taking the X,Y coordinates
 	 * Note that to transform into the ReferenceFrame, we apply the inverse of the ReferenceFrame's transform.
