Fix StructProperty FGeometryScriptRenderCaptureCamera::ViewPosition is not initialized properly

#rnx

diff --git a/Engine/Plugins/Experimental/GeometryScripting/Source/GeometryScriptingCore/Public/GeometryScript/MeshSamplingFunctions.h b/Engine/Plugins/Experimental/GeometryScripting/Source/GeometryScriptingCore/Public/GeometryScript/MeshSamplingFunctions.h
--- a/Engine/Plugins/Experimental/GeometryScripting/Source/GeometryScriptingCore/Public/GeometryScript/MeshSamplingFunctions.h
+++ b/Engine/Plugins/Experimental/GeometryScripting/Source/GeometryScriptingCore/Public/GeometryScript/MeshSamplingFunctions.h
@@ -106,10 +106,10 @@ struct GEOMETRYSCRIPTINGCORE_API FGeometryScriptRenderCaptureCamera
 	double FieldOfViewDegrees = 45.0;
 
 	UPROPERTY(EditAnywhere, BlueprintReadWrite, Category = Options)
-	FVector ViewPosition;
+	FVector ViewPosition = FVector::Zero();
 
 	UPROPERTY(EditAnywhere, BlueprintReadWrite, Category = Options)
-	FVector ViewDirection;
+	FVector ViewDirection = FVector::UnitX();
 
 	UPROPERTY(EditAnywhere, BlueprintReadWrite, Category = Options)
 	double NearPlaneDist = 1.0;
