Add new geometry script functions to copy and transfer simple collision shapes, as groundwork for future functions to operate directly on simple shapes (e.g., merging/simplifying them)

diff --git a/Engine/Plugins/Experimental/GeometryScripting/Source/GeometryScriptingCore/Private/CollisionFunctions.cpp b/Engine/Plugins/Experimental/GeometryScripting/Source/GeometryScriptingCore/Private/CollisionFunctions.cpp
--- a/Engine/Plugins/Experimental/GeometryScripting/Source/GeometryScriptingCore/Private/CollisionFunctions.cpp
+++ b/Engine/Plugins/Experimental/GeometryScripting/Source/GeometryScriptingCore/Private/CollisionFunctions.cpp
@@ -377,6 +377,98 @@ void UGeometryScriptLibrary_CollisionFunctions::ResetDynamicMeshCollision(
 }
 
 
+FGeometryScriptSimpleCollision
+UGeometryScriptLibrary_CollisionFunctions::GetSimpleCollisionFromComponent(
+	UPrimitiveComponent* Component,
+	UGeometryScriptDebug* Debug)
+{
+	FGeometryScriptSimpleCollision ToRet;
+	if (Component == nullptr)
+	{
+		UE::Geometry::AppendError(Debug, EGeometryScriptErrorType::InvalidInputs, LOCTEXT("GetSimpleCollisionFromComponent_InvalidComponent", "GetSimpleCollisionFromComponent: Component is Null"));
+		return ToRet;
+	}
+	const UBodySetup* BodySetup = Component->GetBodySetup();
+	if (BodySetup == nullptr)
+	{
+		UE::Geometry::AppendError(Debug, EGeometryScriptErrorType::InvalidInputs, LOCTEXT("GetSimpleCollisionFromComponent_InvalidBodySetup", "GetSimpleCollisionFromComponent: Component's BodySetup is Null"));
+		return ToRet;
+	}
+	ToRet.AggGeom = BodySetup->AggGeom;
+	
+	return ToRet;
+}
+
+void UGeometryScriptLibrary_CollisionFunctions::SetSimpleCollisionOfDynamicMeshComponent(
+	const FGeometryScriptSimpleCollision& SimpleCollision,
+	UDynamicMeshComponent* DynamicMeshComponent,
+	FGeometryScriptSetSimpleCollisionOptions Options,
+	UGeometryScriptDebug* Debug)
+{
+#if WITH_EDITOR
+	if (Options.bEmitTransaction && GEditor)
+	{
+		GEditor->BeginTransaction(LOCTEXT("UpdateDynamicMesh", "Set Simple Collision"));
+
+		DynamicMeshComponent->Modify();
+	}
+#endif
+
+#if WITH_EDITOR
+	if (Options.bEmitTransaction)
+	{
+		UBodySetup* BodySetup = DynamicMeshComponent->GetBodySetup();
+		if (BodySetup != nullptr)
+		{
+			BodySetup->Modify();
+		}
+	}
+#endif
+
+	// set new collision geometry
+	DynamicMeshComponent->SetSimpleCollisionShapes(SimpleCollision.AggGeom, true /*bUpdateCollision*/);
+
+	// do we need to do a post edit change here??
+
+#if WITH_EDITOR
+	if (Options.bEmitTransaction && GEditor)
+	{
+		GEditor->EndTransaction();
+	}
+#endif
+}
+
+
+FGeometryScriptSimpleCollision UGeometryScriptLibrary_CollisionFunctions::GetSimpleCollisionFromStaticMesh(
+	UStaticMesh* StaticMeshAsset, UGeometryScriptDebug* Debug)
+{
+	FGeometryScriptSimpleCollision ToRet;
+	
+	if (StaticMeshAsset == nullptr)
+	{
+		UE::Geometry::AppendError(Debug, EGeometryScriptErrorType::InvalidInputs, LOCTEXT("GetSimpleCollisionFromStaticMesh_InvalidStaticMesh", "GetSimpleCollisionFromStaticMesh: Input Mesh is Null"));
+		return ToRet;
+	}
+	const UBodySetup* BodySetup = StaticMeshAsset->GetBodySetup();
+	if (BodySetup == nullptr)
+	{
+		UE::Geometry::AppendError(Debug, EGeometryScriptErrorType::InvalidInputs, LOCTEXT("GetSimpleCollisionFromStaticMesh_InvalidBodySetup", "GetSimpleCollisionFromStaticMesh: Input Mesh's BodySetup is Null"));
+		return ToRet;
+	}
+
+	ToRet.AggGeom = BodySetup->AggGeom;
+
+	return ToRet;
+}
+
+void UGeometryScriptLibrary_CollisionFunctions::SetSimpleCollisionOfStaticMesh(
+	const FGeometryScriptSimpleCollision& SimpleCollision,
+	UStaticMesh* StaticMesh, 
+	FGeometryScriptSetSimpleCollisionOptions Options,
+	UGeometryScriptDebug* Debug)
+{
+	UELocal::SetStaticMeshSimpleCollision(StaticMesh, SimpleCollision.AggGeom, Options.bEmitTransaction);
+}
 
 
 #undef LOCTEXT_NAMESPACE
diff --git a/Engine/Plugins/Experimental/GeometryScripting/Source/GeometryScriptingCore/Public/GeometryScript/CollisionFunctions.h b/Engine/Plugins/Experimental/GeometryScripting/Source/GeometryScriptingCore/Public/GeometryScript/CollisionFunctions.h
--- a/Engine/Plugins/Experimental/GeometryScripting/Source/GeometryScriptingCore/Public/GeometryScript/CollisionFunctions.h
+++ b/Engine/Plugins/Experimental/GeometryScripting/Source/GeometryScriptingCore/Public/GeometryScript/CollisionFunctions.h
@@ -153,5 +153,52 @@ public:
 		bool bEmitTransaction = false,
 		UGeometryScriptDebug* Debug = nullptr);
 
+	/*
+	 * Get the simple collision from a Primitive Component
+	 */
+	UFUNCTION(BlueprintCallable, Category = "GeometryScript|Collision")
+	static UPARAM(DisplayName = "Simple Collision") FGeometryScriptSimpleCollision
+	GetSimpleCollisionFromComponent(
+		UPrimitiveComponent* Component,
+		UGeometryScriptDebug* Debug = nullptr);
+
+	/*
+	 * Set the simple collision on a Dynamic Mesh Component
+	 */
+	UFUNCTION(BlueprintCallable, Category = "GeometryScript|Collision")
+	static void
+	SetSimpleCollisionOfDynamicMeshComponent(
+		const FGeometryScriptSimpleCollision& SimpleCollision,
+		UDynamicMeshComponent* DynamicMeshComponent,
+		FGeometryScriptSetSimpleCollisionOptions Options,
+		UGeometryScriptDebug* Debug = nullptr);
+
+	/*
+	 * Get the simple collision from a Static Mesh
+	 */
+	UFUNCTION(BlueprintCallable, Category = "GeometryScript|Collision")
+	static UPARAM(DisplayName = "Simple Collision") FGeometryScriptSimpleCollision
+	GetSimpleCollisionFromStaticMesh(UStaticMesh* StaticMesh, UGeometryScriptDebug* Debug = nullptr);
+
+	/*
+	 * Set the simple collision on a Static Mesh
+	 */
+	UFUNCTION(BlueprintCallable, Category = "GeometryScript|Collision")
+	static void
+	SetSimpleCollisionOfStaticMesh(
+		const FGeometryScriptSimpleCollision& SimpleCollision,
+		UStaticMesh* StaticMesh, 
+		FGeometryScriptSetSimpleCollisionOptions Options, 
+		UGeometryScriptDebug* Debug = nullptr);
+
+	/*
+	 * Count of number of simple collision shapes
+	 */
+	UFUNCTION(BlueprintCallable, Category = "GeometryScript|Collision", meta = (ScriptMethod))
+	static int32 GetSimpleCollisionShapeCount(const FGeometryScriptSimpleCollision& SimpleCollision)
+	{
+		return SimpleCollision.AggGeom.GetElementCount();
+	}
+
 
 };
\ No newline at end of file
diff --git a/Engine/Plugins/Experimental/GeometryScripting/Source/GeometryScriptingCore/Public/GeometryScript/GeometryScriptTypes.h b/Engine/Plugins/Experimental/GeometryScripting/Source/GeometryScriptingCore/Public/GeometryScript/GeometryScriptTypes.h
--- a/Engine/Plugins/Experimental/GeometryScripting/Source/GeometryScriptingCore/Public/GeometryScript/GeometryScriptTypes.h
+++ b/Engine/Plugins/Experimental/GeometryScripting/Source/GeometryScriptingCore/Public/GeometryScript/GeometryScriptTypes.h
@@ -4,6 +4,7 @@
 
 #include "CoreMinimal.h"
 #include "GeometryBase.h"
+#include "PhysicsEngine/AggregateGeom.h"
 #include "GeometryScriptTypes.generated.h"
 
 PREDECLARE_GEOMETRY(class FDynamicMesh3);
@@ -90,6 +91,24 @@ public:
 
 
 
+//
+// Collision Shapes
+//
+
+/**
+ * ~This is a wrapper to let Blueprint functions pass around the FKAggregateGeom simple collision primitives
+ * Holds simple shapes that can be used for collision
+ */
+USTRUCT(BlueprintType)
+struct GEOMETRYSCRIPTINGCORE_API FGeometryScriptSimpleCollision
+{
+	GENERATED_BODY()
+
+public:
+	UPROPERTY()
+	FKAggregateGeom AggGeom;
+	
+};
 
 
 //
