Add an 'AppendSimpleCollisionShapes' Geometry Script function, which constructs mesh representations of simple collision shapes

diff --git a/Engine/Plugins/Experimental/GeometryScripting/Source/GeometryScriptingCore/Private/MeshPrimitiveFunctions.cpp b/Engine/Plugins/Experimental/GeometryScripting/Source/GeometryScriptingCore/Private/MeshPrimitiveFunctions.cpp
--- a/Engine/Plugins/Experimental/GeometryScripting/Source/GeometryScriptingCore/Private/MeshPrimitiveFunctions.cpp
+++ b/Engine/Plugins/Experimental/GeometryScripting/Source/GeometryScriptingCore/Private/MeshPrimitiveFunctions.cpp
@@ -29,9 +29,19 @@ using namespace UE::Geometry;
 #define LOCTEXT_NAMESPACE "UGeometryScriptLibrary_MeshPrimitiveFunctions"
 
 
-static void ApplyPrimitiveOptionsToMesh(FDynamicMesh3& Mesh, const FTransform& Transform, FGeometryScriptPrimitiveOptions PrimitiveOptions, FVector3d PreTranslate = FVector3d::Zero())
+static void ApplyPrimitiveOptionsToMesh(
+	FDynamicMesh3& Mesh, const FTransform& Transform, 
+	FGeometryScriptPrimitiveOptions PrimitiveOptions, 
+	FVector3d PreTranslate = FVector3d::Zero(),
+	TOptional<FQuaterniond> PreRotate = TOptional<FQuaterniond>())
 {
-	if (PreTranslate.SquaredLength() > 0)
+	bool bHasTranslate = PreTranslate.SquaredLength() > 0;
+	if (PreRotate.IsSet())
+	{
+		FFrame3d Frame(PreTranslate, *PreRotate);
+		MeshTransforms::FrameCoordsToWorld(Mesh, Frame);
+	}
+	else if (bHasTranslate)
 	{
 		MeshTransforms::Translate(Mesh, PreTranslate);
 	}
@@ -63,20 +73,21 @@ static void AppendPrimitive(
 	FMeshShapeGenerator* Generator, 
 	FTransform Transform, 
 	FGeometryScriptPrimitiveOptions PrimitiveOptions,
-	FVector3d PreTranslate = FVector3d::Zero())
+	FVector3d PreTranslate = FVector3d::Zero(),
+	TOptional<FQuaterniond> PreRotate = TOptional<FQuaterniond>())
 {
 	if (TargetMesh->IsEmpty())
 	{
 		TargetMesh->EditMesh([&](FDynamicMesh3& EditMesh)
 		{
 			EditMesh.Copy(Generator);
-			ApplyPrimitiveOptionsToMesh(EditMesh, Transform, PrimitiveOptions, PreTranslate);
+			ApplyPrimitiveOptionsToMesh(EditMesh, Transform, PrimitiveOptions, PreTranslate, PreRotate);
 		}, EDynamicMeshChangeType::GeneralEdit, EDynamicMeshAttributeChangeFlags::Unknown, false);
 	}
 	else
 	{
 		FDynamicMesh3 TempMesh(Generator);
-		ApplyPrimitiveOptionsToMesh(TempMesh, Transform, PrimitiveOptions, PreTranslate);
+		ApplyPrimitiveOptionsToMesh(TempMesh, Transform, PrimitiveOptions, PreTranslate, PreRotate);
 		TargetMesh->EditMesh([&](FDynamicMesh3& EditMesh)
 		{
 			FMeshIndexMappings TmpMappings;
@@ -1401,4 +1412,78 @@ UDynamicMesh* UGeometryScriptLibrary_MeshPrimitiveFunctions::AppendPolygonListTr
 }
 
 
+UDynamicMesh* UGeometryScriptLibrary_MeshPrimitiveFunctions::AppendSimpleCollisionShapes(
+	UDynamicMesh* TargetMesh,
+	FGeometryScriptPrimitiveOptions PrimitiveOptions,
+	FTransform Transform,
+	const FGeometryScriptSimpleCollision& SimpleCollision,
+	FGeometryScriptSimpleCollisionTriangulationOptions TriangulationOptions,
+	UGeometryScriptDebug* Debug
+)
+{
+	if (TargetMesh == nullptr)
+	{
+		UE::Geometry::AppendError(Debug, EGeometryScriptErrorType::InvalidInputs, LOCTEXT("PrimitiveFunctions_AppendSimpleCollisionShapes", "AppendSimpleCollisionShapes: TargetMesh is Null"));
+		return TargetMesh;
+	}
+
+	for (const FKBoxElem& Box : SimpleCollision.AggGeom.BoxElems)
+	{
+		FGridBoxMeshGenerator GridBoxGenerator;
+		GridBoxGenerator.Box.Extents = FVector(Box.X * .5, Box.Y * .5, Box.Z * .5);
+		GridBoxGenerator.Box.Frame.Origin = Box.Center;
+		GridBoxGenerator.Box.Frame.Rotation = (FQuaterniond)Box.Rotation;
+		GridBoxGenerator.EdgeVertices = FIndex3i(2, 2, 2);
+		GridBoxGenerator.bPolygroupPerQuad = (PrimitiveOptions.PolygroupMode == EGeometryScriptPrimitivePolygroupMode::PerQuad);
+		GridBoxGenerator.Generate();
+		AppendPrimitive(TargetMesh, &GridBoxGenerator, Transform, PrimitiveOptions);
+	}
+	for (const FKSphereElem& Sphere : SimpleCollision.AggGeom.SphereElems)
+	{
+		FBoxSphereGenerator SphereGenerator;
+		SphereGenerator.Box.Frame.Origin = Sphere.Center;
+		SphereGenerator.Radius = FMath::Max(FMathf::ZeroTolerance, Sphere.Radius);
+		int32 StepsPerSide = FMath::Max(1, TriangulationOptions.SphereStepsPerSide);
+		SphereGenerator.EdgeVertices = FIndex3i(StepsPerSide, StepsPerSide, StepsPerSide);
+		SphereGenerator.bPolygroupPerQuad = (PrimitiveOptions.PolygroupMode == EGeometryScriptPrimitivePolygroupMode::PerQuad);
+		SphereGenerator.Generate();
+		AppendPrimitive(TargetMesh, &SphereGenerator, Transform, PrimitiveOptions);
+	}
+	for (const FKSphylElem& Capsule : SimpleCollision.AggGeom.SphylElems)
+	{
+		FCapsuleGenerator CapsuleGenerator;
+		CapsuleGenerator.Radius = Capsule.Radius;
+		CapsuleGenerator.SegmentLength = Capsule.Length;
+		CapsuleGenerator.NumHemisphereArcSteps = FMath::Max(2, TriangulationOptions.CapsuleHemisphereSteps);
+		CapsuleGenerator.NumCircleSteps = FMath::Max(3, TriangulationOptions.CapsuleCircleSteps);
+		CapsuleGenerator.bPolygroupPerQuad = (PrimitiveOptions.PolygroupMode == EGeometryScriptPrimitivePolygroupMode::PerQuad);
+		CapsuleGenerator.Generate();
+
+		AppendPrimitive(TargetMesh, &CapsuleGenerator, Transform, PrimitiveOptions, Capsule.Center, FQuaterniond(Capsule.Rotation));
+	}
+	for (const FKConvexElem& Convex : SimpleCollision.AggGeom.ConvexElems)
+	{
+		FDynamicMesh3 ConvexMesh;
+		for (FVector V : Convex.VertexData)
+		{
+			ConvexMesh.AppendVertex(V);
+		}
+		for (int32 TriStart = 0; TriStart + 2 < Convex.IndexData.Num(); TriStart += 3)
+		{
+			ConvexMesh.AppendTriangle(Convex.IndexData[TriStart], Convex.IndexData[TriStart + 2], Convex.IndexData[TriStart + 1]);
+		}
+		ConvexMesh.EnableAttributes();
+		FMeshNormals::InitializeOverlayToPerTriangleNormals(ConvexMesh.Attributes()->PrimaryNormals());
+		AppendPrimitiveMesh(TargetMesh, ConvexMesh, Transform, PrimitiveOptions);
+	}
+	if (!SimpleCollision.AggGeom.TaperedCapsuleElems.IsEmpty() || !SimpleCollision.AggGeom.SkinnedLevelSetElems.IsEmpty() || !SimpleCollision.AggGeom.LevelSetElems.IsEmpty())
+	{
+		// Tapered capsules and level sets are not supported yet
+		UE::Geometry::AppendWarning(Debug, EGeometryScriptErrorType::OperationFailed, LOCTEXT("PrimitiveFunctions_AppendSimpleCollisionShapes Tapered Capsules and Level Sets Unsupported", "AppendSimpleCollisionShapes: Tapered Capsules and Level Sets are not supported and will be skipped"));
+	}
+
+	return TargetMesh;
+}
+
+
 #undef LOCTEXT_NAMESPACE
\ No newline at end of file
diff --git a/Engine/Plugins/Experimental/GeometryScripting/Source/GeometryScriptingCore/Public/GeometryScript/GeometryScriptTypes.h b/Engine/Plugins/Experimental/GeometryScripting/Source/GeometryScriptingCore/Public/GeometryScript/GeometryScriptTypes.h
--- a/Engine/Plugins/Experimental/GeometryScripting/Source/GeometryScriptingCore/Public/GeometryScript/GeometryScriptTypes.h
+++ b/Engine/Plugins/Experimental/GeometryScripting/Source/GeometryScriptingCore/Public/GeometryScript/GeometryScriptTypes.h
@@ -110,6 +110,23 @@ public:
 	
 };
 
+// Settings to control the triangulation of simple collision primitives -- used for conversion to mesh or convex hull geometry
+USTRUCT(BlueprintType)
+struct GEOMETRYSCRIPTINGCORE_API FGeometryScriptSimpleCollisionTriangulationOptions
+{
+	GENERATED_BODY()
+public:
+
+	UPROPERTY(EditAnywhere, BlueprintReadWrite, Category = Options)
+	int32 SphereStepsPerSide = 4;
+
+	UPROPERTY(EditAnywhere, BlueprintReadWrite, Category = Options)
+	int32 CapsuleHemisphereSteps = 5;
+
+	UPROPERTY(EditAnywhere, BlueprintReadWrite, Category = Options)
+	int32 CapsuleCircleSteps = 8;
+};
+
 
 //
 // Triangles
diff --git a/Engine/Plugins/Experimental/GeometryScripting/Source/GeometryScriptingCore/Public/GeometryScript/MeshPrimitiveFunctions.h b/Engine/Plugins/Experimental/GeometryScripting/Source/GeometryScriptingCore/Public/GeometryScript/MeshPrimitiveFunctions.h
--- a/Engine/Plugins/Experimental/GeometryScripting/Source/GeometryScriptingCore/Public/GeometryScript/MeshPrimitiveFunctions.h
+++ b/Engine/Plugins/Experimental/GeometryScripting/Source/GeometryScriptingCore/Public/GeometryScript/MeshPrimitiveFunctions.h
@@ -574,7 +574,16 @@ public:
 		UGeometryScriptDebug* Debug = nullptr);
 
 
-
+	UFUNCTION(BlueprintCallable, Category = "GeometryScript|Primitives")
+	static UPARAM(DisplayName = "Target Mesh") UDynamicMesh*
+	AppendSimpleCollisionShapes(
+		UDynamicMesh* TargetMesh,
+		FGeometryScriptPrimitiveOptions PrimitiveOptions,
+		FTransform Transform,
+		const FGeometryScriptSimpleCollision& SimpleCollision,
+		FGeometryScriptSimpleCollisionTriangulationOptions TriangulationOptions,
+		UGeometryScriptDebug* Debug = nullptr
+	);
 
 
 
