Add a MergeSimpleCollisionShapes function to Geometry Script, which can reduce the number of simple collision shapes by merging them. Currently works via convex hull merging, so merged shapes will always be convex hulls, but shapes that are not merged will be left with their original type.

Add a 'source ID' concept to the convex part decomposition, to enable tracking of which hulls were not merged to any other hulls when running the convex hull merging algorithm.

diff --git a/Engine/Plugins/Experimental/GeometryScripting/Source/GeometryScriptingCore/Private/CollisionFunctions.cpp b/Engine/Plugins/Experimental/GeometryScripting/Source/GeometryScriptingCore/Private/CollisionFunctions.cpp
--- a/Engine/Plugins/Experimental/GeometryScripting/Source/GeometryScriptingCore/Private/CollisionFunctions.cpp
+++ b/Engine/Plugins/Experimental/GeometryScripting/Source/GeometryScriptingCore/Private/CollisionFunctions.cpp
@@ -17,6 +17,16 @@
 #include "ShapeApproximation/ShapeDetection3.h"
 #include "ShapeApproximation/MeshSimpleShapeApproximation.h"
 
+#include "CompGeom/ConvexDecomposition3.h"
+#include "OrientedBoxTypes.h"
+#include "MeshQueries.h"
+#include "MeshAdapter.h"
+
+#include "Generators/MeshShapeGenerator.h"
+#include "Generators/GridBoxMeshGenerator.h"
+#include "Generators/BoxSphereGenerator.h"
+#include "Generators/CapsuleGenerator.h"
+
 // physics data
 #include "Engine/StaticMesh.h"
 #include "Components/StaticMeshComponent.h"
@@ -471,4 +481,192 @@ void UGeometryScriptLibrary_CollisionFunctions::SetSimpleCollisionOfStaticMesh(
 }
 
 
+FGeometryScriptSimpleCollision UGeometryScriptLibrary_CollisionFunctions::MergeSimpleCollisionShapes(
+	const FGeometryScriptSimpleCollision& SimpleCollision,
+	const FGeometryScriptMergeSimpleCollisionOptions& MergeOptions,
+	bool& bHasMerged,
+	UGeometryScriptDebug* Debug
+)
+{
+	FGeometryScriptSimpleCollision ToRet;
+	bHasMerged = false;
+	
+	// Nothing to merge
+	if (SimpleCollision.AggGeom.GetElementCount() <= 1)
+	{
+		return SimpleCollision;
+	}
+
+	TArray<FVector> HullVertices;
+	TArray<int32> HullVertexCounts;
+	TArray<double> HullVolumes;
+	TArray<const FKShapeElem*> HullToShapeElem;
+
+	auto TransformVertices = [](TArrayView<FVector3d> Vertices, const FTransform& Transform)
+	{
+		for (FVector3d& Vertex : Vertices)
+		{
+			Vertex = Transform.TransformPosition(Vertex);
+		}
+	};
+	auto AppendHullVertices = [&HullToShapeElem, &HullVolumes, &HullVertices, &HullVertexCounts]
+				(TArrayView<const FVector3d> Vertices, double Volume, const FKShapeElem* ShapeElem)
+	{
+		check(HullToShapeElem.Num() == HullVolumes.Num());
+		HullToShapeElem.Add(ShapeElem);
+		HullVertices.Append(Vertices);
+		HullVertexCounts.Add(Vertices.Num());
+		HullVolumes.Add(Volume);
+	};
+	auto GeneratorVolume = [](FMeshShapeGenerator* Generator) -> double
+	{
+		TIndexVectorMeshArrayAdapter<FIndex3i, double, FVector3d> GenMeshAdapter(&Generator->Vertices, &Generator->Triangles);
+		FVector2d VolArea = TMeshQueries<TIndexVectorMeshArrayAdapter<FIndex3i, double, FVector3d>>::GetVolumeArea(GenMeshAdapter);
+		return VolArea.X;
+	};
+
+	for (const FKBoxElem& Box : SimpleCollision.AggGeom.BoxElems)
+	{
+		FOrientedBox3d OrientedBox;
+		OrientedBox.Extents = FVector(Box.X * .5, Box.Y * .5, Box.Z * .5);
+		OrientedBox.Frame.Origin = Box.Center;
+		OrientedBox.Frame.Rotation = (FQuaterniond)Box.Rotation;
+		TArray<FVector3d, TFixedAllocator<8>> BoxVertices;
+		OrientedBox.EnumerateCorners([&](FVector3d Corner) { BoxVertices.Add(Corner); });
+		AppendHullVertices(BoxVertices, Box.GetScaledVolume(FVector3d::One()), &Box);
+	}
+	for (const FKSphereElem& Sphere : SimpleCollision.AggGeom.SphereElems)
+	{
+		FBoxSphereGenerator SphereGenerator;
+		SphereGenerator.Box.Frame.Origin = Sphere.Center;
+		SphereGenerator.Radius = FMath::Max(FMathf::ZeroTolerance, Sphere.Radius);
+		int32 StepsPerSide = FMath::Max(1, MergeOptions.ShapeToHullTriangulation.SphereStepsPerSide);
+		SphereGenerator.EdgeVertices = FIndex3i(StepsPerSide, StepsPerSide, StepsPerSide);
+		SphereGenerator.bPolygroupPerQuad = false;
+		SphereGenerator.Generate();
+		double Volume = GeneratorVolume(&SphereGenerator);
+		AppendHullVertices(SphereGenerator.Vertices, Volume, &Sphere);
+	}
+	for (const FKSphylElem& Capsule : SimpleCollision.AggGeom.SphylElems)
+	{
+		FCapsuleGenerator CapsuleGenerator;
+		CapsuleGenerator.Radius = Capsule.Radius;
+		CapsuleGenerator.SegmentLength = Capsule.Length;
+		CapsuleGenerator.NumHemisphereArcSteps = FMath::Max(2, MergeOptions.ShapeToHullTriangulation.CapsuleHemisphereSteps);
+		CapsuleGenerator.NumCircleSteps = FMath::Max(3, MergeOptions.ShapeToHullTriangulation.CapsuleCircleSteps);
+		CapsuleGenerator.bPolygroupPerQuad = false;
+		CapsuleGenerator.Generate();
+		FTransform CapsuleTransform(Capsule.Rotation, Capsule.Center);
+		TransformVertices(CapsuleGenerator.Vertices, CapsuleTransform);
+		double Volume = GeneratorVolume(&CapsuleGenerator);
+		AppendHullVertices(CapsuleGenerator.Vertices, Volume, &Capsule);
+	}
+	for (const FKConvexElem& Convex : SimpleCollision.AggGeom.ConvexElems)
+	{
+		// Note: Not reliable to use the FKConvexElem::GetVolume function because it depends on the chaos convex being allocated, and also is not currently exported
+		TIndexMeshArrayAdapter<int32, double, FVector3d> HullMeshAdapter(&Convex.VertexData, &Convex.IndexData);
+		// Note: We take the negative volume because the hull triangles have opposite winding from ordinary meshes
+		double Volume = -TMeshQueries<TIndexMeshArrayAdapter<int32, double, FVector3d>>::GetVolumeArea(HullMeshAdapter).X;
+		AppendHullVertices(Convex.VertexData, Volume, &Convex);
+	}
+
+	const int32 InitialNumConvex = HullVertexCounts.Num();
+	// Nothing we are able to merge
+	if (InitialNumConvex <= 1)
+	{
+		return SimpleCollision;
+	}
+
+	TArray<int32> HullVertexStarts;
+	HullVertexStarts.SetNumUninitialized(InitialNumConvex);
+	HullVertexStarts[0] = 0; // Note InitialNumConvex is > 1 due to above test
+	for (int32 HullIdx = 1, LastEnd = HullVertexCounts[0]; HullIdx < InitialNumConvex; LastEnd += HullVertexCounts[HullIdx++])
+	{
+		HullVertexStarts[HullIdx] = LastEnd;
+	}
+
+	// Currently we use dense proximity to propose which shapes can be merged.
+	// Note: To efficiently handle larger shape counts, consider optionally limiting these by some approximate proximity (e.g. expanded bounding box overlap)
+	TArray<TPair<int32, int32>> HullProximity;
+	for (int32 ConvexA = 0; ConvexA < InitialNumConvex; ++ConvexA)
+	{
+		for (int32 ConvexB = ConvexA + 1; ConvexB < InitialNumConvex; ++ConvexB)
+		{
+			HullProximity.Emplace(ConvexA, ConvexB);
+		}
+	}
+
+	FConvexDecomposition3 Decomposition;
+	Decomposition.InitializeFromHulls(HullVertexStarts.Num(),
+		[&HullVolumes](int32 HullIdx) { return HullVolumes[HullIdx]; }, [&HullVertexCounts](int32 HullIdx) { return HullVertexCounts[HullIdx]; },
+		[&HullVertexStarts, &HullVertices](int32 HullIdx, int32 VertIdx) { return HullVertices[HullVertexStarts[HullIdx] + VertIdx]; }, HullProximity);
+	Decomposition.MergeBest(MergeOptions.MaxShapeCount, MergeOptions.ErrorTolerance, 0.0, true, false, MergeOptions.MaxShapeCount, nullptr /*optional negative space*/, nullptr /*optional FTransform for negative space*/);
+
+	// Algorithm decided not to merge
+	if (Decomposition.NumHulls() == InitialNumConvex)
+	{
+		return SimpleCollision;
+	}
+
+	bHasMerged = true;
+
+	// Merging logic for the below primitives is not implemented, so they are simply copied over for now
+	if (!SimpleCollision.AggGeom.TaperedCapsuleElems.IsEmpty() || !SimpleCollision.AggGeom.SkinnedLevelSetElems.IsEmpty() || !SimpleCollision.AggGeom.LevelSetElems.IsEmpty())
+	{
+		UE::Geometry::AppendWarning(Debug, EGeometryScriptErrorType::OperationFailed, LOCTEXT("PrimitiveFunctions_AppendSimpleCollisionShapes Unsupported Shapes", "MergeSimpleCollisionShapes: Merging for Tapered Capsules and Level Set collision is not yet supported; these shapes will be copied without considering them for merging."));
+		for (const FKTaperedCapsuleElem& Capsule : SimpleCollision.AggGeom.TaperedCapsuleElems)
+		{
+			ToRet.AggGeom.TaperedCapsuleElems.Add(Capsule);
+		}
+		for (const FKSkinnedLevelSetElem& LevelSet : SimpleCollision.AggGeom.SkinnedLevelSetElems)
+		{
+			ToRet.AggGeom.SkinnedLevelSetElems.Add(LevelSet);
+		}
+		for (const FKLevelSetElem& LevelSet : SimpleCollision.AggGeom.LevelSetElems)
+		{
+			ToRet.AggGeom.LevelSetElems.Add(LevelSet);
+		}
+	}
+	
+	for (int32 HullIdx = 0; HullIdx < Decomposition.Decomposition.Num(); ++HullIdx)
+	{
+		const FConvexDecomposition3::FConvexPart& Part = Decomposition.Decomposition[HullIdx];
+		// If part was not merged, use the source ID to map it back to the original collision primitive
+		if (Part.HullSourceID >= 0)
+		{
+			const FKShapeElem* Elem = HullToShapeElem[Part.HullSourceID];
+			bool bHandledShape = true;
+			switch (Elem->GetShapeType())
+			{
+			case EAggCollisionShape::Box:
+				ToRet.AggGeom.BoxElems.Add(*static_cast<const FKBoxElem*>(Elem));
+				break;
+			case EAggCollisionShape::Sphere:
+				ToRet.AggGeom.SphereElems.Add(*static_cast<const FKSphereElem*>(Elem));
+				break;
+			case EAggCollisionShape::Sphyl:
+				ToRet.AggGeom.SphylElems.Add(*static_cast<const FKSphylElem*>(Elem));
+				break;
+			case EAggCollisionShape::Convex:
+				ToRet.AggGeom.ConvexElems.Add(*static_cast<const FKConvexElem*>(Elem));
+				break;
+			default:
+				// Note: All shapes that we add to the HullToShapeElem array should be handled above, so we should not reach here
+				ensureMsgf(false, TEXT("Unhandled shape element type could not be restored from source shapes"));
+				bHandledShape = false;
+			}
+			if (bHandledShape)
+			{
+				continue;
+			}
+		}
+		// Add the merged part
+		FKConvexElem& Convex = ToRet.AggGeom.ConvexElems.Emplace_GetRef();
+		Convex.VertexData = Decomposition.GetVertices<double>(HullIdx);
+		Convex.UpdateElemBox(); // Note: In addition to updating the bounding box, this also re-computes hull indices.
+	}
+	
+	return ToRet;
+}
+
 #undef LOCTEXT_NAMESPACE
diff --git a/Engine/Plugins/Experimental/GeometryScripting/Source/GeometryScriptingCore/Public/GeometryScript/CollisionFunctions.h b/Engine/Plugins/Experimental/GeometryScripting/Source/GeometryScriptingCore/Public/GeometryScript/CollisionFunctions.h
--- a/Engine/Plugins/Experimental/GeometryScripting/Source/GeometryScriptingCore/Public/GeometryScript/CollisionFunctions.h
+++ b/Engine/Plugins/Experimental/GeometryScripting/Source/GeometryScriptingCore/Public/GeometryScript/CollisionFunctions.h
@@ -102,6 +102,30 @@ public:
 	bool bEmitTransaction = true;
 };
 
+// Options controlling how collision shapes can be merged together
+USTRUCT(BlueprintType)
+struct GEOMETRYSCRIPTINGCORE_API FGeometryScriptMergeSimpleCollisionOptions
+{
+	GENERATED_BODY()
+public:
+
+	/**
+	 * If > 0, merge down to at most this many simple shapes. (If <= 0, this value is ignored.)
+	 */
+	UPROPERTY(EditAnywhere, BlueprintReadWrite, Category = Options)
+	int MaxShapeCount = 0;
+
+	/**
+	 * Error tolerance to use to decide to convex hulls together, in cm.
+	 * If merging two hulls would increase the volume by more than this ErrorTolerance cubed, the merge is not accepted.
+	 */
+	UPROPERTY(EditAnywhere, BlueprintReadWrite, Category = Options, meta = (UIMin = "0", UIMax = "100.", Units = cm))
+	double ErrorTolerance = 0.0;
+
+	// Controls for how smooth shapes can be triangulated when/if converted to a convex hull for a merge
+	UPROPERTY(EditAnywhere, BlueprintReadWrite, Category = ConvexHulls)
+	FGeometryScriptSimpleCollisionTriangulationOptions ShapeToHullTriangulation;
+};
 
 
 UCLASS(meta = (ScriptName = "GeometryScript_Collision"))
@@ -201,4 +225,21 @@ public:
 	}
 
 
+	/**
+	 * Attempt to merge collision shapes to create a representation with fewer overall shapes.
+	 * 
+	 * @param SimpleCollision		The collision to attempt to simplify by merging shapes
+	 * @param MergeOptions			Options controlling how shapes can be merged
+	 * @param bHasMerged			Indicates whether any shapes have been merged
+	 * @return						Simple Collision with collision shapes merged, as allowed by settings
+	 */
+	UFUNCTION(BlueprintCallable, Category = "GeometryScript|Collision")
+	static UPARAM(DisplayName = "Merged Simple Collision") FGeometryScriptSimpleCollision MergeSimpleCollisionShapes(
+		const FGeometryScriptSimpleCollision& SimpleCollision,
+		const FGeometryScriptMergeSimpleCollisionOptions& MergeOptions,
+		bool& bHasMerged,
+		UGeometryScriptDebug* Debug = nullptr
+	);
+
+
 };
\ No newline at end of file
diff --git a/Engine/Source/Runtime/GeometryCore/Private/CompGeom/ConvexDecomposition3.cpp b/Engine/Source/Runtime/GeometryCore/Private/CompGeom/ConvexDecomposition3.cpp
--- a/Engine/Source/Runtime/GeometryCore/Private/CompGeom/ConvexDecomposition3.cpp
+++ b/Engine/Source/Runtime/GeometryCore/Private/CompGeom/ConvexDecomposition3.cpp
@@ -1513,6 +1513,7 @@ void FConvexDecomposition3::InitializeFromHulls(int32 NumHulls, TFunctionRef<dou
 	for (int32 HullIdx = 0; HullIdx < NumHulls; ++HullIdx)
 	{
 		FConvexPart* Convex = new FConvexPart(true);
+		Convex->HullSourceID = HullIdx;
 		Convex->bGeometryVolumeUnreliable = false;
 		Convex->bMustMerge = false;
 		double Volume = HullVolumes(HullIdx);
diff --git a/Engine/Source/Runtime/GeometryCore/Public/CompGeom/ConvexDecomposition3.h b/Engine/Source/Runtime/GeometryCore/Public/CompGeom/ConvexDecomposition3.h
--- a/Engine/Source/Runtime/GeometryCore/Public/CompGeom/ConvexDecomposition3.h
+++ b/Engine/Source/Runtime/GeometryCore/Public/CompGeom/ConvexDecomposition3.h
@@ -330,6 +330,21 @@ public:
 		return Decomposition[HullIdx].InternalGeo;
 	}
 
+	const int32 GetHullSourceID(int32 HullIdx) const
+	{
+		return Decomposition[HullIdx].HullSourceID;
+	}
+
+	const int32 CountMergedParts() const
+	{
+		int32 Count = 0;
+		for (int32 HullIdx = 0; HullIdx < Decomposition.Num(); ++HullIdx)
+		{
+			Count += int32(Decomposition[HullIdx].HullSourceID < 0);
+		}
+		return Count;
+	}
+
 	// Representation of a convex hull in the decomposition + associated information to help further split or merge
 	struct FConvexPart
 	{
@@ -346,6 +361,7 @@ public:
 			InternalGeo.Clear();
 			HullTriangles.Reset();
 			HullPlanes.Reset();
+			HullSourceID = -1;
 			HullVolume = 0;
 			GeoVolume = 0;
 			SumHullsVolume = -FMathd::MaxReal;
@@ -375,6 +391,8 @@ public:
 		TArray<FIndex3i> HullTriangles;
 		TArray<FPlane3d> HullPlanes; // 1:1 with HullTriangles
 
+		int32 HullSourceID = -1; // Optional ID, cleared on merge, to track the origin of un-merged hulls
+
 		// Measurements of the geo and hull, to be used when evaluating potential further splits
 		double HullVolume = 0, GeoVolume = 0;
 		double SumHullsVolume = -FMathd::MaxReal; // Sum of volume of any hulls that have been merged to form this hull; only valid if the FConvexPart was created by merging
