Render Capture Baking: Add ability to explicity specify the render capture cameras to use in the BakeRC Geometry Script node

diff --git a/Engine/Plugins/Experimental/GeometryScripting/Source/GeometryScriptingCore/Private/MeshBakeFunctions.cpp b/Engine/Plugins/Experimental/GeometryScripting/Source/GeometryScriptingCore/Private/MeshBakeFunctions.cpp
--- a/Engine/Plugins/Experimental/GeometryScripting/Source/GeometryScriptingCore/Private/MeshBakeFunctions.cpp
+++ b/Engine/Plugins/Experimental/GeometryScripting/Source/GeometryScriptingCore/Private/MeshBakeFunctions.cpp
@@ -863,6 +863,18 @@ namespace GeometryScriptBakeLocals
 } // end namespace GeometryScriptBakeLocals
 
 
+
+int UGeometryScriptLibrary_MeshBakeFunctions::ConvertBakeResolutionToInt(EGeometryScriptBakeResolution BakeResolution)
+{
+	FImageDimensions Dimensions = GeometryScriptBakeLocals::GetDimensions(BakeResolution);
+
+	int Output;
+	Output = Dimensions.GetWidth();
+	checkSlow(Output == Dimensions.GetHeight());
+	return Output;
+}
+
+
 FGeometryScriptBakeTypeOptions UGeometryScriptLibrary_MeshBakeFunctions::MakeBakeTypeTangentNormal()
 {
 	FGeometryScriptBakeTypeOptions Output;
@@ -1111,9 +1123,57 @@ FGeometryScriptRenderCaptureTextures UGeometryScriptLibrary_MeshBakeFunctions::B
 		return {};
 	}
 
-	const FSceneCaptureConfig Options = GeometryScriptBakeLocals::GetSceneCaptureConfig(BakeOptions);
 	TUniquePtr<FSceneCapturePhotoSet> SceneCapture = MakeUnique<FSceneCapturePhotoSet>();
-	ConfigureSceneCapture(SceneCapture, ValidSourceActors, Options, false);
+	const FSceneCaptureConfig Options = GeometryScriptBakeLocals::GetSceneCaptureConfig(BakeOptions);
+
+	// TODO This block is similar to ConfigureSceneCapture function but handles an explicit Cameras array, when the BakeRC
+	// tool also supports explicit Cameras then we can re-unify the code paths
+	{
+		ForEachCaptureType([&SceneCapture, &Options](ERenderCaptureType CaptureType)
+		{
+			const bool bCaptureTypeEnabled = Options.Flags[CaptureType];
+			SceneCapture->SetCaptureTypeEnabled(CaptureType, bCaptureTypeEnabled);
+
+			FRenderCaptureConfig Config;
+			Config.bAntiAliasing = (CaptureType == ERenderCaptureType::DeviceDepth ? false : Options.bAntiAliasing);
+			SceneCapture->SetCaptureConfig(CaptureType, Config);
+		});
+
+		SceneCapture->SetCaptureSceneActors(ValidSourceActors[0]->GetWorld(), ValidSourceActors);
+
+		if (BakeOptions.Cameras.IsEmpty())
+		{
+			AppendWarning(Debug, EGeometryScriptErrorType::InvalidInputs,
+				LOCTEXT("BakeTextureFromRenderCaptures_EmptyCameraSet", "BakeTextureFromRenderCaptures: Cameras array was empty so a default camera set was used. Please update your blueprint by providing an explicit array e.g., by using ComputeRenderCaptureCamerasForBox"));
+		
+			TArray<FSpatialPhotoParams> SpatialParams = ComputeStandardExteriorSpatialPhotoParameters(
+				nullptr,
+				ValidSourceActors,
+				TArray<UActorComponent*>(),
+				FImageDimensions(Options.RenderCaptureImageSize, Options.RenderCaptureImageSize),
+				Options.FieldOfViewDegrees,
+				Options.NearPlaneDist,
+				true, true, true, true, true);
+			SceneCapture->SetSpatialPhotoParams(SpatialParams);
+		}
+		else
+		{
+			TArray<FSpatialPhotoParams> SpatialParams;
+			for (const FGeometryScriptRenderCaptureCamera& Camera : BakeOptions.Cameras)
+			{
+				FSpatialPhotoParams Params;
+				Params.NearPlaneDist = Camera.NearPlaneDist;
+				Params.HorzFOVDegrees = Camera.FieldOfViewDegrees;
+				Params.Dimensions = FImageDimensions(FMath::Max(1, Camera.Resolution), FMath::Max(1, Camera.Resolution));
+				Params.Frame.AlignAxis(0, Camera.ViewDirection);
+				Params.Frame.ConstrainedAlignAxis(2, FVector3d::UnitZ(), Params.Frame.X());
+				Params.Frame.Origin = Camera.ViewPosition;
+				SpatialParams.Add(Params);
+			}
+			SceneCapture->SetSpatialPhotoParams(SpatialParams);
+		}
+	}
+
 	SceneCapture->Compute();
 
 	const FDynamicMeshAABBTree3 TargetMeshSpatial(TargetMesh->GetMeshPtr());
diff --git a/Engine/Plugins/Experimental/GeometryScripting/Source/GeometryScriptingCore/Public/GeometryScript/GeometryScriptTypes.h b/Engine/Plugins/Experimental/GeometryScripting/Source/GeometryScriptingCore/Public/GeometryScript/GeometryScriptTypes.h
--- a/Engine/Plugins/Experimental/GeometryScripting/Source/GeometryScriptingCore/Public/GeometryScript/GeometryScriptTypes.h
+++ b/Engine/Plugins/Experimental/GeometryScripting/Source/GeometryScriptingCore/Public/GeometryScript/GeometryScriptTypes.h
@@ -626,6 +626,81 @@ struct TStructOpsTypeTraits<FGeometryScriptDynamicMeshBVH> : public TStructOpsTy
 
 
 
+
+//
+// Render Capture data structures
+//
+
+
+USTRUCT(BlueprintType)
+struct GEOMETRYSCRIPTINGCORE_API FGeometryScriptRenderCaptureCamera
+{
+	GENERATED_BODY()
+
+	/** The pixel resolution of render capture photo set, this value is used for width and height */
+	UPROPERTY(EditAnywhere, BlueprintReadWrite, Category = Options)
+	int Resolution = 256;
+
+	UPROPERTY(EditAnywhere, BlueprintReadWrite, Category = Options)
+	double FieldOfViewDegrees = 45.0;
+
+	UPROPERTY(EditAnywhere, BlueprintReadWrite, Category = Options)
+	FVector ViewPosition = FVector::Zero();
+
+	UPROPERTY(EditAnywhere, BlueprintReadWrite, Category = Options)
+	FVector ViewDirection = FVector::UnitX();
+
+	UPROPERTY(EditAnywhere, BlueprintReadWrite, Category = Options)
+	double NearPlaneDist = 1.0;
+};
+
+
+
+USTRUCT(BlueprintType)
+struct GEOMETRYSCRIPTINGCORE_API FGeometryScriptRenderCaptureCamerasForBoxOptions
+{
+	GENERATED_BODY()
+
+	/** The pixel resolution of render capture photos, this value is used for width and height */
+	UPROPERTY(EditAnywhere, BlueprintReadWrite, Category = Options)
+	int Resolution = 256;
+
+	UPROPERTY(EditAnywhere, BlueprintReadWrite, Category = Options)
+	double FieldOfViewDegrees = 45.;
+
+	/** Enable 6 directions corresponding to views from box face centers to the box center */
+	UPROPERTY(EditAnywhere, BlueprintReadWrite, Category = Options)
+	bool bViewFromBoxFaces = true;
+
+	/** Enable 4 directions corresponding to views from box upper corners to the box center */
+	UPROPERTY(EditAnywhere, BlueprintReadWrite, Category = Options)
+	bool bViewFromUpperCorners = false;
+
+	/** Enable 4 directions corresponding to views from box lower corners to the box center */
+	UPROPERTY(EditAnywhere, BlueprintReadWrite, Category = Options)
+	bool bViewFromLowerCorners = false;
+
+	/** Enable 4 directions corresponding to views from box upper edges centers to the box center */
+	UPROPERTY(EditAnywhere, BlueprintReadWrite, Category = Options)
+	bool bViewFromUpperEdges = false;
+
+	/** Enable 4 directions corresponding to views from box lower edges centers to the box center */
+	UPROPERTY(EditAnywhere, BlueprintReadWrite, Category = Options)
+	bool bViewFromLowerEdges = false;
+
+	/** Enable 4 directions corresponding to views from box side edges centers to the box center */
+	UPROPERTY(EditAnywhere, BlueprintReadWrite, Category = Options)
+	bool bViewFromSideEdges = false;
+
+	/** Extra positions from which to deduce view directions on the box center (located at (0,0,0)) */
+	UPROPERTY(EditAnywhere, BlueprintReadWrite, Category = Options)
+	TArray<FVector> ExtraViewFromPositions;
+};
+
+
+
+
+
 //
 // Errors/Debugging
 //
diff --git a/Engine/Plugins/Experimental/GeometryScripting/Source/GeometryScriptingCore/Public/GeometryScript/MeshBakeFunctions.h b/Engine/Plugins/Experimental/GeometryScripting/Source/GeometryScriptingCore/Public/GeometryScript/MeshBakeFunctions.h
--- a/Engine/Plugins/Experimental/GeometryScripting/Source/GeometryScriptingCore/Public/GeometryScript/MeshBakeFunctions.h
+++ b/Engine/Plugins/Experimental/GeometryScripting/Source/GeometryScriptingCore/Public/GeometryScript/MeshBakeFunctions.h
@@ -312,13 +312,48 @@ struct GEOMETRYSCRIPTINGCORE_API FGeometryScriptBakeRenderCaptureOptions
 {
 	GENERATED_BODY()
 
+	/*~
+	 * Options to configure the Render Capture
+	 */
+
+	/**
+	 * If not empty the given Render Capture Cameras are used to create the photos from which to Bake textures
+	 * If empty then a default set of Render Capture Cameras with frustums containing the SourceActors bounding box is used
+	 * Note: Providing a value is highly recommended, the default camera set is intended only to support legacy blueprints
+	 */
+	UPROPERTY(EditAnywhere, BlueprintReadWrite, Category = Options)
+	TArray<FGeometryScriptRenderCaptureCamera> Cameras;
+
+	/**
+	 * The pixel resolution of the default render capture camera photo sets.
+	 * Only used if Cameras is empty
+	 */
+	UPROPERTY(EditAnywhere, BlueprintReadWrite, Category = Options, meta=(DisplayName="Default Capture Resolution"))
+	EGeometryScriptBakeResolution RenderCaptureResolution = EGeometryScriptBakeResolution::Resolution256;
+
+	/**
+	 * Default field of view of the default render capture cameras.
+	 * Only used if Cameras is empty
+	 */
+	UPROPERTY()
+	double FieldOfViewDegrees = 45.0;
+
+	/**
+	 * Default near plane distance for the viewing frustums of the default render capture cameras.
+	 * Only used if Cameras is empty
+	 */
+	UPROPERTY()
+	double NearPlaneDist = 1.0;
+
+
+
+	/*~
+	 * Options to configure the Baking step to create textures from the Render Capture photos
+	 */
+
 	/** The pixel resolution of the generated textures */
 	UPROPERTY(EditAnywhere, BlueprintReadWrite, Category = Options)
 	EGeometryScriptBakeResolution Resolution = EGeometryScriptBakeResolution::Resolution256;
-	
-	/** The pixel resolution of render capture photo set */
-	UPROPERTY(EditAnywhere, BlueprintReadWrite, Category = Options)
-	EGeometryScriptBakeResolution RenderCaptureResolution = EGeometryScriptBakeResolution::Resolution256;
 
 	/** Number of samples per pixel */
 	UPROPERTY(EditAnywhere, BlueprintReadWrite, Category = Options)
@@ -371,15 +406,6 @@ struct GEOMETRYSCRIPTINGCORE_API FGeometryScriptBakeRenderCaptureOptions
 
 	UPROPERTY(EditAnywhere, BlueprintReadWrite, Category = Options)
 	bool bSubsurfaceColorMap = true;
-
-	// These are hidden in the UI right now, we might want to expose them if they turn out to be useful for very large
-	// or very small objects (not tested yet) TODO Figure out if we want to expose these options
-
-	UPROPERTY()
-	double FieldOfViewDegrees = 45.0;
-
-	UPROPERTY()
-	double NearPlaneDist = 1.0;
 };
 
 USTRUCT(BlueprintType)
@@ -447,6 +473,10 @@ class GEOMETRYSCRIPTINGCORE_API UGeometryScriptLibrary_MeshBakeFunctions : publi
 {
 	GENERATED_BODY()
 public:
+
+	UFUNCTION(BlueprintPure, Category = "GeometryScript|Bake", meta=(CompactNodeTitle = "->", BlueprintAutocast))
+	static UPARAM(DisplayName="Resolution Out") int ConvertBakeResolutionToInt(EGeometryScriptBakeResolution BakeResolution);
+
 	UFUNCTION(BlueprintPure, Category = "GeometryScript|Bake")
 	static UPARAM(DisplayName="Bake Type Out") FGeometryScriptBakeTypeOptions MakeBakeTypeTangentNormal();
 
diff --git a/Engine/Plugins/Experimental/GeometryScripting/Source/GeometryScriptingCore/Public/GeometryScript/MeshSamplingFunctions.h b/Engine/Plugins/Experimental/GeometryScripting/Source/GeometryScriptingCore/Public/GeometryScript/MeshSamplingFunctions.h
--- a/Engine/Plugins/Experimental/GeometryScripting/Source/GeometryScriptingCore/Public/GeometryScript/MeshSamplingFunctions.h
+++ b/Engine/Plugins/Experimental/GeometryScripting/Source/GeometryScriptingCore/Public/GeometryScript/MeshSamplingFunctions.h
@@ -93,72 +93,6 @@ public:
 };
 
 
-USTRUCT(BlueprintType)
-struct GEOMETRYSCRIPTINGCORE_API FGeometryScriptRenderCaptureCamera
-{
-	GENERATED_BODY()
-
-	/** The pixel resolution of render capture photo set, this value is used for width and height */
-	UPROPERTY(EditAnywhere, BlueprintReadWrite, Category = Options)
-	int Resolution = 256;
-
-	UPROPERTY(EditAnywhere, BlueprintReadWrite, Category = Options)
-	double FieldOfViewDegrees = 45.0;
-
-	UPROPERTY(EditAnywhere, BlueprintReadWrite, Category = Options)
-	FVector ViewPosition = FVector::Zero();
-
-	UPROPERTY(EditAnywhere, BlueprintReadWrite, Category = Options)
-	FVector ViewDirection = FVector::UnitX();
-
-	UPROPERTY(EditAnywhere, BlueprintReadWrite, Category = Options)
-	double NearPlaneDist = 1.0;
-};
-
-
-
-USTRUCT(BlueprintType)
-struct GEOMETRYSCRIPTINGCORE_API FGeometryScriptRenderCaptureCamerasForBoxOptions
-{
-	GENERATED_BODY()
-
-	/** The pixel resolution of render capture photos, this value is used for width and height */
-	UPROPERTY(EditAnywhere, BlueprintReadWrite, Category = Options)
-	int Resolution = 256;
-
-	UPROPERTY(EditAnywhere, BlueprintReadWrite, Category = Options)
-	double FieldOfViewDegrees = 45.;
-
-	/** Enable 6 directions corresponding to views from box face centers to the box center */
-	UPROPERTY(EditAnywhere, BlueprintReadWrite, Category = Options)
-	bool bViewFromBoxFaces = true;
-
-	/** Enable 4 directions corresponding to views from box upper corners to the box center */
-	UPROPERTY(EditAnywhere, BlueprintReadWrite, Category = Options)
-	bool bViewFromUpperCorners = false;
-
-	/** Enable 4 directions corresponding to views from box lower corners to the box center */
-	UPROPERTY(EditAnywhere, BlueprintReadWrite, Category = Options)
-	bool bViewFromLowerCorners = false;
-
-	/** Enable 4 directions corresponding to views from box upper edges centers to the box center */
-	UPROPERTY(EditAnywhere, BlueprintReadWrite, Category = Options)
-	bool bViewFromUpperEdges = false;
-
-	/** Enable 4 directions corresponding to views from box lower edges centers to the box center */
-	UPROPERTY(EditAnywhere, BlueprintReadWrite, Category = Options)
-	bool bViewFromLowerEdges = false;
-
-	/** Enable 4 directions corresponding to views from box side edges centers to the box center */
-	UPROPERTY(EditAnywhere, BlueprintReadWrite, Category = Options)
-	bool bViewFromSideEdges = false;
-
-	/** Extra positions from which to deduce view directions on the box center (located at (0,0,0)) */
-	UPROPERTY(EditAnywhere, BlueprintReadWrite, Category = Options)
-	TArray<FVector> ExtraViewFromPositions;
-};
-
-
 UCLASS(meta = (ScriptName = "GeometryScript_MeshSampling"))
 class GEOMETRYSCRIPTINGCORE_API UGeometryScriptLibrary_MeshSamplingFunctions : public UBlueprintFunctionLibrary
 {
diff --git a/Engine/Plugins/Runtime/MeshModelingToolset/Source/ModelingComponents/Private/Scene/SceneCapturePhotoSet.cpp b/Engine/Plugins/Runtime/MeshModelingToolset/Source/ModelingComponents/Private/Scene/SceneCapturePhotoSet.cpp
--- a/Engine/Plugins/Runtime/MeshModelingToolset/Source/ModelingComponents/Private/Scene/SceneCapturePhotoSet.cpp
+++ b/Engine/Plugins/Runtime/MeshModelingToolset/Source/ModelingComponents/Private/Scene/SceneCapturePhotoSet.cpp
@@ -992,8 +992,9 @@ void FSceneCapturePhotoSet::EmptyPhotoSet(ERenderCaptureType CaptureType)
 	}
 }
 
+// TODO Deprecate this function and replace it with ComputeRenderCaptureCamerasForBox
 TArray<FSpatialPhotoParams> UE::Geometry::ComputeStandardExteriorSpatialPhotoParameters(
-	UWorld* World,
+	UWorld* Unused,
 	const TArray<AActor*>& Actors,
 	const TArray<UActorComponent*>& Components,
 	FImageDimensions PhotoDimensions,
@@ -1005,7 +1006,7 @@ TArray<FSpatialPhotoParams> UE::Geometry::ComputeStandardExteriorSpatialPhotoPar
 	bool bUpperEdges,
 	bool bSideEdges)
 {
-	if (!World || (Actors.IsEmpty() && Components.IsEmpty()))
+	if (Actors.IsEmpty() && Components.IsEmpty())
 	{
 		return {};
 	}
diff --git a/Engine/Plugins/Runtime/MeshModelingToolset/Source/ModelingComponents/Public/Scene/SceneCapturePhotoSet.h b/Engine/Plugins/Runtime/MeshModelingToolset/Source/ModelingComponents/Public/Scene/SceneCapturePhotoSet.h
--- a/Engine/Plugins/Runtime/MeshModelingToolset/Source/ModelingComponents/Public/Scene/SceneCapturePhotoSet.h
+++ b/Engine/Plugins/Runtime/MeshModelingToolset/Source/ModelingComponents/Public/Scene/SceneCapturePhotoSet.h
@@ -397,7 +397,7 @@ FVector4f FSceneCapturePhotoSet::ComputeSampleNearest(
 
 MODELINGCOMPONENTS_API
 TArray<FSpatialPhotoParams> ComputeStandardExteriorSpatialPhotoParameters(
-	UWorld* World,
+	UWorld* Unused,
 	const TArray<AActor*>& Actors,
 	const TArray<UActorComponent*>& Components,
 	FImageDimensions PhotoDimensions,
