Add convex hull simplification to the geometry script collision functions

diff --git a/Engine/Plugins/Experimental/GeometryScripting/Source/GeometryScriptingCore/Private/CollisionFunctions.cpp b/Engine/Plugins/Experimental/GeometryScripting/Source/GeometryScriptingCore/Private/CollisionFunctions.cpp
--- a/Engine/Plugins/Experimental/GeometryScripting/Source/GeometryScriptingCore/Private/CollisionFunctions.cpp
+++ b/Engine/Plugins/Experimental/GeometryScripting/Source/GeometryScriptingCore/Private/CollisionFunctions.cpp
@@ -4,11 +4,15 @@
 
 #include "Async/ParallelFor.h"
 #include "DynamicMesh/DynamicMesh3.h"
+#include "DynamicMesh/DynamicMeshAABBTree3.h"
 #include "DynamicMesh/MeshNormals.h"
 #include "Operations/MeshConvexHull.h"
 #include "Operations/MeshProjectionHull.h"
 #include "UDynamicMesh.h"
 #include "Components/DynamicMeshComponent.h"
+#include "Spatial/FastWinding.h"
+#include "ProjectionTargets.h"
+#include "MeshSimplification.h"
 
 #include "Selections/MeshConnectedComponents.h"
 #include "DynamicSubmesh3.h"
@@ -216,6 +220,29 @@ static void SetStaticMeshSimpleCollision(UStaticMesh* StaticMeshAsset, const FKA
 }
 
 
+// local helper to append a convex elem to a compact dynamic mesh, if it has more than a given number of tris
+static bool AppendConvexElemToCompactDynamicMesh(const FKConvexElem& Elem, FDynamicMesh3& Mesh, int32 MinTris = 0)
+{
+	checkSlow(Mesh.IsCompact());
+	if (Elem.IndexData.Num() <= MinTris * 3)
+	{
+		return false;
+	}
+
+	int32 StartV = Mesh.MaxVertexID();
+	for (FVector V : Elem.VertexData)
+	{
+		Mesh.AppendVertex(V);
+	}
+	for (int32 TriStart = 0; TriStart + 2 < Elem.IndexData.Num(); TriStart += 3)
+	{
+		// Note: We intentially reverse triangle winding here because FKConvexElem stores triangles with the opposite winding vs Dynamic Mesh
+		Mesh.AppendTriangle(StartV + Elem.IndexData[TriStart], StartV + Elem.IndexData[TriStart + 2], StartV + Elem.IndexData[TriStart + 1]);
+	}
+
+	return true;
+}
+
 
 }		// end namespace UELocal
 
@@ -480,6 +507,58 @@ void UGeometryScriptLibrary_CollisionFunctions::SetSimpleCollisionOfStaticMesh(
 	UELocal::SetStaticMeshSimpleCollision(StaticMesh, SimpleCollision.AggGeom, Options.bEmitTransaction);
 }
 
+void UGeometryScriptLibrary_CollisionFunctions::SimplifyConvexHulls(
+	FGeometryScriptSimpleCollision& SimpleCollision,
+	const FGeometryScriptConvexHullSimplificationOptions& SimplifyOptions,
+	bool& bHasSimplified,
+	UGeometryScriptDebug* Debug
+)
+{
+	bHasSimplified = false;
+	TArray<FKConvexElem>& ConvexElems = SimpleCollision.AggGeom.ConvexElems;
+	for (int32 ConvexIdx = 0; ConvexIdx < ConvexElems.Num(); ++ConvexIdx)
+	{
+		FKConvexElem& Elem = ConvexElems[ConvexIdx];
+		Elem.ComputeChaosConvexIndices(false); // make sure indices are computed
+
+		int32 TargetTriangleCount = FMath::Max(4, SimplifyOptions.MinTargetFaceCount);
+
+		// Convert hull to a dynamic mesh
+		FDynamicMesh3 Mesh;
+		if (!UELocal::AppendConvexElemToCompactDynamicMesh(Elem, Mesh, TargetTriangleCount))
+		{
+			continue;
+		}
+
+		int32 InitialTriangleCount = Mesh.TriangleCount();
+
+		// Run simplification
+		FVolPresMeshSimplification Simplifier(&Mesh);
+		Simplifier.CollapseMode = FVolPresMeshSimplification::ESimplificationCollapseModes::MinimalExistingVertexError;
+		Simplifier.GeometricErrorConstraint = UE::Geometry::FVolPresMeshSimplification::EGeometricErrorCriteria::PredictedPointToProjectionTarget;
+		Simplifier.GeometricErrorTolerance = SimplifyOptions.SimplificationDistanceThreshold;
+
+		FDynamicMesh3 ProjectionTargetMesh(Mesh);
+		FDynamicMeshAABBTree3 ProjectionTargetSpatial(&ProjectionTargetMesh, true);
+		FMeshProjectionTarget ProjTarget(&ProjectionTargetMesh, &ProjectionTargetSpatial);
+		Simplifier.SetProjectionTarget(&ProjTarget);
+		Simplifier.SimplifyToTriangleCount(TargetTriangleCount);
+
+		// Simplification didn't reduce triangle count, so skip updating the convex hull
+		if (Mesh.TriangleCount() == InitialTriangleCount)
+		{
+			continue;
+		}
+
+		Elem.VertexData.Reset(Mesh.VertexCount());
+		for (FVector3d V : Mesh.VerticesItr())
+		{
+			Elem.VertexData.Add(V);
+		}
+		Elem.UpdateElemBox();
+		bHasSimplified = true;
+	}
+}
 
 FGeometryScriptSimpleCollision UGeometryScriptLibrary_CollisionFunctions::MergeSimpleCollisionShapes(
 	const FGeometryScriptSimpleCollision& SimpleCollision,
diff --git a/Engine/Plugins/Experimental/GeometryScripting/Source/GeometryScriptingCore/Public/GeometryScript/CollisionFunctions.h b/Engine/Plugins/Experimental/GeometryScripting/Source/GeometryScriptingCore/Public/GeometryScript/CollisionFunctions.h
--- a/Engine/Plugins/Experimental/GeometryScripting/Source/GeometryScriptingCore/Public/GeometryScript/CollisionFunctions.h
+++ b/Engine/Plugins/Experimental/GeometryScripting/Source/GeometryScriptingCore/Public/GeometryScript/CollisionFunctions.h
@@ -127,6 +127,20 @@ public:
 	FGeometryScriptSimpleCollisionTriangulationOptions ShapeToHullTriangulation;
 };
 
+USTRUCT(BlueprintType)
+struct GEOMETRYSCRIPTINGCORE_API FGeometryScriptConvexHullSimplificationOptions
+{
+	GENERATED_BODY()
+public:
+
+	/** Simplified hull should stay within this distance of the initial convex hull. */
+	UPROPERTY(EditAnywhere, BlueprintReadWrite, Category = Options)
+	float SimplificationDistanceThreshold = 10.f;
+
+	/** The minimum number of faces to use for the convex hull. */
+	UPROPERTY(EditAnywhere, BlueprintReadWrite, Category = Options, meta = (ClampMin = 4))
+	int32 MinTargetFaceCount = 12;
+};
 
 UCLASS(meta = (ScriptName = "GeometryScript_Collision"))
 class GEOMETRYSCRIPTINGCORE_API UGeometryScriptLibrary_CollisionFunctions : public UBlueprintFunctionLibrary
@@ -224,6 +238,20 @@ public:
 		return SimpleCollision.AggGeom.GetElementCount();
 	}
 
+	/**
+	 * Simplify any convex hulls in the given simple collision representation. Updates the passed-in Simple Collision.
+	 *
+	 * @param SimpleCollision		The collision in which to attempt to simplify the convex hulls
+	 * @param ConvexSimplifyOptions	Options controlling how convex hulls are simplified
+	 * @param bHasSimplified		Indicates whether any convex hulls were modified
+	 */
+	UFUNCTION(BlueprintCallable, Category = "GeometryScript|Collision", meta = (AutoCreateRefTerm = "SimplifyOptions"))
+	static void SimplifyConvexHulls(
+		UPARAM(ref) FGeometryScriptSimpleCollision& SimpleCollision,
+		const FGeometryScriptConvexHullSimplificationOptions& SimplifyOptions,
+		bool& bHasSimplified,
+		UGeometryScriptDebug* Debug = nullptr
+	);
 
 	/**
 	 * Attempt to merge collision shapes to create a representation with fewer overall shapes.
