Add negative space protection to the geometry script collision merging algorithm

diff --git a/Engine/Plugins/Experimental/GeometryScripting/Source/GeometryScriptingCore/Private/CollisionFunctions.cpp b/Engine/Plugins/Experimental/GeometryScripting/Source/GeometryScriptingCore/Private/CollisionFunctions.cpp
--- a/Engine/Plugins/Experimental/GeometryScripting/Source/GeometryScriptingCore/Private/CollisionFunctions.cpp
+++ b/Engine/Plugins/Experimental/GeometryScripting/Source/GeometryScriptingCore/Private/CollisionFunctions.cpp
@@ -220,6 +220,34 @@ static void SetStaticMeshSimpleCollision(UStaticMesh* StaticMeshAsset, const FKA
 }
 
 
+// local helper to convert the blueprint-accessible enum to the geometrycore equivalent
+static UE::Geometry::FNegativeSpaceSampleSettings::ESampleMethod ConvertNegativeSpaceSampleMethodDataflowEnum(ENegativeSpaceSampleMethod SampleMethod)
+{
+	switch (SampleMethod)
+	{
+	case ENegativeSpaceSampleMethod::Uniform:
+		return UE::Geometry::FNegativeSpaceSampleSettings::ESampleMethod::Uniform;
+	case ENegativeSpaceSampleMethod::VoxelSearch:
+		return UE::Geometry::FNegativeSpaceSampleSettings::ESampleMethod::VoxelSearch;
+	}
+	return UE::Geometry::FNegativeSpaceSampleSettings::ESampleMethod::Uniform;
+}
+
+// local helper to convert negative space settings from the blueprint-accessible struct to the geometrycore equivalent
+static FNegativeSpaceSampleSettings ConvertNegativeSpaceOptions(const FComputeNegativeSpaceOptions& NegativeSpaceOptions)
+{
+	UE::Geometry::FNegativeSpaceSampleSettings NegativeSpaceSettings;
+	NegativeSpaceSettings.TargetNumSamples = NegativeSpaceOptions.TargetNumSamples;
+	NegativeSpaceSettings.MinRadius = NegativeSpaceOptions.MinRadius;
+	NegativeSpaceSettings.ReduceRadiusMargin = NegativeSpaceOptions.NegativeSpaceTolerance;
+	NegativeSpaceSettings.MinSpacing = NegativeSpaceOptions.MinSampleSpacing;
+	NegativeSpaceSettings.SampleMethod = ConvertNegativeSpaceSampleMethodDataflowEnum(NegativeSpaceOptions.SampleMethod);
+	NegativeSpaceSettings.bRequireSearchSampleCoverage = NegativeSpaceOptions.bRequireSearchSampleCoverage;
+	NegativeSpaceSettings.bReferenceMeshHasNegativeWinding = false;
+	NegativeSpaceSettings.Sanitize();
+	return NegativeSpaceSettings;
+}
+
 // local helper to append a convex elem to a compact dynamic mesh, if it has more than a given number of tris
 static bool AppendConvexElemToCompactDynamicMesh(const FKConvexElem& Elem, FDynamicMesh3& Mesh, int32 MinTris = 0)
 {
@@ -236,7 +264,6 @@ static bool AppendConvexElemToCompactDynamicMesh(const FKConvexElem& Elem, FDyna
 	}
 	for (int32 TriStart = 0; TriStart + 2 < Elem.IndexData.Num(); TriStart += 3)
 	{
-		// Note: We intentially reverse triangle winding here because FKConvexElem stores triangles with the opposite winding vs Dynamic Mesh
 		Mesh.AppendTriangle(StartV + Elem.IndexData[TriStart], StartV + Elem.IndexData[TriStart + 2], StartV + Elem.IndexData[TriStart + 1]);
 	}
 
@@ -580,7 +607,29 @@ FGeometryScriptSimpleCollision UGeometryScriptLibrary_CollisionFunctions::MergeS
 	TArray<int32> HullVertexCounts;
 	TArray<double> HullVolumes;
 	TArray<const FKShapeElem*> HullToShapeElem;
+	TUniquePtr<FDynamicMesh3> CollisionMesh;
+	if (MergeOptions.bComputeNegativeSpace)
+	{
+		CollisionMesh = MakeUnique<FDynamicMesh3>();
+	}
 
+	auto AppendGeneratorToCollisionMesh = [&CollisionMesh](const FMeshShapeGenerator& Generator)
+	{
+		if (!CollisionMesh)
+		{
+			return;
+		}
+		checkSlow(CollisionMesh->IsCompact());
+		int32 NewVertStart = CollisionMesh->MaxVertexID();
+		for (FVector3d V : Generator.Vertices)
+		{
+			CollisionMesh->AppendVertex(V);
+		}
+		for (FIndex3i T : Generator.Triangles)
+		{
+			CollisionMesh->AppendTriangle(T.A + NewVertStart, T.B + NewVertStart, T.C + NewVertStart);
+		}
+	};
 	auto TransformVertices = [](TArrayView<FVector3d> Vertices, const FTransform& Transform)
 	{
 		for (FVector3d& Vertex : Vertices)
@@ -610,9 +659,21 @@ FGeometryScriptSimpleCollision UGeometryScriptLibrary_CollisionFunctions::MergeS
 		OrientedBox.Extents = FVector(Box.X * .5, Box.Y * .5, Box.Z * .5);
 		OrientedBox.Frame.Origin = Box.Center;
 		OrientedBox.Frame.Rotation = (FQuaterniond)Box.Rotation;
-		TArray<FVector3d, TFixedAllocator<8>> BoxVertices;
-		OrientedBox.EnumerateCorners([&](FVector3d Corner) { BoxVertices.Add(Corner); });
-		AppendHullVertices(BoxVertices, Box.GetScaledVolume(FVector3d::One()), &Box);
+		if (CollisionMesh)
+		{
+			FGridBoxMeshGenerator BoxGenerator;
+			BoxGenerator.EdgeVertices = FIndex3i(1, 1, 1);
+			BoxGenerator.Box = OrientedBox;
+			BoxGenerator.Generate();
+			AppendHullVertices(BoxGenerator.Vertices, Box.GetScaledVolume(FVector3d::One()), &Box);
+			AppendGeneratorToCollisionMesh(BoxGenerator);
+		}
+		else
+		{
+			TArray<FVector3d, TFixedAllocator<8>> BoxVertices;
+			OrientedBox.EnumerateCorners([&](FVector3d Corner) { BoxVertices.Add(Corner); });
+			AppendHullVertices(BoxVertices, Box.GetScaledVolume(FVector3d::One()), &Box);
+		}
 	}
 	for (const FKSphereElem& Sphere : SimpleCollision.AggGeom.SphereElems)
 	{
@@ -621,10 +682,10 @@ FGeometryScriptSimpleCollision UGeometryScriptLibrary_CollisionFunctions::MergeS
 		SphereGenerator.Radius = FMath::Max(FMathf::ZeroTolerance, Sphere.Radius);
 		int32 StepsPerSide = FMath::Max(1, MergeOptions.ShapeToHullTriangulation.SphereStepsPerSide);
 		SphereGenerator.EdgeVertices = FIndex3i(StepsPerSide, StepsPerSide, StepsPerSide);
-		SphereGenerator.bPolygroupPerQuad = false;
 		SphereGenerator.Generate();
 		double Volume = GeneratorVolume(&SphereGenerator);
 		AppendHullVertices(SphereGenerator.Vertices, Volume, &Sphere);
+		AppendGeneratorToCollisionMesh(SphereGenerator);
 	}
 	for (const FKSphylElem& Capsule : SimpleCollision.AggGeom.SphylElems)
 	{
@@ -633,12 +694,12 @@ FGeometryScriptSimpleCollision UGeometryScriptLibrary_CollisionFunctions::MergeS
 		CapsuleGenerator.SegmentLength = Capsule.Length;
 		CapsuleGenerator.NumHemisphereArcSteps = FMath::Max(2, MergeOptions.ShapeToHullTriangulation.CapsuleHemisphereSteps);
 		CapsuleGenerator.NumCircleSteps = FMath::Max(3, MergeOptions.ShapeToHullTriangulation.CapsuleCircleSteps);
-		CapsuleGenerator.bPolygroupPerQuad = false;
 		CapsuleGenerator.Generate();
 		FTransform CapsuleTransform(Capsule.Rotation, Capsule.Center);
 		TransformVertices(CapsuleGenerator.Vertices, CapsuleTransform);
 		double Volume = GeneratorVolume(&CapsuleGenerator);
 		AppendHullVertices(CapsuleGenerator.Vertices, Volume, &Capsule);
+		AppendGeneratorToCollisionMesh(CapsuleGenerator);
 	}
 	for (const FKConvexElem& Convex : SimpleCollision.AggGeom.ConvexElems)
 	{
@@ -647,6 +708,10 @@ FGeometryScriptSimpleCollision UGeometryScriptLibrary_CollisionFunctions::MergeS
 		// Note: We take the negative volume because the hull triangles have opposite winding from ordinary meshes
 		double Volume = -TMeshQueries<TIndexMeshArrayAdapter<int32, double, FVector3d>>::GetVolumeArea(HullMeshAdapter).X;
 		AppendHullVertices(Convex.VertexData, Volume, &Convex);
+		if (CollisionMesh)
+		{
+			UELocal::AppendConvexElemToCompactDynamicMesh(Convex, *CollisionMesh);
+		}
 	}
 
 	const int32 InitialNumConvex = HullVertexCounts.Num();
@@ -679,7 +744,22 @@ FGeometryScriptSimpleCollision UGeometryScriptLibrary_CollisionFunctions::MergeS
 	Decomposition.InitializeFromHulls(HullVertexStarts.Num(),
 		[&HullVolumes](int32 HullIdx) { return HullVolumes[HullIdx]; }, [&HullVertexCounts](int32 HullIdx) { return HullVertexCounts[HullIdx]; },
 		[&HullVertexStarts, &HullVertices](int32 HullIdx, int32 VertIdx) { return HullVertices[HullVertexStarts[HullIdx] + VertIdx]; }, HullProximity);
-	Decomposition.MergeBest(MergeOptions.MaxShapeCount, MergeOptions.ErrorTolerance, 0.0, true, false, MergeOptions.MaxShapeCount, nullptr /*optional negative space*/, nullptr /*optional FTransform for negative space*/);
+	FSphereCovering NegativeSpace;
+	// Build the negative space of the collision shapes, if requested
+	if (MergeOptions.bComputeNegativeSpace)
+	{
+		FNegativeSpaceSampleSettings SampleSettings = UELocal::ConvertNegativeSpaceOptions(MergeOptions.ComputeNegativeSpaceOptions);
+		FDynamicMeshAABBTree3 CollisionAABBTree(CollisionMesh.Get(), true);
+		TFastWindingTree<FDynamicMesh3> CollisionFastWinding(&CollisionAABBTree, true);
+		NegativeSpace.AddNegativeSpace(CollisionFastWinding, SampleSettings);
+	}
+	// Add any precomputed negative space, if valid/non-empty
+	if (MergeOptions.PrecomputedNegativeSpace.Spheres.IsValid() && MergeOptions.PrecomputedNegativeSpace.Spheres->Num() > 0)
+	{
+		NegativeSpace.Append(*MergeOptions.PrecomputedNegativeSpace.Spheres);
+	}
+	FSphereCovering* UseNegativeSpace = NegativeSpace.Num() > 0 ? &NegativeSpace : nullptr;
+	Decomposition.MergeBest(MergeOptions.MaxShapeCount, MergeOptions.ErrorTolerance, 0.0, true, false, MergeOptions.MaxShapeCount, UseNegativeSpace, nullptr /*optional FTransform for negative space*/);
 
 	// Algorithm decided not to merge
 	if (Decomposition.NumHulls() == InitialNumConvex)
@@ -748,4 +828,24 @@ FGeometryScriptSimpleCollision UGeometryScriptLibrary_CollisionFunctions::MergeS
 	return ToRet;
 }
 
+FGeometryScriptSphereCovering UGeometryScriptLibrary_CollisionFunctions::ComputeNegativeSpace(
+	const FGeometryScriptDynamicMeshBVH& MeshBVH,
+	const FComputeNegativeSpaceOptions& NegativeSpaceOptions,
+	UGeometryScriptDebug* Debug
+)
+{
+	FGeometryScriptSphereCovering ToRet;
+
+	if (!MeshBVH.FWNTree || !MeshBVH.Spatial)
+	{
+		UE::Geometry::AppendError(Debug, EGeometryScriptErrorType::InvalidInputs, LOCTEXT("ComputeNegativeSpace_Null_BVH", "ComputeNegativeSpace: BVH must be initialized"));
+		return ToRet;
+	}
+	
+	FNegativeSpaceSampleSettings UseSettings = UELocal::ConvertNegativeSpaceOptions(NegativeSpaceOptions);
+	ToRet.Reset();
+	ToRet.Spheres->AddNegativeSpace(*MeshBVH.FWNTree, UseSettings);
+	return ToRet;
+}
+
 #undef LOCTEXT_NAMESPACE
diff --git a/Engine/Plugins/Experimental/GeometryScripting/Source/GeometryScriptingCore/Private/GeometryScriptTypes.cpp b/Engine/Plugins/Experimental/GeometryScripting/Source/GeometryScriptingCore/Private/GeometryScriptTypes.cpp
--- a/Engine/Plugins/Experimental/GeometryScripting/Source/GeometryScriptingCore/Private/GeometryScriptTypes.cpp
+++ b/Engine/Plugins/Experimental/GeometryScripting/Source/GeometryScriptingCore/Private/GeometryScriptTypes.cpp
@@ -4,6 +4,7 @@
 #include "GeometryBase.h"
 
 #include "Curve/GeneralPolygon2.h"
+#include "CompGeom/ConvexDecomposition3.h"
 
 
 FGeometryScriptDebugMessage UE::Geometry::MakeScriptError(EGeometryScriptErrorType ErrorTypeIn, const FText& MessageIn)
@@ -62,6 +63,16 @@ void UE::Geometry::AppendWarning(TArray<FGeometryScriptDebugMessage>* DebugMessa
 }
 
 
+void FGeometryScriptSphereCovering::Reset()
+{
+	if (Spheres.IsValid() == false)
+	{
+		Spheres = MakeShared<UE::Geometry::FSphereCovering>();
+	}
+	Spheres->Reset();
+}
+
+
 void FGeometryScriptGeneralPolygonList::Reset()
 {
 	if (!Polygons.IsValid())
diff --git a/Engine/Plugins/Experimental/GeometryScripting/Source/GeometryScriptingCore/Private/MeshPrimitiveFunctions.cpp b/Engine/Plugins/Experimental/GeometryScripting/Source/GeometryScriptingCore/Private/MeshPrimitiveFunctions.cpp
--- a/Engine/Plugins/Experimental/GeometryScripting/Source/GeometryScriptingCore/Private/MeshPrimitiveFunctions.cpp
+++ b/Engine/Plugins/Experimental/GeometryScripting/Source/GeometryScriptingCore/Private/MeshPrimitiveFunctions.cpp
@@ -21,6 +21,7 @@
 #include "ConstrainedDelaunay2.h"
 #include "Arrangement2d.h"
 #include "CompGeom/Delaunay2.h"
+#include "CompGeom/ConvexDecomposition3.h"
 
 #include UE_INLINE_GENERATED_CPP_BY_NAME(MeshPrimitiveFunctions)
 
@@ -1486,4 +1487,41 @@ UDynamicMesh* UGeometryScriptLibrary_MeshPrimitiveFunctions::AppendSimpleCollisi
 }
 
 
+UDynamicMesh* UGeometryScriptLibrary_MeshPrimitiveFunctions::AppendSphereCovering(
+	UDynamicMesh* TargetMesh,
+	FGeometryScriptPrimitiveOptions PrimitiveOptions,
+	FTransform Transform,
+	const FGeometryScriptSphereCovering& SphereCovering,
+	int32 StepsX,
+	int32 StepsY,
+	int32 StepsZ,
+	UGeometryScriptDebug* Debug
+)
+{
+	if (TargetMesh == nullptr)
+	{
+		UE::Geometry::AppendError(Debug, EGeometryScriptErrorType::InvalidInputs, LOCTEXT("PrimitiveFunctions_AppendSphereCovering_NullTargetMesh", "AppendSphereCovering: TargetMesh is Null"));
+		return TargetMesh;
+	}
+	if (!SphereCovering.Spheres.IsValid())
+	{
+		UE::Geometry::AppendError(Debug, EGeometryScriptErrorType::InvalidInputs, LOCTEXT("PrimitiveFunctions_AppendSphereCovering_NullSphereCovering", "AppendSphereCovering: Sphere Covering is Null"));
+		return TargetMesh;
+	}
+
+	for (int32 SphereIdx = 0; SphereIdx < SphereCovering.Spheres->Num(); ++SphereIdx)
+	{
+		FBoxSphereGenerator SphereGenerator;
+		SphereGenerator.Box.Frame.Origin = SphereCovering.Spheres->GetCenter(SphereIdx);
+		SphereGenerator.Radius = FMath::Max(FMathf::ZeroTolerance, SphereCovering.Spheres->GetRadius(SphereIdx));
+		SphereGenerator.EdgeVertices = FIndex3i(StepsX, StepsY, StepsZ);
+		SphereGenerator.bPolygroupPerQuad = (PrimitiveOptions.PolygroupMode == EGeometryScriptPrimitivePolygroupMode::PerQuad);
+		SphereGenerator.Generate();
+		AppendPrimitive(TargetMesh, &SphereGenerator, Transform, PrimitiveOptions);
+	}
+
+	return TargetMesh;
+}
+
+
 #undef LOCTEXT_NAMESPACE
\ No newline at end of file
diff --git a/Engine/Plugins/Experimental/GeometryScripting/Source/GeometryScriptingCore/Public/GeometryScript/CollisionFunctions.h b/Engine/Plugins/Experimental/GeometryScripting/Source/GeometryScriptingCore/Public/GeometryScript/CollisionFunctions.h
--- a/Engine/Plugins/Experimental/GeometryScripting/Source/GeometryScriptingCore/Public/GeometryScript/CollisionFunctions.h
+++ b/Engine/Plugins/Experimental/GeometryScripting/Source/GeometryScriptingCore/Public/GeometryScript/CollisionFunctions.h
@@ -102,6 +102,49 @@ public:
 	bool bEmitTransaction = true;
 };
 
+
+// Method to distribute sampling spheres, used by FComputeNegativeSpaceOptions
+UENUM(BlueprintType)
+enum class ENegativeSpaceSampleMethod : uint8
+{
+	// Place sample spheres in a uniform grid pattern
+	Uniform,
+	// Use voxel-based subtraction and offsetting methods to specifically target concavities
+	VoxelSearch
+};
+
+// Options controlling how to sample the negative space of shapes, e.g. to define a region that must be avoided when merging collision shapes
+USTRUCT(BlueprintType)
+struct GEOMETRYSCRIPTINGCORE_API FComputeNegativeSpaceOptions
+{
+	GENERATED_BODY()
+public:
+
+	/** Method to use to find and sample negative space */
+	UPROPERTY(EditAnywhere, BlueprintReadWrite, Category = NegativeSpace)
+	ENegativeSpaceSampleMethod SampleMethod = ENegativeSpaceSampleMethod::Uniform;
+
+	/** Whether to require that all candidate locations identified by Voxel Search are covered by negative space samples, up to the specified Min Sample Spacing. Only applies to Voxel Search. */
+	UPROPERTY(EditAnywhere, BlueprintReadWrite, Category = NegativeSpace, meta = (EditCondition = "SampleMethod == EGeometryScriptNegativeSpaceSampleMethod::VoxelSearch", EditConditionHides))
+	bool bRequireSearchSampleCoverage = false;
+
+	/** Approximate number of spheres to consider when covering negative space */
+	UPROPERTY(EditAnywhere, BlueprintReadWrite, Category = NegativeSpace, meta = (ClampMin = 1))
+	int32 TargetNumSamples = 50;
+
+	/** Minimum desired spacing between sphere centers; if > 0, will attempt not to place sphere centers closer than this */
+	UPROPERTY(EditAnywhere, BlueprintReadWrite, Category = NegativeSpace, meta = (ClampMin = 0, Units = cm))
+	double MinSampleSpacing = 1.0;
+
+	/** Amount of space to leave between convex hulls and protected negative space */
+	UPROPERTY(EditAnywhere, BlueprintReadWrite, Category = NegativeSpace, meta = (ClampMin = .01, UIMin = .1, Units = cm))
+	double NegativeSpaceTolerance = 2.0;
+
+	/** Spheres smaller than this are not included in the negative space */
+	UPROPERTY(EditAnywhere, BlueprintReadWrite, Category = NegativeSpace, meta = (ClampMin = 0, Units = cm))
+	double MinRadius = 10.0;
+};
+
 // Options controlling how collision shapes can be merged together
 USTRUCT(BlueprintType)
 struct GEOMETRYSCRIPTINGCORE_API FGeometryScriptMergeSimpleCollisionOptions
@@ -122,6 +165,18 @@ public:
 	UPROPERTY(EditAnywhere, BlueprintReadWrite, Category = Options, meta = (UIMin = "0", UIMax = "100.", Units = cm))
 	double ErrorTolerance = 0.0;
 
+	// Negative space that must be preserved during merging
+	UPROPERTY(EditAnywhere, BlueprintReadWrite, Category = NegativeSpace)
+	FGeometryScriptSphereCovering PrecomputedNegativeSpace;
+
+	// Whether to compute a new sphere covering representing the negative space of the input shapes
+	UPROPERTY(EditAnywhere, BlueprintReadWrite, Category = NegativeSpace)
+	bool bComputeNegativeSpace = false;
+
+	// Options controlling how the negative space is computed, if ComputeNegativeSpace is true
+	UPROPERTY(EditAnywhere, BlueprintReadWrite, Category = NegativeSpace)
+	FComputeNegativeSpaceOptions ComputeNegativeSpaceOptions;
+
 	// Controls for how smooth shapes can be triangulated when/if converted to a convex hull for a merge
 	UPROPERTY(EditAnywhere, BlueprintReadWrite, Category = ConvexHulls)
 	FGeometryScriptSimpleCollisionTriangulationOptions ShapeToHullTriangulation;
@@ -261,7 +316,7 @@ public:
 	 * @param bHasMerged			Indicates whether any shapes have been merged
 	 * @return						Simple Collision with collision shapes merged, as allowed by settings
 	 */
-	UFUNCTION(BlueprintCallable, Category = "GeometryScript|Collision")
+	UFUNCTION(BlueprintCallable, Category = "GeometryScript|Collision", meta = (AutoCreateRefTerm = "MergeOptions"))
 	static UPARAM(DisplayName = "Merged Simple Collision") FGeometryScriptSimpleCollision MergeSimpleCollisionShapes(
 		const FGeometryScriptSimpleCollision& SimpleCollision,
 		const FGeometryScriptMergeSimpleCollisionOptions& MergeOptions,
@@ -269,5 +324,18 @@ public:
 		UGeometryScriptDebug* Debug = nullptr
 	);
 
+	/**
+	 * Compute the negative space of an input mesh surface that should be protected when merging simple collision shapes
+	 * 
+	 * @param MeshBVH				A Dynamic Mesh BVH structure of the surface for which we will compute negative space
+	 * @param NegativeSpaceOptions	Options controlling how the negative space is generated
+	 * @return						A set of spheres that cover the negative space of the input shape
+	 */
+	UFUNCTION(BlueprintCallable, Category = "GeometryScript|Collision", meta = (AutoCreateRefTerm = "NegativeSpaceOptions"))
+	static UPARAM(DisplayName = "Negative Space") FGeometryScriptSphereCovering ComputeNegativeSpace(
+		const FGeometryScriptDynamicMeshBVH& MeshBVH,
+		const FComputeNegativeSpaceOptions& NegativeSpaceOptions,
+		UGeometryScriptDebug* Debug = nullptr
+	);
 
 };
\ No newline at end of file
diff --git a/Engine/Plugins/Experimental/GeometryScripting/Source/GeometryScriptingCore/Public/GeometryScript/GeometryScriptTypes.h b/Engine/Plugins/Experimental/GeometryScripting/Source/GeometryScriptingCore/Public/GeometryScript/GeometryScriptTypes.h
--- a/Engine/Plugins/Experimental/GeometryScripting/Source/GeometryScriptingCore/Public/GeometryScript/GeometryScriptTypes.h
+++ b/Engine/Plugins/Experimental/GeometryScripting/Source/GeometryScriptingCore/Public/GeometryScript/GeometryScriptTypes.h
@@ -13,6 +13,7 @@ PREDECLARE_GEOMETRY(template<typename MeshType> class TFastWindingTree);
 PREDECLARE_GEOMETRY(typedef TMeshAABBTree3<FDynamicMesh3> FDynamicMeshAABBTree3);
 PREDECLARE_GEOMETRY(template<typename RealType> class TGeneralPolygon2);
 PREDECLARE_GEOMETRY(typedef TGeneralPolygon2<double> FGeneralPolygon2d);
+PREDECLARE_GEOMETRY(class FSphereCovering);
 
 
 UENUM(BlueprintType)
@@ -95,10 +96,7 @@ public:
 // Collision Shapes
 //
 
-/**
- * ~This is a wrapper to let Blueprint functions pass around the FKAggregateGeom simple collision primitives
- * Holds simple shapes that can be used for collision
- */
+// Holds simple shapes that can be used for collision
 USTRUCT(BlueprintType)
 struct GEOMETRYSCRIPTINGCORE_API FGeometryScriptSimpleCollision
 {
@@ -110,6 +108,38 @@ public:
 	
 };
 
+// A set of spheres used to represent a volume
+USTRUCT(BlueprintType, meta = (DisplayName = "Sphere Covering"))
+struct GEOMETRYSCRIPTINGCORE_API FGeometryScriptSphereCovering
+{
+	GENERATED_BODY()
+public:
+	TSharedPtr<UE::Geometry::FSphereCovering> Spheres;
+
+	void Reset();
+
+	// Required by TStructOpsTypeTraits interface
+	bool operator==(const FGeometryScriptSphereCovering& Other) const
+	{
+		return Spheres.Get() == Other.Spheres.Get();
+	}
+	bool operator!=(const FGeometryScriptSphereCovering& Other) const
+	{
+		return Spheres.Get() != Other.Spheres.Get();
+	}
+};
+
+template<>
+struct TStructOpsTypeTraits<FGeometryScriptSphereCovering> : public TStructOpsTypeTraitsBase2<FGeometryScriptSphereCovering>
+{
+	enum
+	{
+		WithIdenticalViaEquality = true,
+	};
+};
+
+
+
 // Settings to control the triangulation of simple collision primitives -- used for conversion to mesh or convex hull geometry
 USTRUCT(BlueprintType)
 struct GEOMETRYSCRIPTINGCORE_API FGeometryScriptSimpleCollisionTriangulationOptions
diff --git a/Engine/Plugins/Experimental/GeometryScripting/Source/GeometryScriptingCore/Public/GeometryScript/MeshPrimitiveFunctions.h b/Engine/Plugins/Experimental/GeometryScripting/Source/GeometryScriptingCore/Public/GeometryScript/MeshPrimitiveFunctions.h
--- a/Engine/Plugins/Experimental/GeometryScripting/Source/GeometryScriptingCore/Public/GeometryScript/MeshPrimitiveFunctions.h
+++ b/Engine/Plugins/Experimental/GeometryScripting/Source/GeometryScriptingCore/Public/GeometryScript/MeshPrimitiveFunctions.h
@@ -573,7 +573,10 @@ public:
 		bool& bTriangulationError,
 		UGeometryScriptDebug* Debug = nullptr);
 
-
+	
+	/**
+	 * Appends Simple Collision shapes to the Target Mesh, triangulated as specified by Triangulation Options
+	 */
 	UFUNCTION(BlueprintCallable, Category = "GeometryScript|Primitives")
 	static UPARAM(DisplayName = "Target Mesh") UDynamicMesh*
 	AppendSimpleCollisionShapes(
@@ -585,6 +588,21 @@ public:
 		UGeometryScriptDebug* Debug = nullptr
 	);
 
+	/**
+	 * Appends the spheres in the Sphere Covering to the Target Mesh
+	 */
+	UFUNCTION(BlueprintCallable, Category = "GeometryScript|Primitives")
+	static UPARAM(DisplayName = "Target Mesh") UDynamicMesh*
+	AppendSphereCovering(
+		UDynamicMesh* TargetMesh,
+		FGeometryScriptPrimitiveOptions PrimitiveOptions,
+		FTransform Transform,
+		const FGeometryScriptSphereCovering& SphereCovering,
+		int32 StepsX = 6,
+		int32 StepsY = 6,
+		int32 StepsZ = 6,
+		UGeometryScriptDebug* Debug = nullptr
+	);
 
 
 	//---------------------------------------------
diff --git a/Engine/Source/Runtime/GeometryCore/Private/CompGeom/ConvexDecomposition3.cpp b/Engine/Source/Runtime/GeometryCore/Private/CompGeom/ConvexDecomposition3.cpp
--- a/Engine/Source/Runtime/GeometryCore/Private/CompGeom/ConvexDecomposition3.cpp
+++ b/Engine/Source/Runtime/GeometryCore/Private/CompGeom/ConvexDecomposition3.cpp
@@ -31,6 +31,8 @@ bool FSphereCovering::AddNegativeSpace(const TFastWindingTree<FDynamicMesh3>& Sp
 
 	FAxisAlignedBox3d Bounds = Spatial.GetTree()->GetBoundingBox();
 
+	double WindingSign = SampleSettings.bReferenceMeshHasNegativeWinding ? -1 : 1;
+
 	if (SampleSettings.SampleMethod == FNegativeSpaceSampleSettings::ESampleMethod::Uniform)
 	{
 		// Expand the sampling region by ~ the radius of interest, to make sure there is room to fit negative space spheres on concavities near the borders
@@ -62,8 +64,8 @@ bool FSphereCovering::AddNegativeSpace(const TFastWindingTree<FDynamicMesh3>& Sp
 				{
 					double ZPos = Bounds.Min.Z + StepSizes.Z * (.5 + Z);
 					FVector3d Pos(XPos, YPos, ZPos);
-					double Winding = Spatial.FastWindingNumber(Pos);
-					if (Winding < -.5)
+					double Winding = Spatial.FastWindingNumber(Pos) * WindingSign;
+					if (Winding > .5)
 					{
 						continue;
 					}
@@ -148,10 +150,10 @@ bool FSphereCovering::AddNegativeSpace(const TFastWindingTree<FDynamicMesh3>& Sp
 		MarchingCubes.RootMode = ERootfindingModes::Bisection;
 		MarchingCubes.RootModeSteps = 3;
 		MarchingCubes.IsoValue = 0;
-		MarchingCubes.Implicit = [&HullWinding, &Spatial](const FVector3d& Pt) -> double
+		MarchingCubes.Implicit = [&HullWinding, &Spatial, WindingSign](const FVector3d& Pt) -> double
 		{
-			// Volume is anything inside the hull and outside the input surface; note the input surface has negative winding
-			return (HullWinding.FastWindingNumber(Pt) > .5) && !(Spatial.FastWindingNumber(Pt) < -.5) ? 1.0 : -1.0;
+			// Volume is anything inside the hull and outside the input surface
+			return (HullWinding.FastWindingNumber(Pt) > .5) && (Spatial.FastWindingNumber(Pt) * WindingSign <= .5) ? 1.0 : -1.0;
 		};
 		TArray<FVector3d> MCSeeds;
 		MCSeeds.Reserve(Mesh->VertexCount());
@@ -178,10 +180,10 @@ bool FSphereCovering::AddNegativeSpace(const TFastWindingTree<FDynamicMesh3>& Sp
 		// Make sure the mesh is compact to simplify downsampling below
 		MorphologyMesh.CompactInPlace();
 
-		auto AddSample = [this, &Mesh, &Spatial, &SampleSettings, &bAddedPoints](FVector3d Pos)
+		auto AddSample = [this, &Mesh, &Spatial, &SampleSettings, &bAddedPoints, WindingSign](FVector3d Pos)
 		{
-			double Winding = Spatial.FastWindingNumber(Pos);
-			if (Winding < -.5)
+			double Winding = Spatial.FastWindingNumber(Pos) * WindingSign;
+			if (Winding > .5)
 			{
 				return;
 			}
@@ -205,7 +207,7 @@ bool FSphereCovering::AddNegativeSpace(const TFastWindingTree<FDynamicMesh3>& Sp
 				
 				// Move away and re-test the sample
 				Pos += Away * ((SampleSettings.MinRadius - R) * 1.1);
-				if (Spatial.FastWindingNumber(Pos) >= -.5)
+				if (Spatial.FastWindingNumber(Pos) * WindingSign <= .5)
 				{
 					NearTID = Spatial.GetTree()->FindNearestTriangle(Pos, NearDistSq);
 					R = FMath::Sqrt(NearDistSq) - SampleSettings.ReduceRadiusMargin;
diff --git a/Engine/Source/Runtime/GeometryCore/Public/CompGeom/ConvexDecomposition3.h b/Engine/Source/Runtime/GeometryCore/Public/CompGeom/ConvexDecomposition3.h
--- a/Engine/Source/Runtime/GeometryCore/Public/CompGeom/ConvexDecomposition3.h
+++ b/Engine/Source/Runtime/GeometryCore/Public/CompGeom/ConvexDecomposition3.h
@@ -68,6 +68,9 @@ struct FNegativeSpaceSampleSettings
 	// Note: This takes priority over TargetNumSamples if the TargetNumSamples did not achieve the required coverage.
 	bool bRequireSearchSampleCoverage = false;
 
+	// Whether the reference mesh used has flipped orientation, so winding < -.5 is 'inside'
+	bool bReferenceMeshHasNegativeWinding = true;
+
 	// Make sure the settings values are in valid ranges
 	void Sanitize()
 	{
