Add option to set static mesh collision via geometry script without marking the static mesh as having customized collision -- can be useful for example in cases where the geometry script is run as part of an import pipeline

diff --git a/Engine/Plugins/Experimental/GeometryScripting/Source/GeometryScriptingCore/Private/CollisionFunctions.cpp b/Engine/Plugins/Experimental/GeometryScripting/Source/GeometryScriptingCore/Private/CollisionFunctions.cpp
--- a/Engine/Plugins/Experimental/GeometryScripting/Source/GeometryScriptingCore/Private/CollisionFunctions.cpp
+++ b/Engine/Plugins/Experimental/GeometryScripting/Source/GeometryScriptingCore/Private/CollisionFunctions.cpp
@@ -150,7 +150,7 @@ void ComputeCollisionFromMesh(
 
 
 
-static void SetStaticMeshSimpleCollision(UStaticMesh* StaticMeshAsset, const FKAggregateGeom& NewSimpleCollision, bool bEmitTransaction)
+static void SetStaticMeshSimpleCollision(UStaticMesh* StaticMeshAsset, const FKAggregateGeom& NewSimpleCollision, bool bEmitTransaction, bool bMarkCollisionAsCustomized = true)
 {
 #if WITH_EDITOR
 	if (bEmitTransaction && GEditor)
@@ -206,7 +206,10 @@ static void SetStaticMeshSimpleCollision(UStaticMesh* StaticMeshAsset, const FKA
 
 #if WITH_EDITORONLY_DATA
 		// mark the static mesh as having customized collision so it is not regenerated on reimport
-		StaticMeshAsset->bCustomizedCollision = true;
+		if (bMarkCollisionAsCustomized)
+		{
+			StaticMeshAsset->bCustomizedCollision = true;
+		}
 #endif // WITH_EDITORONLY_DATA
 	}
 
@@ -278,6 +281,7 @@ UDynamicMesh* UGeometryScriptLibrary_CollisionFunctions::SetStaticMeshCollisionF
 	UDynamicMesh* FromDynamicMesh,
 	UStaticMesh* ToStaticMeshAsset,
 	FGeometryScriptCollisionFromMeshOptions Options,
+	FGeometryScriptSetStaticMeshCollisionOptions StaticMeshCollisionOptions,
 	UGeometryScriptDebug* Debug)
 {
 	if (FromDynamicMesh == nullptr)
@@ -297,7 +301,7 @@ UDynamicMesh* UGeometryScriptLibrary_CollisionFunctions::SetStaticMeshCollisionF
 		UELocal::ComputeCollisionFromMesh(ReadMesh, NewCollision, Options);
 	});
 
-	UELocal::SetStaticMeshSimpleCollision(ToStaticMeshAsset, NewCollision, Options.bEmitTransaction);
+	UELocal::SetStaticMeshSimpleCollision(ToStaticMeshAsset, NewCollision, Options.bEmitTransaction, StaticMeshCollisionOptions.bMarkAsCustomized);
 
 	return FromDynamicMesh;
 }
@@ -309,6 +313,7 @@ void UGeometryScriptLibrary_CollisionFunctions::SetStaticMeshCollisionFromCompon
 	UStaticMesh* UpdateStaticMeshAsset, 
 	UPrimitiveComponent* SourceComponent,
 	FGeometryScriptSetSimpleCollisionOptions Options,
+	FGeometryScriptSetStaticMeshCollisionOptions StaticMeshCollisionOptions,
 	UGeometryScriptDebug* Debug)
 {
 	if (UpdateStaticMeshAsset == nullptr)
@@ -329,10 +334,18 @@ void UGeometryScriptLibrary_CollisionFunctions::SetStaticMeshCollisionFromCompon
 		return;
 	}
 
-	UELocal::SetStaticMeshSimpleCollision(UpdateStaticMeshAsset, BodySetup->AggGeom, Options.bEmitTransaction);
+	UELocal::SetStaticMeshSimpleCollision(UpdateStaticMeshAsset, BodySetup->AggGeom, Options.bEmitTransaction, StaticMeshCollisionOptions.bMarkAsCustomized);
 }
 
 
+bool UGeometryScriptLibrary_CollisionFunctions::StaticMeshHasCustomizedCollision(UStaticMesh* StaticMeshAsset)
+{
+#if WITH_EDITORONLY_DATA
+	return StaticMeshAsset->bCustomizedCollision;
+#else
+	return false;
+#endif
+}
 
 
 
@@ -529,9 +542,10 @@ void UGeometryScriptLibrary_CollisionFunctions::SetSimpleCollisionOfStaticMesh(
 	const FGeometryScriptSimpleCollision& SimpleCollision,
 	UStaticMesh* StaticMesh, 
 	FGeometryScriptSetSimpleCollisionOptions Options,
+	FGeometryScriptSetStaticMeshCollisionOptions StaticMeshCollisionOptions,
 	UGeometryScriptDebug* Debug)
 {
-	UELocal::SetStaticMeshSimpleCollision(StaticMesh, SimpleCollision.AggGeom, Options.bEmitTransaction);
+	UELocal::SetStaticMeshSimpleCollision(StaticMesh, SimpleCollision.AggGeom, Options.bEmitTransaction, StaticMeshCollisionOptions.bMarkAsCustomized);
 }
 
 void UGeometryScriptLibrary_CollisionFunctions::SimplifyConvexHulls(
diff --git a/Engine/Plugins/Experimental/GeometryScripting/Source/GeometryScriptingCore/Public/GeometryScript/CollisionFunctions.h b/Engine/Plugins/Experimental/GeometryScripting/Source/GeometryScriptingCore/Public/GeometryScript/CollisionFunctions.h
--- a/Engine/Plugins/Experimental/GeometryScripting/Source/GeometryScriptingCore/Public/GeometryScript/CollisionFunctions.h
+++ b/Engine/Plugins/Experimental/GeometryScripting/Source/GeometryScriptingCore/Public/GeometryScript/CollisionFunctions.h
@@ -103,6 +103,18 @@ public:
 };
 
 
+USTRUCT(BlueprintType)
+struct GEOMETRYSCRIPTINGCORE_API FGeometryScriptSetStaticMeshCollisionOptions
+{
+	GENERATED_BODY()
+public:
+	// Whether to mark the static mesh collision as customized when it is set, so that it will not be overwritten on next import.
+	// If false, Static Mesh collision will not be un-marked as Customized; its state will just be left unchanged.
+	UPROPERTY(EditAnywhere, BlueprintReadWrite, Category = Options)
+	bool bMarkAsCustomized = true;
+};
+
+
 // Method to distribute sampling spheres, used by FComputeNegativeSpaceOptions
 UENUM(BlueprintType)
 enum class ENegativeSpaceSampleMethod : uint8
@@ -212,6 +224,7 @@ public:
 		UDynamicMesh* FromDynamicMesh, 
 		UStaticMesh* ToStaticMeshAsset, 
 		FGeometryScriptCollisionFromMeshOptions Options,
+		FGeometryScriptSetStaticMeshCollisionOptions StaticMeshCollisionOptions = FGeometryScriptSetStaticMeshCollisionOptions(),
 		UGeometryScriptDebug* Debug = nullptr);
 
 	/**
@@ -223,8 +236,15 @@ public:
 		UStaticMesh* StaticMeshAsset, 
 		UPrimitiveComponent* SourceComponent,
 		FGeometryScriptSetSimpleCollisionOptions Options = FGeometryScriptSetSimpleCollisionOptions(),
+		FGeometryScriptSetStaticMeshCollisionOptions StaticMeshCollisionOptions = FGeometryScriptSetStaticMeshCollisionOptions(),
 		UGeometryScriptDebug* Debug = nullptr);
 
+	/*
+	 * @returns true if the static mesh has customized collision. If no editor data is available, returns false.
+	 */
+	UFUNCTION(BlueprintCallable, Category = "GeometryScript|Collision")
+	static UPARAM(DisplayName = "IsCustomized") bool StaticMeshHasCustomizedCollision(UStaticMesh* StaticMeshAsset);
+
 	/** 
 	* Generate Simple Collision shapes for a Dynamic Mesh Component based on the input Dynamic Mesh. 
 	*/
@@ -282,6 +302,7 @@ public:
 		const FGeometryScriptSimpleCollision& SimpleCollision,
 		UStaticMesh* StaticMesh, 
 		FGeometryScriptSetSimpleCollisionOptions Options, 
+		FGeometryScriptSetStaticMeshCollisionOptions StaticMeshCollisionOptions = FGeometryScriptSetStaticMeshCollisionOptions(),
 		UGeometryScriptDebug* Debug = nullptr);
 
 	/*
