Add missing equality struct for geometry script polypath list type

diff --git a/Engine/Plugins/Experimental/GeometryScripting/Source/GeometryScriptingCore/Public/GeometryScript/GeometryScriptTypes.h b/Engine/Plugins/Experimental/GeometryScripting/Source/GeometryScriptingCore/Public/GeometryScript/GeometryScriptTypes.h
--- a/Engine/Plugins/Experimental/GeometryScripting/Source/GeometryScriptingCore/Public/GeometryScript/GeometryScriptTypes.h
+++ b/Engine/Plugins/Experimental/GeometryScripting/Source/GeometryScriptingCore/Public/GeometryScript/GeometryScriptTypes.h
@@ -586,6 +586,17 @@ public:
 	}
 };
 
+template<>
+struct TStructOpsTypeTraits<FGeometryScriptSimplePolygon> : public TStructOpsTypeTraitsBase2<FGeometryScriptSimplePolygon>
+{
+	enum
+	{
+		WithIdenticalViaEquality = true,
+	};
+};
+
+
+
 // A list of general polygons, which may have holes.
 USTRUCT(BlueprintType, meta = (DisplayName = "PolygonList"))
 struct GEOMETRYSCRIPTINGCORE_API FGeometryScriptGeneralPolygonList
