fix incorrect translation for capsule simple collision when meshing in geometry script

diff --git a/Engine/Plugins/Experimental/GeometryScripting/Source/GeometryScriptingCore/Private/CollisionFunctions.cpp b/Engine/Plugins/Experimental/GeometryScripting/Source/GeometryScriptingCore/Private/CollisionFunctions.cpp
--- a/Engine/Plugins/Experimental/GeometryScripting/Source/GeometryScriptingCore/Private/CollisionFunctions.cpp
+++ b/Engine/Plugins/Experimental/GeometryScripting/Source/GeometryScriptingCore/Private/CollisionFunctions.cpp
@@ -709,7 +709,7 @@ FGeometryScriptSimpleCollision UGeometryScriptLibrary_CollisionFunctions::MergeS
 		CapsuleGenerator.NumHemisphereArcSteps = FMath::Max(2, MergeOptions.ShapeToHullTriangulation.CapsuleHemisphereSteps);
 		CapsuleGenerator.NumCircleSteps = FMath::Max(3, MergeOptions.ShapeToHullTriangulation.CapsuleCircleSteps);
 		CapsuleGenerator.Generate();
-		FTransform CapsuleTransform(Capsule.Rotation, Capsule.Center);
+		FTransform CapsuleTransform(Capsule.Rotation, Capsule.Center + Capsule.Rotation.RotateVector(FVector(0, 0, -Capsule.Length * .5)));
 		TransformVertices(CapsuleGenerator.Vertices, CapsuleTransform);
 		double Volume = GeneratorVolume(&CapsuleGenerator);
 		AppendHullVertices(CapsuleGenerator.Vertices, Volume, &Capsule);
diff --git a/Engine/Plugins/Experimental/GeometryScripting/Source/GeometryScriptingCore/Private/MeshPrimitiveFunctions.cpp b/Engine/Plugins/Experimental/GeometryScripting/Source/GeometryScriptingCore/Private/MeshPrimitiveFunctions.cpp
--- a/Engine/Plugins/Experimental/GeometryScripting/Source/GeometryScriptingCore/Private/MeshPrimitiveFunctions.cpp
+++ b/Engine/Plugins/Experimental/GeometryScripting/Source/GeometryScriptingCore/Private/MeshPrimitiveFunctions.cpp
@@ -1459,8 +1459,9 @@ UDynamicMesh* UGeometryScriptLibrary_MeshPrimitiveFunctions::AppendSimpleCollisi
 		CapsuleGenerator.NumCircleSteps = FMath::Max(3, TriangulationOptions.CapsuleCircleSteps);
 		CapsuleGenerator.bPolygroupPerQuad = (PrimitiveOptions.PolygroupMode == EGeometryScriptPrimitivePolygroupMode::PerQuad);
 		CapsuleGenerator.Generate();
-
-		AppendPrimitive(TargetMesh, &CapsuleGenerator, Transform, PrimitiveOptions, Capsule.Center, FQuaterniond(Capsule.Rotation));
+		
+		FQuaterniond Rotation(Capsule.Rotation);
+		AppendPrimitive(TargetMesh, &CapsuleGenerator, Transform, PrimitiveOptions, Capsule.Center + Rotation * FVector(0, 0, -Capsule.Length*.5), Rotation);
 	}
 	for (const FKConvexElem& Convex : SimpleCollision.AggGeom.ConvexElems)
 	{
