Add geometry script function to transform simple collision shapes
When transform cannot be applied exactly (e.g., non-uniform scale to a sphere), the function does a best-effort transform, optionally warns the user, and reports failure via a Success bool output

diff --git a/Engine/Plugins/Experimental/GeometryScripting/Source/GeometryScriptingCore/Private/CollisionFunctions.cpp b/Engine/Plugins/Experimental/GeometryScripting/Source/GeometryScriptingCore/Private/CollisionFunctions.cpp
--- a/Engine/Plugins/Experimental/GeometryScripting/Source/GeometryScriptingCore/Private/CollisionFunctions.cpp
+++ b/Engine/Plugins/Experimental/GeometryScripting/Source/GeometryScriptingCore/Private/CollisionFunctions.cpp
@@ -548,6 +548,249 @@ void UGeometryScriptLibrary_CollisionFunctions::SetSimpleCollisionOfStaticMesh(
 	UELocal::SetStaticMeshSimpleCollision(StaticMesh, SimpleCollision.AggGeom, Options.bEmitTransaction, StaticMeshCollisionOptions.bMarkAsCustomized);
 }
 
+FGeometryScriptSimpleCollision UGeometryScriptLibrary_CollisionFunctions::TransformSimpleCollisionShapes(
+	const FGeometryScriptSimpleCollision& SimpleCollision,
+	FTransform Transform,
+	const FGeometryScriptTransformCollisionOptions& TransformOptions,
+	bool& bSuccess,
+	UGeometryScriptDebug* Debug)
+{
+	bSuccess = true;
+
+	FGeometryScriptSimpleCollision TransformedCollision;
+	FVector Scale = Transform.GetScale3D();
+	bool bUniformScale = Scale.AllComponentsEqual();
+	bool bNoRotation = Transform.GetRotation().IsIdentity();
+	double ApproxUniformScale = Scale.X;
+	if (!bUniformScale)
+	{
+		ApproxUniformScale = FMathd::Cbrt(Scale.X * Scale.Y * Scale.Z);
+	}
+
+	int32 NumSpheres = SimpleCollision.AggGeom.SphereElems.Num();
+	if (NumSpheres > 0)
+	{
+		if (!bUniformScale)
+		{
+			if (TransformOptions.bWarnOnInvalidTransforms)
+			{
+				UE::Geometry::AppendWarning(Debug, EGeometryScriptErrorType::InvalidInputs, LOCTEXT("TransformSimpleCollisionShapes Non-Uniform Scale on Sphere", "TransformSimpleCollisionShapes: Cannot apply Non-Uniform Scale to Spheres"));
+			}
+			bSuccess = false;
+		}
+		TransformedCollision.AggGeom.SphereElems.Reserve(NumSpheres);
+		for (FKSphereElem Sphere : SimpleCollision.AggGeom.SphereElems)
+		{
+			if (TransformOptions.bCenterTransformPivotPerShape)
+			{
+				// Apply Transform w/ pivot at the sphere center; only the translation part will have an effect
+				Sphere.Center += Transform.GetTranslation();
+			}
+			else
+			{
+				Sphere.Center = Transform.TransformPosition(Sphere.Center);
+			}
+			Sphere.Radius = FMath::Abs(ApproxUniformScale * Sphere.Radius);
+			TransformedCollision.AggGeom.SphereElems.Add(Sphere);
+		}
+	}
+
+	int32 NumBoxes = SimpleCollision.AggGeom.BoxElems.Num();
+	if (NumBoxes > 0)
+	{
+		if (!bUniformScale)
+		{
+			bool bHasRotatedBox = false;
+			for (const FKBoxElem& Box : SimpleCollision.AggGeom.BoxElems)
+			{
+				if (!Box.Rotation.IsNearlyZero())
+				{
+					bHasRotatedBox = true;
+					break;
+				}
+			}
+			if (bHasRotatedBox)
+			{
+				if (TransformOptions.bWarnOnInvalidTransforms)
+				{
+					UE::Geometry::AppendWarning(Debug, EGeometryScriptErrorType::InvalidInputs, LOCTEXT("TransformSimpleCollisionShapes Non-Uniform Scale on Rotated Box", "TransformSimpleCollisionShapes: Cannot apply Non-Uniform Scale to Rotated Box"));
+				}
+				bSuccess = false;
+			}
+		}
+		TransformedCollision.AggGeom.BoxElems.Reserve(NumBoxes);
+		for (FKBoxElem Box : SimpleCollision.AggGeom.BoxElems)
+		{
+			if (TransformOptions.bCenterTransformPivotPerShape)
+			{
+				Box.Center += Transform.GetTranslation();
+			}
+			else
+			{
+				Box.Center = Transform.TransformPosition(Box.Center);
+			}
+			Box.Rotation = (Transform.GetRotation() * Box.Rotation.Quaternion()).Rotator();
+			Box.X = FMath::Abs(Box.X * Scale.X);
+			Box.Y = FMath::Abs(Box.Y * Scale.Y);
+			Box.Z = FMath::Abs(Box.Z * Scale.Z);
+			TransformedCollision.AggGeom.BoxElems.Add(Box);
+		}
+	}
+
+	int32 NumSphyl = SimpleCollision.AggGeom.SphylElems.Num();
+	if (NumSphyl > 0)
+	{
+		if (!bUniformScale)
+		{
+			if (TransformOptions.bWarnOnInvalidTransforms)
+			{
+				UE::Geometry::AppendWarning(Debug, EGeometryScriptErrorType::InvalidInputs, LOCTEXT("TransformSimpleCollisionShapes Non-Uniform Scale on Sphyl", "TransformSimpleCollisionShapes: Cannot apply Non-Uniform Scale to Sphyls"));
+			}
+			bSuccess = false;
+		}
+
+		TransformedCollision.AggGeom.SphylElems.Reserve(NumSphyl);
+		for (FKSphylElem Sphyl : SimpleCollision.AggGeom.SphylElems)
+		{
+			if (TransformOptions.bCenterTransformPivotPerShape)
+			{
+				Sphyl.Center += Transform.GetTranslation();
+			}
+			else
+			{
+				Sphyl.Center = Transform.TransformPosition(Sphyl.Center);
+			}
+			Sphyl.Radius = FMath::Abs(Sphyl.Radius * ApproxUniformScale);
+			Sphyl.Length = FMath::Abs(Sphyl.Length * ApproxUniformScale);
+			Sphyl.Rotation = (Transform.GetRotation() * Sphyl.Rotation.Quaternion()).Rotator();
+			TransformedCollision.AggGeom.SphylElems.Add(Sphyl);
+		}
+	}
+
+	int32 NumConvex = SimpleCollision.AggGeom.ConvexElems.Num();
+	TransformedCollision.AggGeom.ConvexElems.Reserve(NumConvex);
+	for (const FKConvexElem& Convex : SimpleCollision.AggGeom.ConvexElems)
+	{
+		FKConvexElem& AddConvex = TransformedCollision.AggGeom.ConvexElems.Add_GetRef(Convex);
+		AddConvex.ElemBox = FBox(EForceInit::ForceInit);
+		for (FVector& V : AddConvex.VertexData)
+		{
+			if (TransformOptions.bCenterTransformPivotPerShape)
+			{
+				FVector Center = Convex.ElemBox.GetCenter();
+				V = Transform.TransformPosition(V - Center) + Center;
+			}
+			else
+			{
+				V = Transform.TransformPosition(V);
+			}
+			AddConvex.ElemBox += V;
+		}
+	}
+
+	int32 NumTaperedCapsule = SimpleCollision.AggGeom.TaperedCapsuleElems.Num();
+	if (NumTaperedCapsule > 0)
+	{
+		if (!bUniformScale)
+		{
+			if (TransformOptions.bWarnOnInvalidTransforms)
+			{
+				UE::Geometry::AppendWarning(Debug, EGeometryScriptErrorType::InvalidInputs, LOCTEXT("TransformSimpleCollisionShapes Non-Uniform Scale on Tapered Capsule", "TransformSimpleCollisionShapes: Cannot apply Non-Uniform Scale to Tapered Capsules"));
+			}
+			bSuccess = false;
+		}
+
+		TransformedCollision.AggGeom.TaperedCapsuleElems.Reserve(NumTaperedCapsule);
+		for (FKTaperedCapsuleElem TaperedCapsule : SimpleCollision.AggGeom.TaperedCapsuleElems)
+		{
+			if (TransformOptions.bCenterTransformPivotPerShape)
+			{
+				TaperedCapsule.Center += Transform.GetTranslation();
+			}
+			else
+			{
+				TaperedCapsule.Center = Transform.TransformPosition(TaperedCapsule.Center);
+			}
+			if (ApproxUniformScale < 0)
+			{
+				TaperedCapsule.Radius0 = TaperedCapsule.Radius1 * ApproxUniformScale;
+				TaperedCapsule.Radius1 = TaperedCapsule.Radius0 * ApproxUniformScale;
+				TaperedCapsule.Length *= -ApproxUniformScale;
+			}
+			else
+			{
+				TaperedCapsule.Radius0 *= ApproxUniformScale;
+				TaperedCapsule.Radius1 *= ApproxUniformScale;
+				TaperedCapsule.Length *= ApproxUniformScale;
+			}
+
+			TaperedCapsule.Rotation = (Transform.GetRotation() * TaperedCapsule.Rotation.Quaternion()).Rotator();
+			TransformedCollision.AggGeom.TaperedCapsuleElems.Add(TaperedCapsule);
+		}
+	}
+
+	int32 NumLevelSet = SimpleCollision.AggGeom.LevelSetElems.Num();
+	if (NumLevelSet > 0)
+	{
+		// Non-uniform scale cannot be properly represented if a level set was already rotated
+		if (!bUniformScale)
+		{
+			bool bHasUnsupportedTransform = false;
+			for (const FKLevelSetElem& LevelSet : SimpleCollision.AggGeom.LevelSetElems)
+			{
+				const FTransform ElemTransform = LevelSet.GetTransform();
+				bool bElementNoRotation = ElemTransform.GetRotation().IsIdentity();
+				if (!bElementNoRotation)
+				{
+					bHasUnsupportedTransform = true;
+					break;
+				}
+			}
+			if (bHasUnsupportedTransform)
+			{
+				if (TransformOptions.bWarnOnInvalidTransforms)
+				{
+					UE::Geometry::AppendWarning(Debug, EGeometryScriptErrorType::InvalidInputs, LOCTEXT("TransformSimpleCollisionShapes Unsupported Level Set Transform", "TransformSimpleCollisionShapes: Cannot apply requested transform to Level Set"));
+				}
+				bSuccess = false;
+			}
+		}
+
+		TransformedCollision.AggGeom.LevelSetElems.Reserve(NumLevelSet);
+		for (const FKLevelSetElem& LevelSet : SimpleCollision.AggGeom.LevelSetElems)
+		{
+			FKLevelSetElem& AddLevelSet = TransformedCollision.AggGeom.LevelSetElems.Add_GetRef(LevelSet);
+			if (TransformOptions.bCenterTransformPivotPerShape)
+			{
+				FTransform LevelSetTransform = LevelSet.GetTransform();
+				// Translate to the bounds center, apply Transform, and translate back
+				FVector OrigLocation = LevelSetTransform.TransformPosition(LevelSet.UntransformedAABB().GetCenter());
+				LevelSetTransform.SetLocation(LevelSetTransform.GetLocation() - OrigLocation);
+				LevelSetTransform = LevelSetTransform * Transform;
+				LevelSetTransform.SetTranslation(LevelSetTransform.GetLocation() + OrigLocation);
+				AddLevelSet.SetTransform(LevelSetTransform);
+			}
+			else
+			{
+				AddLevelSet.SetTransform(LevelSet.GetTransform() * Transform);
+			}
+		}
+	}
+
+	int32 NumSkinnedLevelSet = SimpleCollision.AggGeom.SkinnedLevelSetElems.Num();
+	if (NumSkinnedLevelSet > 0)
+	{
+		if (TransformOptions.bWarnOnInvalidTransforms)
+		{
+			UE::Geometry::AppendWarning(Debug, EGeometryScriptErrorType::InvalidInputs, LOCTEXT("TransformSimpleCollisionShapes Cannot Transform Skinned Level Set", "TransformSimpleCollisionShapes: Cannot transform Skinned Level Sets"));
+		}
+		TransformedCollision.AggGeom.SkinnedLevelSetElems = SimpleCollision.AggGeom.SkinnedLevelSetElems;
+		bSuccess = false;
+	}
+
+	return TransformedCollision;
+}
+
 void UGeometryScriptLibrary_CollisionFunctions::SimplifyConvexHulls(
 	FGeometryScriptSimpleCollision& SimpleCollision,
 	const FGeometryScriptConvexHullSimplificationOptions& SimplifyOptions,
diff --git a/Engine/Plugins/Experimental/GeometryScripting/Source/GeometryScriptingCore/Public/GeometryScript/CollisionFunctions.h b/Engine/Plugins/Experimental/GeometryScripting/Source/GeometryScriptingCore/Public/GeometryScript/CollisionFunctions.h
--- a/Engine/Plugins/Experimental/GeometryScripting/Source/GeometryScriptingCore/Public/GeometryScript/CollisionFunctions.h
+++ b/Engine/Plugins/Experimental/GeometryScripting/Source/GeometryScriptingCore/Public/GeometryScript/CollisionFunctions.h
@@ -209,6 +209,27 @@ public:
 	int32 MinTargetFaceCount = 12;
 };
 
+USTRUCT(BlueprintType)
+struct GEOMETRYSCRIPTINGCORE_API FGeometryScriptTransformCollisionOptions
+{
+	GENERATED_BODY()
+public:
+
+	/** Whether to log a warning when a requested transform is not compatible with the simple collision shapes */
+	UPROPERTY(EditAnywhere, BlueprintReadWrite, Category = Options)
+	bool bWarnOnInvalidTransforms = true;
+
+	/**
+	 * If true, we apply the Transform to each collision shape separately, and pivot the Transform around the local center of each shape.
+	 * Otherwise, we apply the Transform to all shapes in the same space, with the pivot at the origin of the origin of that space.
+	 * 
+	 * For example, if we apply a uniform 2x scale to a sphere w/ center (1,1,1), with this enabled, the center will not move and only the radius will scale.
+	 * If this setting is not enabled, the 2x scale will move the sphere center to (2,2,2)
+	 */
+	UPROPERTY(EditAnywhere, BlueprintReadWrite, Category = Options)
+	bool bCenterTransformPivotPerShape = false;
+};
+
 UCLASS(meta = (ScriptName = "GeometryScript_Collision"))
 class GEOMETRYSCRIPTINGCORE_API UGeometryScriptLibrary_CollisionFunctions : public UBlueprintFunctionLibrary
 {
@@ -314,6 +335,19 @@ public:
 		return SimpleCollision.AggGeom.GetElementCount();
 	}
 
+	/*
+	 * Transform simple collision shapes
+	 * @param bSuccess	Indicates whether all collision shapes were accurately transformed. On failure, shapes will still be copied over and a best-effort transform will still be applied.
+	 */
+	UFUNCTION(BlueprintCallable, Category = "GeometryScript|Collision", meta = (AutoCreateRefTerm = "TransformOptions"))
+	static UPARAM(DisplayName = "Transformed Collision") FGeometryScriptSimpleCollision TransformSimpleCollisionShapes(
+		const FGeometryScriptSimpleCollision& SimpleCollision,
+		FTransform Transform,
+		const FGeometryScriptTransformCollisionOptions& TransformOptions,
+		bool& bSuccess,
+		UGeometryScriptDebug* Debug = nullptr
+	);
+	
 	/**
 	 * Simplify any convex hulls in the given simple collision representation. Updates the passed-in Simple Collision.
 	 *
@@ -321,7 +355,7 @@ public:
 	 * @param ConvexSimplifyOptions	Options controlling how convex hulls are simplified
 	 * @param bHasSimplified		Indicates whether any convex hulls were modified
 	 */
-	UFUNCTION(BlueprintCallable, Category = "GeometryScript|Collision", meta = (AutoCreateRefTerm = "SimplifyOptions"))
+	UFUNCTION(BlueprintCallable, Category = "GeometryScript|Collision", meta = (ScriptMethod, AutoCreateRefTerm = "SimplifyOptions"))
 	static void SimplifyConvexHulls(
 		UPARAM(ref) FGeometryScriptSimpleCollision& SimpleCollision,
 		const FGeometryScriptConvexHullSimplificationOptions& SimplifyOptions,
