Add Geometry Script function to combine simple collision shapes, by adding the shapes from the second to the first.

diff --git a/Engine/Plugins/Experimental/GeometryScripting/Source/GeometryScriptingCore/Private/CollisionFunctions.cpp b/Engine/Plugins/Experimental/GeometryScripting/Source/GeometryScriptingCore/Private/CollisionFunctions.cpp
--- a/Engine/Plugins/Experimental/GeometryScripting/Source/GeometryScriptingCore/Private/CollisionFunctions.cpp
+++ b/Engine/Plugins/Experimental/GeometryScripting/Source/GeometryScriptingCore/Private/CollisionFunctions.cpp
@@ -791,6 +791,20 @@ FGeometryScriptSimpleCollision UGeometryScriptLibrary_CollisionFunctions::Transf
 	return TransformedCollision;
 }
 
+void UGeometryScriptLibrary_CollisionFunctions::CombineSimpleCollision(
+	FGeometryScriptSimpleCollision& Collision,
+	const FGeometryScriptSimpleCollision& AppendCollision
+)
+{
+	Collision.AggGeom.BoxElems.Append(AppendCollision.AggGeom.BoxElems);
+	Collision.AggGeom.ConvexElems.Append(AppendCollision.AggGeom.ConvexElems);
+	Collision.AggGeom.SphylElems.Append(AppendCollision.AggGeom.SphylElems);
+	Collision.AggGeom.SphereElems.Append(AppendCollision.AggGeom.SphereElems);
+	Collision.AggGeom.TaperedCapsuleElems.Append(AppendCollision.AggGeom.TaperedCapsuleElems);
+	Collision.AggGeom.LevelSetElems.Append(AppendCollision.AggGeom.LevelSetElems);
+	Collision.AggGeom.SkinnedLevelSetElems.Append(AppendCollision.AggGeom.SkinnedLevelSetElems);
+}
+
 void UGeometryScriptLibrary_CollisionFunctions::SimplifyConvexHulls(
 	FGeometryScriptSimpleCollision& SimpleCollision,
 	const FGeometryScriptConvexHullSimplificationOptions& SimplifyOptions,
diff --git a/Engine/Plugins/Experimental/GeometryScripting/Source/GeometryScriptingCore/Public/GeometryScript/CollisionFunctions.h b/Engine/Plugins/Experimental/GeometryScripting/Source/GeometryScriptingCore/Public/GeometryScript/CollisionFunctions.h
--- a/Engine/Plugins/Experimental/GeometryScripting/Source/GeometryScriptingCore/Public/GeometryScript/CollisionFunctions.h
+++ b/Engine/Plugins/Experimental/GeometryScripting/Source/GeometryScriptingCore/Public/GeometryScript/CollisionFunctions.h
@@ -347,6 +347,15 @@ public:
 		bool& bSuccess,
 		UGeometryScriptDebug* Debug = nullptr
 	);
+
+	/*
+	 * Add simple collision shapes from AppendCollision to CollisionToUpdate
+	 */
+	UFUNCTION(BlueprintCallable, Category = "GeometryScript|Collision", meta = (ScriptMethod))
+	static void CombineSimpleCollision(
+		UPARAM(ref) FGeometryScriptSimpleCollision& CollisionToUpdate,
+		const FGeometryScriptSimpleCollision& AppendCollision
+	);
 	
 	/**
 	 * Simplify any convex hulls in the given simple collision representation. Updates the passed-in Simple Collision.
