add new properties to control the sample range for the geometry script spline->polypath and spline->transform array functions

diff --git a/Engine/Plugins/Experimental/GeometryScripting/Source/GeometryScriptingCore/Private/PolyPathFunctions.cpp b/Engine/Plugins/Experimental/GeometryScripting/Source/GeometryScriptingCore/Private/PolyPathFunctions.cpp
--- a/Engine/Plugins/Experimental/GeometryScripting/Source/GeometryScriptingCore/Private/PolyPathFunctions.cpp
+++ b/Engine/Plugins/Experimental/GeometryScripting/Source/GeometryScriptingCore/Private/PolyPathFunctions.cpp
@@ -191,83 +191,312 @@ void UGeometryScriptLibrary_PolyPathFunctions::ConvertArrayOfVector2DToPolyPath(
 	}
 }
 
+namespace UE::Local::SplinePathHelpers
+{
+	// Get the sampling range, and whether it is a distance range (otherwise, assume it's a time range)
+	void GetRangeFromSamplingOptions(const USplineComponent* Spline, const FGeometryScriptSplineSamplingOptions& SamplingOptions, float& OutStart, float& OutEnd, bool& bOutIsDistanceRange)
+	{
+		OutStart = SamplingOptions.RangeStart;
+		OutEnd = SamplingOptions.RangeEnd;
+		bOutIsDistanceRange = 
+			   SamplingOptions.RangeMethod == EGeometryScriptEvaluateSplineRange::FullSpline
+			|| SamplingOptions.RangeMethod == EGeometryScriptEvaluateSplineRange::DistanceRange;
+		if (bOutIsDistanceRange)
+		{
+			float SplineLength = Spline->GetSplineLength();
+			
+			if (SamplingOptions.RangeMethod == EGeometryScriptEvaluateSplineRange::FullSpline)
+			{
+				OutStart = 0;
+				OutEnd = SplineLength;
+			}
+		}
+		// else it's a time range; values can be left as they are
+	}
+
+	// Note: this does *not* support wrapped values
+	float DistanceToTime(const USplineComponent* Spline, float Distance, bool bConstantSpeed)
+	{
+		if (bConstantSpeed)
+		{
+			return Distance * Spline->Duration / Spline->GetSplineLength();
+		}
+		else
+		{
+			checkSlow(Distance >= 0 && Distance <= Spline->GetSplineLength());
+			return Spline->GetTimeAtDistanceAlongSpline(Distance);
+		}
+	};
+
+	// Convert a Time parameter from a 'Constant Speed' to a 'Varying Speed' value, or vice versa
+	// e.g., if bWasConstantSpeed is true, this reports what Time would get to the same location on the spline if we were instead evaluating the spline at varying speed
+	float SwitchTimeType(const USplineComponent* Spline, float Time, bool bWasConstantSpeed, bool bAllowWrappingIfClosed = true)
+	{
+		const float Duration = Spline->Duration;
+		float InRangeTime = Time;
+		float TimeAtStartOfLoop = 0;
+		if (!bAllowWrappingIfClosed)
+		{
+			InRangeTime = FMath::Clamp(Time, 0, Duration);
+		}
+		else
+		{
+			InRangeTime = FMath::Fmod(Time, Duration);
+			if (InRangeTime < 0)
+			{
+				InRangeTime += Duration;
+			}
+			TimeAtStartOfLoop = FMath::Floor(Time / Duration) * Duration;
+		}
+		const float InRangeTimeFrac = InRangeTime / Duration;
+		float InRangeDistance = 0;
+		if (bWasConstantSpeed)
+		{
+			InRangeDistance = InRangeTimeFrac * Spline->GetSplineLength();
+		}
+		else
+		{
+			const int32 NumSegments = Spline->GetNumberOfSplineSegments();
+			// Note: 'InputKey' values correspond to the spline in parameter space, in the range of 0 to NumSegments
+			const float InRangeInputKey = InRangeTimeFrac * NumSegments;
+			InRangeDistance = Spline->GetDistanceAlongSplineAtSplineInputKey(InRangeInputKey);
+		}
+
+		float ConvertedInRangeTime = DistanceToTime(Spline, InRangeDistance, !bWasConstantSpeed);
+		return TimeAtStartOfLoop + ConvertedInRangeTime;
+	}
+
+	// Helper to handle UniformDistance or UniformTime sampling
+	bool IterateSamplesInRange(
+		const USplineComponent* Spline,
+		const FGeometryScriptSplineSamplingOptions& SamplingOptions,
+		bool bWantPos,
+		TFunctionRef<void(const FVector&)> PositionFunc,
+		bool bWantTime,
+		TFunctionRef<void(float)> TimeFunc
+	)
+	{
+		if (!bWantPos && !bWantTime)
+		{
+			return false;
+		}
+
+		float Start, End;
+		bool bWasDistanceRange;
+		GetRangeFromSamplingOptions(Spline, SamplingOptions, Start, End, bWasDistanceRange);
+		// For consistent handling below, always convert Distance to a TimeRange_ConstantSpeed-based range (by re-scaling)
+		if (bWasDistanceRange)
+		{
+			float DistanceToTimeFrac = Spline->Duration / Spline->GetSplineLength();
+			Start *= DistanceToTimeFrac;
+			End *= DistanceToTimeFrac;
+		}
+
+		const bool bAllowWrap = Spline->IsClosedLoop() && SamplingOptions.RangeMethod != EGeometryScriptEvaluateSplineRange::FullSpline;
+
+		// UniformDistance and UniformTime sampling correspond to two different parameterizations of the spline:
+		// TimeRange_ConstantSpeed and TimeRange_VariableSpeed, respectively.
+		// To make the sampling easier, we convert the range values to use the same parameterization as the sampling.
+		bool bTimeRangeWasConstantSpeed = SamplingOptions.RangeMethod != EGeometryScriptEvaluateSplineRange::TimeRange_VariableSpeed;
+		const bool bUniformSampleSpacing = SamplingOptions.SampleSpacing == EGeometryScriptSampleSpacing::UniformDistance;
+		if (bUniformSampleSpacing != bTimeRangeWasConstantSpeed)
+		{
+			Start = UE::Local::SplinePathHelpers::SwitchTimeType(Spline, Start, bTimeRangeWasConstantSpeed, bAllowWrap);
+			End = UE::Local::SplinePathHelpers::SwitchTimeType(Spline, End, bTimeRangeWasConstantSpeed, bAllowWrap);
+		}
+		
+		const float MaxValue = Spline->Duration;
+		if (!bAllowWrap)
+		{
+			End = FMath::Min(MaxValue, End);
+			Start = FMath::Max(0, Start);
+		}
+
+		const float Range = End - Start;
+		if (Range < 0)
+		{
+			return false;
+		}
+
+		auto GetPos = [bUniformSampleSpacing, Spline, &SamplingOptions](float Value) -> FVector
+		{
+			return Spline->GetLocationAtTime(Value, SamplingOptions.CoordinateSpace, bUniformSampleSpacing);
+		};
+		auto WrapValue = [MaxValue](float Value)
+		{
+			float WrappedValue = FMath::Fmod(Value, MaxValue);
+			if (WrappedValue < 0)
+			{
+				WrappedValue += MaxValue;
+			}
+			return WrappedValue;
+		};
+
+		int32 UseSamples = FMath::Max(2, SamplingOptions.NumSamples);
+		// If we span 0 range, just report the single point
+		if (Range == 0)
+		{
+			UseSamples = 1;
+		}
+		// In non-loops, we adjust DivNum so we exactly sample the end of the spline
+		// In loops we don't sample the endpoint, by convention, as it's the same as the start
+		const bool bSampleExactEnd = !(Spline->IsClosedLoop() && SamplingOptions.RangeMethod == EGeometryScriptEvaluateSplineRange::FullSpline);
+		const float DivNum = float(UseSamples - (int32)bSampleExactEnd);
+
+		for (int32 Idx = 0; Idx < UseSamples; Idx++)
+		{
+			float Value = Start + Range * ((float)Idx / DivNum);
+			if (bAllowWrap)
+			{
+				Value = WrapValue(Value);
+			}
+
+			if (bWantPos)
+			{
+				FVector Pos = GetPos(Value);
+				PositionFunc(Pos);
+			}
+			if (bWantTime)
+			{
+				TimeFunc(Value);
+			}
+		}
+		return true;
+	}
+
+	// return true if it's a full loop, and we've omitted the final point (which matched the start point)
+	bool ErrorBasedSampleInRange(const USplineComponent* Spline, FGeometryScriptSplineSamplingOptions SamplingOptions, TArray<FVector>& OutPositions, TArray<double>& OutDistances)
+	{
+		// this function does error tolerance sampling so we expect the options to match that
+		ensure(SamplingOptions.SampleSpacing == EGeometryScriptSampleSpacing::ErrorTolerance);
+
+		bool bIsLoop = Spline->IsClosedLoop() && SamplingOptions.RangeMethod == EGeometryScriptEvaluateSplineRange::FullSpline;
+		float SquaredErrorTolerance = FMath::Max(KINDA_SMALL_NUMBER, SamplingOptions.ErrorTolerance * SamplingOptions.ErrorTolerance);
+		float UseRangeStart, UseRangeEnd;
+		bool bIsDistanceRange;
+		UE::Local::SplinePathHelpers::GetRangeFromSamplingOptions(Spline, SamplingOptions, UseRangeStart, UseRangeEnd, bIsDistanceRange);
+		if (bIsDistanceRange)
+		{
+			bool bAllowWrap = SamplingOptions.RangeMethod != EGeometryScriptEvaluateSplineRange::FullSpline;
+			Spline->ConvertSplineToPolyline_InDistanceRange(
+				SamplingOptions.CoordinateSpace, SquaredErrorTolerance, UseRangeStart, UseRangeEnd, OutPositions, OutDistances, bAllowWrap);
+		}
+		else
+		{
+			bool bConstantSpeed = SamplingOptions.RangeMethod == EGeometryScriptEvaluateSplineRange::TimeRange_ConstantSpeed;
+			Spline->ConvertSplineToPolyline_InTimeRange(
+				SamplingOptions.CoordinateSpace, SquaredErrorTolerance, UseRangeStart, UseRangeEnd, bConstantSpeed, OutPositions, OutDistances, true);
+		}
+		if (bIsLoop && !OutPositions.IsEmpty())
+		{
+			// delete the duplicate end-point for loops
+			OutPositions.RemoveAt(OutPositions.Num() - 1);
+			OutDistances.RemoveAt(OutDistances.Num() - 1);
+		}
+		checkSlow(OutPositions.Num() == OutDistances.Num());
+
+		return bIsLoop;
+	}
+}
+
 void UGeometryScriptLibrary_PolyPathFunctions::ConvertSplineToPolyPath(const USplineComponent* Spline, FGeometryScriptPolyPath& PolyPath, FGeometryScriptSplineSamplingOptions SamplingOptions)
 {
 	PolyPath.Reset();
 	if (Spline)
 	{
-		bool bIsLoop = Spline->IsClosedLoop();
+		bool bIsLoop = Spline->IsClosedLoop() && SamplingOptions.RangeMethod == EGeometryScriptEvaluateSplineRange::FullSpline;
 		PolyPath.bClosedLoop = bIsLoop;
 		if (SamplingOptions.SampleSpacing == EGeometryScriptSampleSpacing::ErrorTolerance)
 		{
-			float SquaredErrorTolerance = FMath::Max(KINDA_SMALL_NUMBER, SamplingOptions.ErrorTolerance * SamplingOptions.ErrorTolerance);
-			Spline->ConvertSplineToPolyLine(SamplingOptions.CoordinateSpace, SquaredErrorTolerance, *PolyPath.Path);
-			if (bIsLoop)
-			{
-				PolyPath.Path->Pop(); // delete the duplicate end-point for loops
-			}
+			TArray<double> Distances_Unused;
+			UE::Local::SplinePathHelpers::ErrorBasedSampleInRange(Spline, SamplingOptions, *PolyPath.Path, Distances_Unused);
 		}
 		else
 		{
-			float Duration = Spline->Duration;
-			
-			bool bUseConstantVelocity = SamplingOptions.SampleSpacing == EGeometryScriptSampleSpacing::UniformDistance;
-			int32 UseSamples = FMath::Max(2, SamplingOptions.NumSamples); // Always use at least 2 samples
-			// In non-loops, we adjust DivNum so we exactly sample the end of the spline
-			// In loops we don't sample the endpoint, by convention, as it's the same as the start
-			float DivNum = float(UseSamples - (int32)!bIsLoop);
-			PolyPath.Path->Reserve(UseSamples);
-			for (int32 Idx = 0; Idx < UseSamples; Idx++)
-			{
-				float Time = Duration * ((float)Idx / DivNum);
-				PolyPath.Path->Add(Spline->GetLocationAtTime(Time, SamplingOptions.CoordinateSpace, bUseConstantVelocity));
-			}
+			int32 ExpectSamples = FMath::Max(2, SamplingOptions.NumSamples);
+			PolyPath.Path->Reserve(ExpectSamples);
+			UE::Local::SplinePathHelpers::IterateSamplesInRange(
+				Spline,
+				SamplingOptions,
+				true,
+				[&PolyPath](FVector Pos)
+				{
+					PolyPath.Path->Add(Pos);
+				},
+				false,
+				[](float) {});
 		}
 	}
 }
 
 
-void UGeometryScriptLibrary_PolyPathFunctions::SampleSplineToTransforms(
+bool UGeometryScriptLibrary_PolyPathFunctions::SampleSplineToTransforms(
 	const USplineComponent* Spline, 
 	TArray<FTransform>& Frames, 
 	TArray<double>& FrameTimes,
 	FGeometryScriptSplineSamplingOptions SamplingOptions,
 	FTransform RelativeTransform,
-	bool bIncludeScale)
+	bool bIncludeScale
+)
 {
 	Frames.Reset();
 	FrameTimes.Reset();
 
-	// Currently ErrorTolerance sampling can only be done via Spline->ConvertSplineToPolyLine, which only returns a list of points.
-	// To convert to Transforms we would have to reverse-engineer the Time at each Point which could be very expensive...
+	// If we're using a time-based sampling method, times are constant-speed only if the sampling also is
+	// (Note: This is how the time values are generated in IterateSamplesInRange, so if we want to change this, we'd need to also change their generation there.
+	//   ... also, if this is changed, it should be done in a way that doesn't break the behavior of existing scripts that used this function.)
+	bool bOutputConstantSpeed = SamplingOptions.SampleSpacing == EGeometryScriptSampleSpacing::UniformDistance;
+	// If we're using error-tolerance-based sampling, we can match output time values to the input range time values
 	if (SamplingOptions.SampleSpacing == EGeometryScriptSampleSpacing::ErrorTolerance)
 	{
-		UE_LOG(LogGeometry, Warning, TEXT("SampleSplineToTransforms: ErrorTolerance sampling mode is currently not supported, falling back to UniformDistance"));
-		SamplingOptions.SampleSpacing = EGeometryScriptSampleSpacing::UniformDistance;
+		bOutputConstantSpeed = SamplingOptions.RangeMethod != EGeometryScriptEvaluateSplineRange::TimeRange_VariableSpeed;
 	}
-
-	if (Spline != nullptr )
+	
+	if (!Spline)
 	{
-		bool bIsLoop = Spline->IsClosedLoop();
+		return bOutputConstantSpeed;
+	}
 
-		float Duration = Spline->Duration;
+	auto AddAtTime = [Spline, &Frames, &FrameTimes, &SamplingOptions, &RelativeTransform, bIncludeScale, bOutputConstantSpeed](float Time)
+	{
+		FTransform Transform = Spline->GetTransformAtTime(Time, SamplingOptions.CoordinateSpace, bOutputConstantSpeed, bIncludeScale);
+		FTransform::Multiply(&Transform, &RelativeTransform, &Transform);
+		Frames.Add(Transform);
+		FrameTimes.Add(Time);
+	};
 
-		bool bUseConstantVelocity = SamplingOptions.SampleSpacing == EGeometryScriptSampleSpacing::UniformDistance;
-		int32 UseSamples = FMath::Max(2, SamplingOptions.NumSamples); // Always use at least 2 samples
-																		// In non-loops, we adjust DivNum so we exactly sample the end of the spline
-																		// In loops we don't sample the endpoint, by convention, as it's the same as the start
-		float DivNum = float(UseSamples - (int32)!bIsLoop);
-		Frames.Reserve(UseSamples);
-		FrameTimes.Reserve(UseSamples);
-		for (int32 Idx = 0; Idx < UseSamples; Idx++)
+	bool bIsLoop = Spline->IsClosedLoop() && SamplingOptions.RangeMethod == EGeometryScriptEvaluateSplineRange::FullSpline;
+	if (SamplingOptions.SampleSpacing == EGeometryScriptSampleSpacing::ErrorTolerance)
+	{
+		TArray<double> Distances;
+		TArray<FVector> Positions_Unused;
+		UE::Local::SplinePathHelpers::ErrorBasedSampleInRange(Spline, SamplingOptions, Positions_Unused, Distances);
+		Frames.Reserve(Distances.Num());
+		FrameTimes.Reserve(Distances.Num());
+		for (float Dist : Distances)
 		{
-			float Time = Duration * ((float)Idx / DivNum);
-			FTransform Transform = Spline->GetTransformAtTime(Time, SamplingOptions.CoordinateSpace, bUseConstantVelocity, bIncludeScale);
-			FTransform::Multiply(&Transform, &RelativeTransform, &Transform);
-			Frames.Add(Transform);
-			FrameTimes.Add(Time);
+			float Time = UE::Local::SplinePathHelpers::DistanceToTime(Spline, Dist, bOutputConstantSpeed);
+			AddAtTime(Time);
 		}
 	}
+	else
+	{
+		int32 ExpectSamples = FMath::Max(2, SamplingOptions.NumSamples);
+		Frames.Reserve(ExpectSamples);
+		FrameTimes.Reserve(ExpectSamples);
+		UE::Local::SplinePathHelpers::IterateSamplesInRange(
+			Spline,
+			SamplingOptions,
+			false,
+			[](FVector Pos) {},
+			true,
+			[&](float Time) 
+			{
+				AddAtTime(Time);
+			});
+	}
+
+	return bOutputConstantSpeed;
 }
 
 
diff --git a/Engine/Plugins/Experimental/GeometryScripting/Source/GeometryScriptingCore/Public/GeometryScript/PolyPathFunctions.h b/Engine/Plugins/Experimental/GeometryScripting/Source/GeometryScriptingCore/Public/GeometryScript/PolyPathFunctions.h
--- a/Engine/Plugins/Experimental/GeometryScripting/Source/GeometryScriptingCore/Public/GeometryScript/PolyPathFunctions.h
+++ b/Engine/Plugins/Experimental/GeometryScripting/Source/GeometryScriptingCore/Public/GeometryScript/PolyPathFunctions.h
@@ -17,6 +17,18 @@ enum class EGeometryScriptSampleSpacing : uint8
 	ErrorTolerance
 };
 
+UENUM(BlueprintType)
+enum class EGeometryScriptEvaluateSplineRange : uint8
+{
+	// Evaluate the full spline, ignoring any specified range
+	FullSpline,
+	// Evaluate a range specified by distances along the spline
+	DistanceRange,
+	// Evaluate a range specified by times, based on travelling at constant speed along the spline
+	TimeRange_ConstantSpeed,
+	// Evaluate a range specified by times, based on travelling at a constant rate of spline segments/second
+	TimeRange_VariableSpeed
+};
 
 USTRUCT(BlueprintType)
 struct GEOMETRYSCRIPTINGCORE_API FGeometryScriptSplineSamplingOptions
@@ -34,6 +46,19 @@ public:
 
 	UPROPERTY(EditAnywhere, BlueprintReadWrite, Category = Options)
 	TEnumAsByte<ESplineCoordinateSpace::Type> CoordinateSpace = ESplineCoordinateSpace::Type::Local;
+
+	// How the RangeStart and RangeEnd parameters will be interpreted
+	UPROPERTY(EditAnywhere, BlueprintReadWrite, Category = Options)
+	EGeometryScriptEvaluateSplineRange RangeMethod = EGeometryScriptEvaluateSplineRange::FullSpline;
+
+	// If not evaluating the full spline, where to start sampling. Expressed in units based on the EvaluateRange value.
+	UPROPERTY(BlueprintReadWrite, Category = Options, meta = (EditCondition = "EvaluateRange != EGeometryScriptEvaluateSplineRange::FullSpline"))
+	float RangeStart = 0;
+
+	// If not evaluating the full spline, where to stop sampling. Expressed in units based on the EvaluateRange value.
+	UPROPERTY(BlueprintReadWrite, Category = Options, meta = (EditCondition = "EvaluateRange != EGeometryScriptEvaluateSplineRange::FullSpline"))
+	float RangeEnd = 1;
+
 };
 
 
@@ -124,12 +149,13 @@ public:
 	/**
 	 * Sample a USplineComponent into a list of FTransforms, based on the given SamplingOptions.
 	 * @param Frames Transforms are returned here, with X axis oriented along spline Tangent and Z as the 'Up' vector.
-	 * @param FrameTimes the spline Time value used for each Frame
+	 * @param FrameTimes the spline Time value used for each Frame. Note the Times Use Constant Velocity output indicates whether these times are w.r.t. to a constant-speed parameterization of the spline.
 	 * @param RelativeTransform a constant Transform applied to each sample Transform in its local frame of reference. So, eg, an X Rotation will rotate each frame around the local spline Tangent vector
 	 * @param bIncludeScale if true, the Scale of each FTransform is taken from the Spline, otherwise the Transforms have unit scale
+	 * @return whether the FrameTimes are w.r.t. a 'constant speed' traversal of the spline
 	 */
 	UFUNCTION(BlueprintCallable, Category = "GeometryScript|PolyPath")
-	static void SampleSplineToTransforms(
+	static UPARAM(DisplayName = "Times Use Constant Velocity") bool SampleSplineToTransforms(
 		const USplineComponent* Spline, 
 		TArray<FTransform>& Frames, 
 		TArray<double>& FrameTimes,
