Add angle-tolerance method of convex hull simplification to the geometry script convex hull simplify options

Add a face normal clustering function that returns the vertices at corners of the clusters (vertices touching at least 3 different clusters), for use in the angle-tolerance convex hull simplification method

diff --git a/Engine/Plugins/Experimental/GeometryScripting/Source/GeometryScriptingCore/Private/CollisionFunctions.cpp b/Engine/Plugins/Experimental/GeometryScripting/Source/GeometryScriptingCore/Private/CollisionFunctions.cpp
--- a/Engine/Plugins/Experimental/GeometryScripting/Source/GeometryScriptingCore/Private/CollisionFunctions.cpp
+++ b/Engine/Plugins/Experimental/GeometryScripting/Source/GeometryScriptingCore/Private/CollisionFunctions.cpp
@@ -21,6 +21,7 @@
 #include "ShapeApproximation/ShapeDetection3.h"
 #include "ShapeApproximation/MeshSimpleShapeApproximation.h"
 
+#include "Clustering/FaceNormalClustering.h"
 #include "CompGeom/ConvexDecomposition3.h"
 #include "OrientedBoxTypes.h"
 #include "MeshQueries.h"
@@ -224,7 +225,7 @@ static void SetStaticMeshSimpleCollision(UStaticMesh* StaticMeshAsset, const FKA
 
 
 // local helper to convert the blueprint-accessible enum to the geometrycore equivalent
-static UE::Geometry::FNegativeSpaceSampleSettings::ESampleMethod ConvertNegativeSpaceSampleMethodDataflowEnum(ENegativeSpaceSampleMethod SampleMethod)
+static UE::Geometry::FNegativeSpaceSampleSettings::ESampleMethod ConvertNegativeSpaceSampleMethodEnum(ENegativeSpaceSampleMethod SampleMethod)
 {
 	switch (SampleMethod)
 	{
@@ -244,7 +245,7 @@ static FNegativeSpaceSampleSettings ConvertNegativeSpaceOptions(const FComputeNe
 	NegativeSpaceSettings.MinRadius = NegativeSpaceOptions.MinRadius;
 	NegativeSpaceSettings.ReduceRadiusMargin = NegativeSpaceOptions.NegativeSpaceTolerance;
 	NegativeSpaceSettings.MinSpacing = NegativeSpaceOptions.MinSampleSpacing;
-	NegativeSpaceSettings.SampleMethod = ConvertNegativeSpaceSampleMethodDataflowEnum(NegativeSpaceOptions.SampleMethod);
+	NegativeSpaceSettings.SampleMethod = ConvertNegativeSpaceSampleMethodEnum(NegativeSpaceOptions.SampleMethod);
 	NegativeSpaceSettings.bRequireSearchSampleCoverage = NegativeSpaceOptions.bRequireSearchSampleCoverage;
 	NegativeSpaceSettings.bReferenceMeshHasNegativeWinding = false;
 	NegativeSpaceSettings.Sanitize();
@@ -819,40 +820,60 @@ void UGeometryScriptLibrary_CollisionFunctions::SimplifyConvexHulls(
 		FKConvexElem& Elem = ConvexElems[ConvexIdx];
 		Elem.ComputeChaosConvexIndices(false); // make sure indices are computed
 
-		int32 TargetTriangleCount = FMath::Max(4, SimplifyOptions.MinTargetFaceCount);
+		int32 TargetFaceCount = FMath::Max(4, SimplifyOptions.MinTargetFaceCount);
 
 		// Convert hull to a dynamic mesh
 		FDynamicMesh3 Mesh;
-		if (!UELocal::AppendConvexElemToCompactDynamicMesh(Elem, Mesh, TargetTriangleCount))
+		if (!UELocal::AppendConvexElemToCompactDynamicMesh(Elem, Mesh, TargetFaceCount))
 		{
 			continue;
 		}
 
 		int32 InitialTriangleCount = Mesh.TriangleCount();
 
-		// Run simplification
-		FVolPresMeshSimplification Simplifier(&Mesh);
-		Simplifier.CollapseMode = FVolPresMeshSimplification::ESimplificationCollapseModes::MinimalExistingVertexError;
-		Simplifier.GeometricErrorConstraint = UE::Geometry::FVolPresMeshSimplification::EGeometricErrorCriteria::PredictedPointToProjectionTarget;
-		Simplifier.GeometricErrorTolerance = SimplifyOptions.SimplificationDistanceThreshold;
-
-		FDynamicMesh3 ProjectionTargetMesh(Mesh);
-		FDynamicMeshAABBTree3 ProjectionTargetSpatial(&ProjectionTargetMesh, true);
-		FMeshProjectionTarget ProjTarget(&ProjectionTargetMesh, &ProjectionTargetSpatial);
-		Simplifier.SetProjectionTarget(&ProjTarget);
-		Simplifier.SimplifyToTriangleCount(TargetTriangleCount);
-
-		// Simplification didn't reduce triangle count, so skip updating the convex hull
-		if (Mesh.TriangleCount() == InitialTriangleCount)
+		if (SimplifyOptions.SimplificationMethod == EGeometryScriptConvexHullSimplifyMethod::MeshQSlim)
 		{
-			continue;
-		}
+			// Run simplification
+			FVolPresMeshSimplification Simplifier(&Mesh);
+			Simplifier.CollapseMode = FVolPresMeshSimplification::ESimplificationCollapseModes::MinimalExistingVertexError;
+			Simplifier.GeometricErrorConstraint = UE::Geometry::FVolPresMeshSimplification::EGeometricErrorCriteria::PredictedPointToProjectionTarget;
+			Simplifier.GeometricErrorTolerance = SimplifyOptions.SimplificationDistanceThreshold;
+
+			FDynamicMesh3 ProjectionTargetMesh(Mesh);
+			FDynamicMeshAABBTree3 ProjectionTargetSpatial(&ProjectionTargetMesh, true);
+			FMeshProjectionTarget ProjTarget(&ProjectionTargetMesh, &ProjectionTargetSpatial);
+			Simplifier.SetProjectionTarget(&ProjTarget);
+			Simplifier.SimplifyToTriangleCount(TargetFaceCount);
+
+			// Simplification didn't reduce triangle count, so skip updating the convex hull
+			if (Mesh.TriangleCount() == InitialTriangleCount)
+			{
+				continue;
+			}
 
-		Elem.VertexData.Reset(Mesh.VertexCount());
-		for (FVector3d V : Mesh.VerticesItr())
+			Elem.VertexData.Reset(Mesh.VertexCount());
+			for (FVector3d V : Mesh.VerticesItr())
+			{
+				Elem.VertexData.Add(V);
+			}
+		}
+		else // EGeometryScriptConvexHullSimplifyMethod::AngleTolerance
 		{
-			Elem.VertexData.Add(V);
+			TArray<int32> VertexCorners;
+			FaceNormalClustering::FClusterOptions ClusterOptions;
+			ClusterOptions.TargetMinGroups = TargetFaceCount;
+			ClusterOptions.bApplyNormalToleranceToClusters = true;
+			ClusterOptions.SmallFaceAreaThreshold = UE_DOUBLE_KINDA_SMALL_NUMBER;
+			ClusterOptions.SetNormalAngleToleranceInDegrees(SimplifyOptions.SimplificationAngleThreshold);
+			FaceNormalClustering::ComputeClusterCornerVertices(Mesh, VertexCorners, ClusterOptions, nullptr);
+			Elem.VertexData.Reset(VertexCorners.Num());
+			for (int32 CornerVID : VertexCorners)
+			{
+				Elem.VertexData.Add(Mesh.GetVertex(CornerVID));
+			}
 		}
+
+
 		Elem.UpdateElemBox();
 		bHasSimplified = true;
 	}
diff --git a/Engine/Plugins/Experimental/GeometryScripting/Source/GeometryScriptingCore/Public/GeometryScript/CollisionFunctions.h b/Engine/Plugins/Experimental/GeometryScripting/Source/GeometryScriptingCore/Public/GeometryScript/CollisionFunctions.h
--- a/Engine/Plugins/Experimental/GeometryScripting/Source/GeometryScriptingCore/Public/GeometryScript/CollisionFunctions.h
+++ b/Engine/Plugins/Experimental/GeometryScripting/Source/GeometryScriptingCore/Public/GeometryScript/CollisionFunctions.h
@@ -194,17 +194,38 @@ public:
 	FGeometryScriptSimpleCollisionTriangulationOptions ShapeToHullTriangulation;
 };
 
+// Methods to simplify convex hulls, used by FGeometryScriptConvexHullSimplificationOptions
+UENUM(BlueprintType)
+enum class EGeometryScriptConvexHullSimplifyMethod : uint8
+{
+	// Simplify convex hulls using a general mesh-based simplifier, and taking the convex hull of the simplified mesh
+	MeshQSlim,
+	// Simplify convex hulls by merging hull faces that have similar normals
+	AngleTolerance
+};
+
 USTRUCT(BlueprintType)
 struct GEOMETRYSCRIPTINGCORE_API FGeometryScriptConvexHullSimplificationOptions
 {
 	GENERATED_BODY()
 public:
 
-	/** Simplified hull should stay within this distance of the initial convex hull. */
+	/** Method to use to simplify convex hulls */
+	UPROPERTY(EditAnywhere, BlueprintReadWrite, Category = Options)
+	EGeometryScriptConvexHullSimplifyMethod SimplificationMethod = EGeometryScriptConvexHullSimplifyMethod::MeshQSlim;
+
+	/** Simplified hull should stay within this distance of the initial convex hull. Used by the MeshQSlim simplification method. */
 	UPROPERTY(EditAnywhere, BlueprintReadWrite, Category = Options)
 	float SimplificationDistanceThreshold = 10.f;
 
-	/** The minimum number of faces to use for the convex hull. */
+	/** Simplified hull should preserve angles larger than this (in degrees). Used by the AngleTolerance simplification method. */
+	UPROPERTY(EditAnywhere, BlueprintReadWrite, Category = Options)
+	float SimplificationAngleThreshold = 10.f;
+
+	/** 
+	 * The minimum number of faces to use for the convex hull.
+	 * Note that for the MeshQSlim method all faces are triangles, while the AngleTolerance method can consider more general polygons. 
+	 */
 	UPROPERTY(EditAnywhere, BlueprintReadWrite, Category = Options, meta = (ClampMin = 4))
 	int32 MinTargetFaceCount = 12;
 };
diff --git a/Engine/Source/Runtime/GeometryCore/Private/Clustering/FaceNormalClustering.cpp b/Engine/Source/Runtime/GeometryCore/Private/Clustering/FaceNormalClustering.cpp
--- a/Engine/Source/Runtime/GeometryCore/Private/Clustering/FaceNormalClustering.cpp
+++ b/Engine/Source/Runtime/GeometryCore/Private/Clustering/FaceNormalClustering.cpp
@@ -10,11 +10,13 @@
 #include "VectorUtil.h"
 #include "VertexConnectedComponents.h" // for FSizedDisjointSet
 
-namespace UE::Geometry::FaceNormalClustering
+namespace UE::Local::FaceNormalClusteringHelpers
 {
 
+using namespace UE::Geometry;
 
-void ComputeMeshPolyGroupsFromClusters(FDynamicMesh3& Mesh, TArray<TArray<int32>>& OutPolyGroups, const FClusterOptions& Options, TSet<int32>* IgnoreEdges)
+// Common clustering functionality used by multiple functions below
+void ComputeClustersHelper(FDynamicMesh3& Mesh, FSizedDisjointSet& OutPlaneGroups, const FaceNormalClustering::FClusterOptions& Options, TSet<int32>* IgnoreEdges)
 {
 	int32 NumT = Mesh.MaxTriangleID();
 
@@ -38,7 +40,7 @@ void ComputeMeshPolyGroupsFromClusters(FDynamicMesh3& Mesh, TArray<TArray<int32>
 	}
 
 	FSizedDisjointSet PlaneGroups;
-	PlaneGroups.Init(NumT, [&Mesh](int32 TID) { return Mesh.IsTriangle(TID); });
+	OutPlaneGroups.Init(NumT, [&Mesh](int32 TID) { return Mesh.IsTriangle(TID); });
 	int32 NumEdges = Mesh.MaxEdgeID();
 	Mesh.GetEdge(0);
 	FIndexPriorityQueue EdgeQueue(NumEdges);
@@ -53,13 +55,13 @@ void ComputeMeshPolyGroupsFromClusters(FDynamicMesh3& Mesh, TArray<TArray<int32>
 		float NormalAlignment = 1 - float(GroupNormalAreas[EdgeGroups.A].Normal.Dot(GroupNormalAreas[EdgeGroups.B].Normal));
 		return NormalAlignment;
 	};
-	auto GetMergeWeightForEdge = [&GetMergeWeightForGroups, &Mesh, &PlaneGroups](int32 EdgeID, FIndex2i& EdgeGroups) -> float
+	auto GetMergeWeightForEdge = [&GetMergeWeightForGroups, &Mesh, &OutPlaneGroups](int32 EdgeID, FIndex2i& EdgeGroups) -> float
 	{
 		EdgeGroups = Mesh.GetEdgeT(EdgeID);
-		EdgeGroups.A = PlaneGroups.Find(EdgeGroups.A);
-		EdgeGroups.B = PlaneGroups.Find(EdgeGroups.B);
+		EdgeGroups.A = OutPlaneGroups.Find(EdgeGroups.A);
+		EdgeGroups.B = OutPlaneGroups.Find(EdgeGroups.B);
 		if (EdgeGroups.A == EdgeGroups.B)
-		{		
+		{
 			// return a value guaranteed to be above any merge tolerance, so we ignore edges connecting already-merged groups
 			constexpr float CannotMergeValue = FMathf::MaxReal;
 			return CannotMergeValue;
@@ -116,7 +118,7 @@ void ComputeMeshPolyGroupsFromClusters(FDynamicMesh3& Mesh, TArray<TArray<int32>
 		}
 
 		// We know UpdatedWt < Threshold, so merge the groups
-		int32 Parent = PlaneGroups.Union(EdgeGroups.A, EdgeGroups.B);
+		int32 Parent = OutPlaneGroups.Union(EdgeGroups.A, EdgeGroups.B);
 		// Update the NormalArea of the new parent
 		const FNormalArea& NA_A = GroupNormalAreas[EdgeGroups.A];
 		const FNormalArea& NA_B = GroupNormalAreas[EdgeGroups.B];
@@ -130,6 +132,19 @@ void ComputeMeshPolyGroupsFromClusters(FDynamicMesh3& Mesh, TArray<TArray<int32>
 			break;
 		}
 	}
+}
+
+
+} // UE::Local::FaceNormalClusteringHelpers
+
+namespace UE::Geometry::FaceNormalClustering
+{
+
+
+void ComputeMeshPolyGroupsFromClusters(FDynamicMesh3& Mesh, TArray<TArray<int32>>& OutPolyGroups, const FClusterOptions& Options, TSet<int32>* IgnoreEdges)
+{
+	FSizedDisjointSet PlaneGroups;
+	UE::Local::FaceNormalClusteringHelpers::ComputeClustersHelper(Mesh, PlaneGroups, Options, IgnoreEdges);
 
 	// Copy triangle IDs of plane groups into the OutPolyGroups arrays
 	TArray<int32> CompactIdxToGroupID, GroupIDToCompactIdx;
@@ -150,5 +165,39 @@ void ComputeMeshPolyGroupsFromClusters(FDynamicMesh3& Mesh, TArray<TArray<int32>
 	}
 }
 
+void ComputeClusterCornerVertices(FDynamicMesh3& Mesh, TArray<int32>& OutCornerVertices, const FClusterOptions& Options, TSet<int32>* IgnoreEdges)
+{
+	FSizedDisjointSet PlaneGroups;
+	UE::Local::FaceNormalClusteringHelpers::ComputeClustersHelper(Mesh, PlaneGroups, Options, IgnoreEdges);
+
+	for (int32 VID : Mesh.VertexIndicesItr())
+	{
+		int32 GIDs[3]{ -1,-1,-1 };
+		int32 FoundGIDs = 0;
+		Mesh.EnumerateVertexTriangles(VID,
+			[&](int32 TID)
+			{
+				if (FoundGIDs > 2)
+				{
+					return;
+				}
+				int32 GID = PlaneGroups.Find(TID);
+				for (int32 TestIdx = 0; TestIdx < FoundGIDs; ++TestIdx)
+				{
+					if (GIDs[TestIdx] == GID)
+					{
+						return;
+					}
+				}
+				GIDs[FoundGIDs++] = GID;
+			}
+		);
+		if (FoundGIDs > 2)
+		{
+			OutCornerVertices.Add(VID);
+		}
+	}
+}
+
 
 } // namespace UE::Geometry::FaceNormalClustering
\ No newline at end of file
diff --git a/Engine/Source/Runtime/GeometryCore/Public/Clustering/FaceNormalClustering.h b/Engine/Source/Runtime/GeometryCore/Public/Clustering/FaceNormalClustering.h
--- a/Engine/Source/Runtime/GeometryCore/Public/Clustering/FaceNormalClustering.h
+++ b/Engine/Source/Runtime/GeometryCore/Public/Clustering/FaceNormalClustering.h
@@ -38,6 +38,15 @@ namespace FaceNormalClustering
 	 */
 	void GEOMETRYCORE_API ComputeMeshPolyGroupsFromClusters(FDynamicMesh3& Mesh, TArray<TArray<int32>>& OutPolyGroups, const FClusterOptions& Options, TSet<int32>* IgnoreEdges = nullptr);
 
+	/**
+	 * Find the vertices at the corners of face clusters -- i.e., vertices which touch faces from three or more clusters
+	 *
+	 * @param Mesh				The Mesh to operate on
+	 * @param OutCornerVertices	All vertices which touch at least three clusters
+	 * @param Options			Options controlling the tolerances used for clustering
+	 * @param IgnoreEdges		Optional set of edges that clusters should not be merged across
+	 */
+	 void GEOMETRYCORE_API ComputeClusterCornerVertices(FDynamicMesh3& Mesh, TArray<int32>& OutCornerVertices, const FClusterOptions& Options, TSet<int32>* IgnoreEdges = nullptr);
 }
 
 } // end namespace UE::Geometry
\ No newline at end of file
