Add missing nullptr check to geometry script's  SetSimpleCollisionOfDynamicMeshComponent function

diff --git a/Engine/Plugins/Experimental/GeometryScripting/Source/GeometryScriptingCore/Private/CollisionFunctions.cpp b/Engine/Plugins/Experimental/GeometryScripting/Source/GeometryScriptingCore/Private/CollisionFunctions.cpp
--- a/Engine/Plugins/Experimental/GeometryScripting/Source/GeometryScriptingCore/Private/CollisionFunctions.cpp
+++ b/Engine/Plugins/Experimental/GeometryScripting/Source/GeometryScriptingCore/Private/CollisionFunctions.cpp
@@ -483,6 +483,12 @@ void UGeometryScriptLibrary_CollisionFunctions::SetSimpleCollisionOfDynamicMeshC
 	FGeometryScriptSetSimpleCollisionOptions Options,
 	UGeometryScriptDebug* Debug)
 {
+	if (DynamicMeshComponent == nullptr)
+	{
+		UE::Geometry::AppendError(Debug, EGeometryScriptErrorType::InvalidInputs, LOCTEXT("SetSimpleCollisionOfDynamicMeshComponent_InvalidComponent", "SetSimpleCollisionOfDynamicMeshComponent: Component is Null"));
+		return;
+	}
+
 #if WITH_EDITOR
 	if (Options.bEmitTransaction && GEditor)
 	{
