Add geometry script collision functions for converting sphere coverings to/from sphere arrays, to support manually creating/adjusting negative space

diff --git a/Engine/Plugins/Experimental/GeometryScripting/Source/GeometryScriptingCore/Private/CollisionFunctions.cpp b/Engine/Plugins/Experimental/GeometryScripting/Source/GeometryScriptingCore/Private/CollisionFunctions.cpp
--- a/Engine/Plugins/Experimental/GeometryScripting/Source/GeometryScriptingCore/Private/CollisionFunctions.cpp
+++ b/Engine/Plugins/Experimental/GeometryScripting/Source/GeometryScriptingCore/Private/CollisionFunctions.cpp
@@ -1146,4 +1146,30 @@ FGeometryScriptSphereCovering UGeometryScriptLibrary_CollisionFunctions::Compute
 	return ToRet;
 }
 
+TArray<FSphere> UGeometryScriptLibrary_CollisionFunctions::Conv_GeometryScriptSphereCoveringToSphereArray(const FGeometryScriptSphereCovering& SphereCovering)
+{
+	TArray<FSphere> ToRet;
+	if (SphereCovering.Spheres.IsValid())
+	{
+		int32 NumSpheres = SphereCovering.Spheres->Num();
+		ToRet.SetNumUninitialized(NumSpheres);
+		for (int32 Idx = 0; Idx < NumSpheres; ++Idx)
+		{
+			ToRet[Idx].Center = SphereCovering.Spheres->GetCenter(Idx);
+			ToRet[Idx].W = SphereCovering.Spheres->GetRadius(Idx);
+		}
+	}
+	return ToRet;
+}
+
+FGeometryScriptSphereCovering UGeometryScriptLibrary_CollisionFunctions::Conv_SphereArrayToGeometryScriptSphereCovering(const TArray<FSphere>& Spheres)
+{
+	FGeometryScriptSphereCovering Covering;
+	Covering.Reset();
+	Covering.Spheres->AppendSpheres(Spheres);
+	return Covering;
+}
+
+
+
 #undef LOCTEXT_NAMESPACE
diff --git a/Engine/Plugins/Experimental/GeometryScripting/Source/GeometryScriptingCore/Public/GeometryScript/CollisionFunctions.h b/Engine/Plugins/Experimental/GeometryScripting/Source/GeometryScriptingCore/Public/GeometryScript/CollisionFunctions.h
--- a/Engine/Plugins/Experimental/GeometryScripting/Source/GeometryScriptingCore/Public/GeometryScript/CollisionFunctions.h
+++ b/Engine/Plugins/Experimental/GeometryScripting/Source/GeometryScriptingCore/Public/GeometryScript/CollisionFunctions.h
@@ -423,4 +423,17 @@ public:
 		UGeometryScriptDebug* Debug = nullptr
 	);
 
+	/**
+	 * @return An array of the spheres in the given Sphere Covering
+	 */
+	UFUNCTION(BlueprintPure, meta = (DisplayName = "SphereCovering To Array Of Spheres", CompactNodeTitle = "->", BlueprintAutocast), Category = "GeometryScript|Collision")
+	static TArray<FSphere> Conv_GeometryScriptSphereCoveringToSphereArray(const FGeometryScriptSphereCovering& SphereCovering);
+
+	/**
+	* @return A sphere covering containing the spheres in the given Spheres array
+	*/
+	UFUNCTION(BlueprintPure, meta = (DisplayName = "Array Of Spheres To SphereCovering", CompactNodeTitle = "->", BlueprintAutocast), Category = "GeometryScript|Collision")
+	static FGeometryScriptSphereCovering Conv_SphereArrayToGeometryScriptSphereCovering(const TArray<FSphere>& Spheres);
+
+
 };
\ No newline at end of file
diff --git a/Engine/Source/Runtime/GeometryCore/Public/CompGeom/ConvexDecomposition3.h b/Engine/Source/Runtime/GeometryCore/Public/CompGeom/ConvexDecomposition3.h
--- a/Engine/Source/Runtime/GeometryCore/Public/CompGeom/ConvexDecomposition3.h
+++ b/Engine/Source/Runtime/GeometryCore/Public/CompGeom/ConvexDecomposition3.h
@@ -117,6 +117,19 @@ public:
 		Radius.Append(Other.Radius);
 	}
 
+	void AppendSpheres(TArrayView<const FSphere> Spheres)
+	{
+		const int32 AddCount = Spheres.Num();
+		const int32 OrigCount = Position.Num();
+		Position.SetNum(AddCount + OrigCount);
+		Radius.SetNum(AddCount + OrigCount);
+		for (int32 Idx = 0; Idx < AddCount; ++Idx)
+		{
+			Position[Idx + OrigCount] = Spheres[Idx].Center;
+			Radius[Idx + OrigCount] = Spheres[Idx].W;
+		}
+	}
+
 private:
 	// Sphere centers
 	TArray<FVector3d> Position;
