Add geometry script 'approximate convex hulls with simpler collision shapes' function, to optionally replace convex hulls with spheres or boxes

diff --git a/Engine/Plugins/Experimental/GeometryScripting/Source/GeometryScriptingCore/Private/CollisionFunctions.cpp b/Engine/Plugins/Experimental/GeometryScripting/Source/GeometryScriptingCore/Private/CollisionFunctions.cpp
--- a/Engine/Plugins/Experimental/GeometryScripting/Source/GeometryScriptingCore/Private/CollisionFunctions.cpp
+++ b/Engine/Plugins/Experimental/GeometryScripting/Source/GeometryScriptingCore/Private/CollisionFunctions.cpp
@@ -20,12 +20,15 @@
 
 #include "ShapeApproximation/ShapeDetection3.h"
 #include "ShapeApproximation/MeshSimpleShapeApproximation.h"
+#include "MinVolumeSphere3.h"
+#include "MinVolumeBox3.h"
 
 #include "Clustering/FaceNormalClustering.h"
 #include "CompGeom/ConvexDecomposition3.h"
 #include "OrientedBoxTypes.h"
 #include "MeshQueries.h"
 #include "MeshAdapter.h"
+#include "SphereTypes.h"
 
 #include "Generators/MeshShapeGenerator.h"
 #include "Generators/GridBoxMeshGenerator.h"
@@ -275,6 +278,15 @@ static bool AppendConvexElemToCompactDynamicMesh(const FKConvexElem& Elem, FDyna
 }
 
 
+static double GetConvexElemVolume(const FKConvexElem& Convex)
+{
+	// Note: Not reliable to use the FKConvexElem::GetVolume function because it depends on the chaos convex being allocated, and also is not currently exported
+	TIndexMeshArrayAdapter<int32, double, FVector3d> HullMeshAdapter(&Convex.VertexData, &Convex.IndexData);
+	// Note: We take the negative volume because the hull triangles have opposite winding from ordinary meshes
+	double Volume = -TMeshQueries<TIndexMeshArrayAdapter<int32, double, FVector3d>>::GetVolumeArea(HullMeshAdapter).X;
+	return Volume;
+}
+
 }		// end namespace UELocal
 
 
@@ -885,6 +897,121 @@ void UGeometryScriptLibrary_CollisionFunctions::SimplifyConvexHulls(
 	}
 }
 
+void UGeometryScriptLibrary_CollisionFunctions::ApproximateConvexHullsWithSimplerCollisionShapes(
+	FGeometryScriptSimpleCollision& SimpleCollision,
+	const FGeometryScriptConvexHullApproximationOptions& ApproximateOptions,
+	bool& bHasApproximated,
+	UGeometryScriptDebug* Debug
+)
+{
+	bHasApproximated = false;
+
+	if (!ApproximateOptions.bFitBoxes && !ApproximateOptions.bFitSpheres)
+	{
+		// nothing to approximate
+		return;
+	}
+
+	// Helper to call ProcessPt on vertices and triangle-centroids of a convex hull mesh, stopping early if ProcessPt returns false
+	auto SampleHull = [](const FKConvexElem& Elem, TFunctionRef<bool(FVector)> ProcessPt) -> bool
+	{
+		for (FVector V : Elem.VertexData)
+		{
+			if (!ProcessPt(V))
+			{
+				return false;
+			}
+		}
+		for (int32 TriStart = 0; TriStart + 2 < Elem.IndexData.Num(); TriStart += 3)
+		{
+			FVector V0 = Elem.VertexData[Elem.IndexData[TriStart]];
+			FVector V1 = Elem.VertexData[Elem.IndexData[TriStart+1]];
+			FVector V2 = Elem.VertexData[Elem.IndexData[TriStart+2]];
+			if (!ProcessPt((V0 + V1 + V2) / 3))
+			{
+				return false;
+			}
+		}
+		return true;
+	};
+
+	auto GetVolumeDifferenceFrac = [](double HullVolume, double ApproxShapeVolume) -> double
+	{
+		return FMath::Abs(ApproxShapeVolume / HullVolume - 1.0);
+	};
+
+	TArray<FKConvexElem>& ConvexElems = SimpleCollision.AggGeom.ConvexElems;
+	for (int32 ConvexIdx = 0; ConvexIdx < ConvexElems.Num(); ++ConvexIdx)
+	{
+		const FKConvexElem& Elem = ConvexElems[ConvexIdx];
+		double HullVolume = UELocal::GetConvexElemVolume(Elem);
+
+		bool bFoundBox = false, bFoundSphere = false;
+		UE::Geometry::FSphere3d ApproxSphere;
+		FOrientedBox3d ApproxBox;
+		double SphereVolumeDiff = FMathd::MaxReal, BoxVolumeDiff = FMathd::MaxReal;
+
+		if (ApproximateOptions.bFitSpheres)
+		{
+			FMinVolumeSphere3d FitSphere;
+			if (FitSphere.Solve(Elem.VertexData.Num(), [&Elem](int32 Idx) {return Elem.VertexData[Idx];}))
+			{
+				FitSphere.GetResult(ApproxSphere);
+				double SphereVolume = ApproxSphere.Volume();
+				SphereVolumeDiff = GetVolumeDifferenceFrac(HullVolume, SphereVolume);
+				bFoundSphere = (SphereVolumeDiff < ApproximateOptions.VolumeDiffThreshold_Fraction)
+					&& SampleHull(Elem, [&](FVector V) -> bool
+					{
+						return FMath::Abs(ApproxSphere.SignedDistance(V)) < ApproximateOptions.DistanceThreshold;
+					});
+			}
+		}
+
+		if (ApproximateOptions.bFitBoxes)
+		{
+			FMinVolumeBox3d FitBox;
+			if (FitBox.Solve(Elem.VertexData.Num(), [&Elem](int32 Idx) {return Elem.VertexData[Idx];}, false, nullptr))
+			{
+				FitBox.GetResult(ApproxBox);
+				double BoxVolume = ApproxBox.Volume();
+				BoxVolumeDiff = GetVolumeDifferenceFrac(HullVolume, BoxVolume);
+				bFoundBox = (BoxVolumeDiff < ApproximateOptions.VolumeDiffThreshold_Fraction)
+					&& SampleHull(Elem, [&](FVector V) -> bool
+					{
+						return FMath::Abs(ApproxBox.SignedDistance(V)) < ApproximateOptions.DistanceThreshold;
+					});
+			}
+		}
+
+		if (bFoundSphere || bFoundBox)
+		{
+			// Add the approximating element
+			if (BoxVolumeDiff < SphereVolumeDiff)
+			{
+				FKBoxElem& ElemBox = SimpleCollision.AggGeom.BoxElems.Emplace_GetRef();
+				ElemBox.Center = ApproxBox.Center();
+				ElemBox.Rotation = (FRotator)ApproxBox.Frame.Rotation;
+				ElemBox.X = ApproxBox.Extents.X * 2;
+				ElemBox.Y = ApproxBox.Extents.Y * 2;
+				ElemBox.Z = ApproxBox.Extents.Z * 2;
+			}
+			else // sphere is better
+			{
+				FKSphereElem& ElemSphere = SimpleCollision.AggGeom.SphereElems.Emplace_GetRef();
+				ElemSphere.Center = ApproxSphere.Center;
+				ElemSphere.Radius = ApproxSphere.Radius;
+			}
+
+			// Remove the approximated convex hull
+			ConvexElems.RemoveAtSwap(ConvexIdx);
+			ConvexIdx--;
+
+			// Log that we accepted an approximation
+			bHasApproximated = true;
+		}
+	}
+}
+
 FGeometryScriptSimpleCollision UGeometryScriptLibrary_CollisionFunctions::MergeSimpleCollisionShapes(
 	const FGeometryScriptSimpleCollision& SimpleCollision,
 	const FGeometryScriptMergeSimpleCollisionOptions& MergeOptions,
@@ -1001,10 +1128,7 @@ FGeometryScriptSimpleCollision UGeometryScriptLibrary_CollisionFunctions::MergeS
 	}
 	for (const FKConvexElem& Convex : SimpleCollision.AggGeom.ConvexElems)
 	{
-		// Note: Not reliable to use the FKConvexElem::GetVolume function because it depends on the chaos convex being allocated, and also is not currently exported
-		TIndexMeshArrayAdapter<int32, double, FVector3d> HullMeshAdapter(&Convex.VertexData, &Convex.IndexData);
-		// Note: We take the negative volume because the hull triangles have opposite winding from ordinary meshes
-		double Volume = -TMeshQueries<TIndexMeshArrayAdapter<int32, double, FVector3d>>::GetVolumeArea(HullMeshAdapter).X;
+		double Volume = UELocal::GetConvexElemVolume(Convex);
 		AppendHullVertices(Convex.VertexData, Volume, &Convex);
 		if (CollisionMesh)
 		{
diff --git a/Engine/Plugins/Experimental/GeometryScripting/Source/GeometryScriptingCore/Public/GeometryScript/CollisionFunctions.h b/Engine/Plugins/Experimental/GeometryScripting/Source/GeometryScriptingCore/Public/GeometryScript/CollisionFunctions.h
--- a/Engine/Plugins/Experimental/GeometryScripting/Source/GeometryScriptingCore/Public/GeometryScript/CollisionFunctions.h
+++ b/Engine/Plugins/Experimental/GeometryScripting/Source/GeometryScriptingCore/Public/GeometryScript/CollisionFunctions.h
@@ -230,6 +230,29 @@ public:
 	int32 MinTargetFaceCount = 12;
 };
 
+USTRUCT(BlueprintType)
+struct GEOMETRYSCRIPTINGCORE_API FGeometryScriptConvexHullApproximationOptions
+{
+	GENERATED_BODY()
+public:
+
+	/** Whether to attempt to replace convex hulls with spheres */
+	UPROPERTY(EditAnywhere, BlueprintReadWrite, Category = Options)
+	bool bFitSpheres = true;
+
+	/** Whether to attempt to replace convex hulls with boxes */
+	UPROPERTY(EditAnywhere, BlueprintReadWrite, Category = Options)
+	bool bFitBoxes = true;
+
+	/** Approximating shape should be at least this close to the original shape */
+	UPROPERTY(EditAnywhere, BlueprintReadWrite, Category = Options)
+	float DistanceThreshold = 5.f;
+
+	/** Acceptable difference between approximating shape volume and convex hull volume, as a fraction of convex hull volume */
+	UPROPERTY(EditAnywhere, BlueprintReadWrite, Category = Options)
+	float VolumeDiffThreshold_Fraction = .15;
+};
+
 USTRUCT(BlueprintType)
 struct GEOMETRYSCRIPTINGCORE_API FGeometryScriptTransformCollisionOptions
 {
@@ -393,6 +416,22 @@ public:
 		UGeometryScriptDebug* Debug = nullptr
 	);
 
+	/**
+	 * Attempt to approximate any convex hulls in the given simple collision representation. Updates the passed-in Simple Collision.
+	 * Convex hulls that aren't well approximated (to tolerances set in ApproximateOptions) will remain as convex hulls.
+	 * 
+	 * @param SimpleCollision		The collision in which to attempt to approximate the convex hulls
+	 * @param ConvexSimplifyOptions	Options controlling how convex hulls are approximated
+	 * @param bHasApproximated		Indicates whether any convex hulls were replaced with simpler approximations
+	 */
+	UFUNCTION(BlueprintCallable, Category = "GeometryScript|Collision", meta = (ScriptMethod, AutoCreateRefTerm = "ApproximateOptions"))
+	static void ApproximateConvexHullsWithSimplerCollisionShapes(
+		UPARAM(ref) FGeometryScriptSimpleCollision& SimpleCollision,
+		const FGeometryScriptConvexHullApproximationOptions& ApproximateOptions,
+		bool& bHasApproximated,
+		UGeometryScriptDebug* Debug = nullptr
+	);
+
 	/**
 	 * Attempt to merge collision shapes to create a representation with fewer overall shapes.
 	 * 
diff --git a/Engine/Source/Runtime/GeometryCore/Public/OrientedBoxTypes.h b/Engine/Source/Runtime/GeometryCore/Public/OrientedBoxTypes.h
--- a/Engine/Source/Runtime/GeometryCore/Public/OrientedBoxTypes.h
+++ b/Engine/Source/Runtime/GeometryCore/Public/OrientedBoxTypes.h
@@ -207,6 +207,33 @@ struct TOrientedBox3
 		return ClampedBeyond.SquaredLength();
 	}
 
+	/**
+	 * Find signed distance to box surface
+	 * @param Point input point
+	 * @return Signed distance from point to box surface; negative if point is inside box
+	 */
+	RealType SignedDistance(TVector<RealType> Point) const
+	{
+		TVector<RealType> Local = Frame.ToFramePoint(Point);
+		TVector<RealType> VsExtents = Local.GetAbs() - Extents;
+		RealType MaxComponent = VsExtents.GetMax();
+		if (MaxComponent < 0) // Inside the box on every dimension
+		{
+			// The least-inside dimension gives the distance to the box surface
+			return MaxComponent;
+		}
+		else // Outside the box along at least one dimension
+		{
+			// clamp negative (inside) to zero
+			TVector<RealType> ClampedBeyond(
+				FMath::Max((RealType)0, VsExtents.X),
+				FMath::Max((RealType)0, VsExtents.Y),
+				FMath::Max((RealType)0, VsExtents.Z)
+			);
+			return ClampedBeyond.Length();
+		}
+	}
+
 	 /**
 	  * Find closest point on box
 	  * @param Point input point
