expose out convex decomposition negative space's voxel search 'min voxels per dimension' parameter to the setting structs, rather than hardcoding it as 128

diff --git a/Engine/Plugins/Experimental/GeometryScripting/Source/GeometryScriptingCore/Private/CollisionFunctions.cpp b/Engine/Plugins/Experimental/GeometryScripting/Source/GeometryScriptingCore/Private/CollisionFunctions.cpp
--- a/Engine/Plugins/Experimental/GeometryScripting/Source/GeometryScriptingCore/Private/CollisionFunctions.cpp
+++ b/Engine/Plugins/Experimental/GeometryScripting/Source/GeometryScriptingCore/Private/CollisionFunctions.cpp
@@ -251,6 +251,7 @@ static FNegativeSpaceSampleSettings ConvertNegativeSpaceOptions(const FComputeNe
 	NegativeSpaceSettings.SampleMethod = ConvertNegativeSpaceSampleMethodEnum(NegativeSpaceOptions.SampleMethod);
 	NegativeSpaceSettings.bRequireSearchSampleCoverage = NegativeSpaceOptions.bRequireSearchSampleCoverage;
 	NegativeSpaceSettings.bOnlyConnectedToHull = NegativeSpaceOptions.bOnlyConnectedToHull;
+	NegativeSpaceSettings.MaxVoxelsPerDim = NegativeSpaceOptions.MaxVoxelsPerDim;
 	NegativeSpaceSettings.bReferenceMeshHasNegativeWinding = false;
 	NegativeSpaceSettings.Sanitize();
 	return NegativeSpaceSettings;
diff --git a/Engine/Plugins/Experimental/GeometryScripting/Source/GeometryScriptingCore/Public/GeometryScript/CollisionFunctions.h b/Engine/Plugins/Experimental/GeometryScripting/Source/GeometryScriptingCore/Public/GeometryScript/CollisionFunctions.h
--- a/Engine/Plugins/Experimental/GeometryScripting/Source/GeometryScriptingCore/Public/GeometryScript/CollisionFunctions.h
+++ b/Engine/Plugins/Experimental/GeometryScripting/Source/GeometryScriptingCore/Public/GeometryScript/CollisionFunctions.h
@@ -144,6 +144,10 @@ public:
 	UPROPERTY(EditAnywhere, BlueprintReadWrite, Category = NegativeSpace, meta = (EditCondition = "SampleMethod == EGeometryScriptNegativeSpaceSampleMethod::VoxelSearch", EditConditionHides))
 	bool bOnlyConnectedToHull = false;
 
+	/** When performing Voxel Search, maximum number of voxels to use along each dimension */
+	UPROPERTY(EditAnywhere, BlueprintReadWrite, Category = NegativeSpace, meta = (ClampMin = 4, ClampMax = 4096))
+	int32 MaxVoxelsPerDim = 128;
+
 	/** Approximate number of spheres to consider when covering negative space */
 	UPROPERTY(EditAnywhere, BlueprintReadWrite, Category = NegativeSpace, meta = (ClampMin = 1))
 	int32 TargetNumSamples = 50;
diff --git a/Engine/Source/Runtime/GeometryCore/Private/CompGeom/ConvexDecomposition3.cpp b/Engine/Source/Runtime/GeometryCore/Private/CompGeom/ConvexDecomposition3.cpp
--- a/Engine/Source/Runtime/GeometryCore/Private/CompGeom/ConvexDecomposition3.cpp
+++ b/Engine/Source/Runtime/GeometryCore/Private/CompGeom/ConvexDecomposition3.cpp
@@ -145,7 +145,7 @@ bool FSphereCovering::AddNegativeSpace(const TFastWindingTree<FDynamicMesh3>& Sp
 		FAxisAlignedBox3d HullBox = HullAABB.GetBoundingBox();
 		HullBox.Expand(1);
 		FMarchingCubes MarchingCubes;
-		MarchingCubes.CubeSize = FMath::Clamp(SampleSettings.ReduceRadiusMargin * .5, HullBox.MaxDim() / 128.0, HullBox.MinDim() * .5);
+		MarchingCubes.CubeSize = FMath::Clamp(SampleSettings.ReduceRadiusMargin * .5, HullBox.MaxDim() / (double)SampleSettings.MaxVoxelsPerDim, HullBox.MinDim() * .5);
 		MarchingCubes.Bounds = HullBox;
 		MarchingCubes.RootMode = ERootfindingModes::Bisection;
 		MarchingCubes.RootModeSteps = 3;
diff --git a/Engine/Source/Runtime/GeometryCore/Public/CompGeom/ConvexDecomposition3.h b/Engine/Source/Runtime/GeometryCore/Public/CompGeom/ConvexDecomposition3.h
--- a/Engine/Source/Runtime/GeometryCore/Public/CompGeom/ConvexDecomposition3.h
+++ b/Engine/Source/Runtime/GeometryCore/Public/CompGeom/ConvexDecomposition3.h
@@ -73,6 +73,8 @@ struct FNegativeSpaceSampleSettings
 	bool bRequireSearchSampleCoverage = false;
 	// Whether to only consider negative space that is connected to the bounding convex hull, i.e., to ignore hollow inner pockets of negative space that cannot be reached from the outside via a larger-than-ReduceRadiusMargin-wide path
 	bool bOnlyConnectedToHull = false;
+	// Maximum number of voxels to use per dimension, when performing VoxelSearch
+	int32 MaxVoxelsPerDim = 128;
 
 	// Make sure the settings values are in valid ranges
 	void Sanitize()
@@ -81,6 +83,7 @@ struct FNegativeSpaceSampleSettings
 		MinSpacing = FMath::Max(0.0, MinSpacing);
 		ReduceRadiusMargin = FMath::Max(0.0, ReduceRadiusMargin);
 		MinRadius = FMath::Max(0.0, MinRadius);
+		MaxVoxelsPerDim = FMath::Clamp(MaxVoxelsPerDim, 4, 4096);
 	}
 };
 
