Fix reversed green and blue conditions for SetMeshConstantVertexColor

diff --git a/Engine/Plugins/Experimental/GeometryScripting/Source/GeometryScriptingCore/Private/MeshVertexColorFunctions.cpp b/Engine/Plugins/Experimental/GeometryScripting/Source/GeometryScriptingCore/Private/MeshVertexColorFunctions.cpp
--- a/Engine/Plugins/Experimental/GeometryScripting/Source/GeometryScriptingCore/Private/MeshVertexColorFunctions.cpp
+++ b/Engine/Plugins/Experimental/GeometryScripting/Source/GeometryScriptingCore/Private/MeshVertexColorFunctions.cpp
@@ -18,8 +18,8 @@ using namespace UE::Geometry;
 static FLinearColor CombineColors(FLinearColor ExistingColor, FLinearColor NewColor, FGeometryScriptColorFlags Flags)
 {
 	ExistingColor.R = (Flags.bRed) ? NewColor.R : ExistingColor.R;
-	ExistingColor.G = (Flags.bBlue) ? NewColor.G : ExistingColor.G;
-	ExistingColor.B = (Flags.bGreen) ? NewColor.B : ExistingColor.B;
+	ExistingColor.G = (Flags.bGreen) ? NewColor.G : ExistingColor.G;
+	ExistingColor.B = (Flags.bBlue) ? NewColor.B : ExistingColor.B;
 	ExistingColor.A = (Flags.bAlpha) ? NewColor.A : ExistingColor.A;
 	return ExistingColor;
 }
