Add a 3D polygon triangulation geometry script node, using the simple ear clipping algorithm, so there is an easy geometry script path to triangulation in 3D w/out having to manage converting to/from 2D

diff --git a/Engine/Plugins/Experimental/GeometryScripting/Source/GeometryScriptingCore/Private/MeshPrimitiveFunctions.cpp b/Engine/Plugins/Experimental/GeometryScripting/Source/GeometryScriptingCore/Private/MeshPrimitiveFunctions.cpp
--- a/Engine/Plugins/Experimental/GeometryScripting/Source/GeometryScriptingCore/Private/MeshPrimitiveFunctions.cpp
+++ b/Engine/Plugins/Experimental/GeometryScripting/Source/GeometryScriptingCore/Private/MeshPrimitiveFunctions.cpp
@@ -1020,6 +1020,93 @@ UDynamicMesh* UGeometryScriptLibrary_MeshPrimitiveFunctions::AppendDisc(
 }
 
 
+UDynamicMesh* UGeometryScriptLibrary_MeshPrimitiveFunctions::AppendTriangulatedPolygon3D(
+	UDynamicMesh* TargetMesh,
+	FGeometryScriptPrimitiveOptions PrimitiveOptions,
+	FTransform Transform,
+	const TArray<FVector>& PolygonVertices,
+	UGeometryScriptDebug* Debug)
+{
+	if (TargetMesh == nullptr)
+	{
+		UE::Geometry::AppendError(Debug, EGeometryScriptErrorType::InvalidInputs, LOCTEXT("AppendTriangulatedPolygon3D_InvalidInput", "AppendTriangulatedPolygon3D: TargetMesh is Null"));
+		return TargetMesh;
+	}
+	if (PolygonVertices.Num() < 3)
+	{
+		UE::Geometry::AppendError(Debug, EGeometryScriptErrorType::InvalidInputs, LOCTEXT("AppendTriangulatedPolygon3D_InvalidPolygon", "AppendTriangulatedPolygon3D: PolygonVertices array requires at least 3 positions"));
+		return TargetMesh;
+	}
+
+	FVector3d Normal, UnusedPlanePoint = FVector::ZeroVector;
+	PolygonTriangulation::ComputePolygonPlane<double>(PolygonVertices, Normal, UnusedPlanePoint);
+
+	// Create 2D basis for projected points
+	FVector BasisX, BasisY;
+	VectorUtil::MakePerpVectors(Normal, BasisX, BasisY);
+	
+	TArray<FVector2d> ProjectedPoints;
+	ProjectedPoints.SetNumUninitialized(PolygonVertices.Num());
+	for (int32 Idx = 0; Idx < PolygonVertices.Num(); ++Idx)
+	{
+		FVector V = PolygonVertices[Idx];
+		ProjectedPoints[Idx] = FVector2d(BasisX.Dot(V), BasisY.Dot(V));
+	}
+
+	// Triangulate via the projected positions
+	TArray<FIndex3i> Triangles;
+	PolygonTriangulation::TriangulateSimplePolygon(ProjectedPoints, Triangles, false);
+
+	// Find 2D bounds to find UV origin and scale
+	FAxisAlignedBox2d Bounds;
+	for (FVector2d Pt : ProjectedPoints)
+	{
+		Bounds.Contain(Pt);
+	}
+	double Width = Bounds.Width(), Height = Bounds.Height();
+	double UVScale = .01;
+	if (PrimitiveOptions.UVMode == EGeometryScriptPrimitiveUVMode::ScaleToFill)
+	{
+		UVScale = 1.0 / FMath::Max3(Width, Height, FMathd::ZeroTolerance);
+	}
+
+	// Create the polygon mesh w/ projected UVs and fixed normal
+	FDynamicMesh3 PolygonMesh;
+	PolygonMesh.EnableAttributes();
+	FDynamicMeshUVOverlay* UVOverlay = PolygonMesh.Attributes()->PrimaryUV();
+	FDynamicMeshNormalOverlay* NormalOverlay = PolygonMesh.Attributes()->PrimaryNormals();
+	FVector3f Normalf = FVector3f(Normal);
+	for (int32 Idx = 0; Idx < PolygonVertices.Num(); ++Idx)
+	{
+		PolygonMesh.AppendVertex(PolygonVertices[Idx]);
+		NormalOverlay->AppendElement(Normalf);
+		FVector2f UV = FVector2f((ProjectedPoints[Idx] - Bounds.Min) * UVScale);
+		UVOverlay->AppendElement(UV);
+	}
+	for (FIndex3i Tri : Triangles)
+	{
+		int32 TID = PolygonMesh.AppendTriangle(Tri);
+		if (TID >= 0)
+		{
+			UVOverlay->SetTriangle(TID, Tri, false);
+			NormalOverlay->SetTriangle(TID, Tri, false);
+		}
+	}
+
+	// Append to the target
+	AppendPrimitiveMesh(
+		TargetMesh,
+		PolygonMesh,
+		Transform,
+		PrimitiveOptions);
+
+	return TargetMesh;
+}
+
+
+
+
+
 UDynamicMesh* UGeometryScriptLibrary_MeshPrimitiveFunctions::AppendTriangulatedPolygon(
 	UDynamicMesh* TargetMesh,
 	FGeometryScriptPrimitiveOptions PrimitiveOptions,
diff --git a/Engine/Plugins/Experimental/GeometryScripting/Source/GeometryScriptingCore/Public/GeometryScript/MeshPrimitiveFunctions.h b/Engine/Plugins/Experimental/GeometryScripting/Source/GeometryScriptingCore/Public/GeometryScript/MeshPrimitiveFunctions.h
--- a/Engine/Plugins/Experimental/GeometryScripting/Source/GeometryScriptingCore/Public/GeometryScript/MeshPrimitiveFunctions.h
+++ b/Engine/Plugins/Experimental/GeometryScripting/Source/GeometryScriptingCore/Public/GeometryScript/MeshPrimitiveFunctions.h
@@ -491,6 +491,20 @@ public:
 		UGeometryScriptDebug* Debug = nullptr);
 
 
+	/**
+	* Appends a Triangulated Polygon (with vertices specified in 3D) to the Target Mesh.
+	* Uses Ear Clipping-based triangulation. Output vertices will always be 1:1 with input vertices.
+	* Polygon endpoint is not repeated.
+	*/
+	UFUNCTION(BlueprintCallable, Category = "GeometryScript|Primitives", meta = (ScriptMethod))
+	static UPARAM(DisplayName = "Target Mesh") UDynamicMesh*
+	AppendTriangulatedPolygon3D(
+		UDynamicMesh* TargetMesh,
+		FGeometryScriptPrimitiveOptions PrimitiveOptions,
+		FTransform Transform,
+		const TArray<FVector>& PolygonVertices3D,
+		UGeometryScriptDebug* Debug = nullptr);
+
 	/**
 	* Appends a linear staircase to the Target Mesh.
 	*/
