GeometryScripting: improve support for handling StaticMesh/DynamicMesh LOD and Material aspects

UE::AssetUtils::GetStaticMeshLODMaterialListBySection() now returns the slot name for each Material
BP function GetSectionMaterialListFromStaticMesh() now returns slot names. Fixed issue in this function where LODIndex was clamped to Render LODs, not SourceModel LODs. This function now also works for RenderData LODs. In CopyMeshToStaticMesh, for RenderData LODs, the MaterialID is set to the MaterialIndex referenced by each mesh section. So in this case in GetSectionMaterialListFromStaticMesh, the Asset Material List is returned directly. This is not really a "per Section Material List' as the function name implies, however it allows this function to provide comparable behavior between SourceModel and RenderData LODs.

Add new BP function RemapToNewMaterialIDsByMaterial() which figures out MaterialID remapping based on original and new Material lists (assuming MaterialID indexes into original Material list)
Add BP functions GetNumStaticMeshLODsOfType() and CheckStaticMeshHasAvailableLOD(), these utilities are useful for dealing w/ multi-LOD assets

diff --git a/Engine/Plugins/Experimental/GeometryScripting/Source/GeometryScriptingCore/Private/MeshAssetFunctions.cpp b/Engine/Plugins/Experimental/GeometryScripting/Source/GeometryScriptingCore/Private/MeshAssetFunctions.cpp
--- a/Engine/Plugins/Experimental/GeometryScripting/Source/GeometryScriptingCore/Private/MeshAssetFunctions.cpp
+++ b/Engine/Plugins/Experimental/GeometryScripting/Source/GeometryScriptingCore/Private/MeshAssetFunctions.cpp
@@ -406,6 +406,90 @@ UDynamicMesh*  UGeometryScriptLibrary_StaticMeshFunctions::CopyMeshToStaticMesh(
 
 
 
+bool UGeometryScriptLibrary_StaticMeshFunctions::CheckStaticMeshHasAvailableLOD(
+	UStaticMesh* FromStaticMeshAsset,
+	FGeometryScriptMeshReadLOD RequestedLOD,
+	EGeometryScriptSearchOutcomePins& Outcome,
+	UGeometryScriptDebug* Debug)
+{
+	Outcome = EGeometryScriptSearchOutcomePins::NotFound;
+	if (FromStaticMeshAsset == nullptr)
+	{
+		UE::Geometry::AppendError(Debug, EGeometryScriptErrorType::InvalidInputs, LOCTEXT("CheckStaticMeshHasAvailableLOD_InvalidInput1", "CheckStaticMeshHasAvailableLOD: FromStaticMeshAsset is Null"));
+		return false;
+	}
+
+	if (RequestedLOD.LODType == EGeometryScriptLODType::RenderData)
+	{
+		Outcome = (RequestedLOD.LODIndex >= 0 && RequestedLOD.LODIndex < FromStaticMeshAsset->GetNumLODs()) ?
+			EGeometryScriptSearchOutcomePins::Found : EGeometryScriptSearchOutcomePins::NotFound;
+
+#if !WITH_EDITOR
+		if (FromStaticMeshAsset->bAllowCPUAccess == false)
+		{
+			Outcome = EGeometryScriptSearchOutcomePins::NotFound;
+		}
+#endif
+
+		return (Outcome == EGeometryScriptSearchOutcomePins::Found);
+	}
+
+#if WITH_EDITOR
+	bool bResult = false;
+	if (RequestedLOD.LODType == EGeometryScriptLODType::HiResSourceModel)
+	{
+		bResult = FromStaticMeshAsset->IsHiResMeshDescriptionValid();
+	}
+	else if (RequestedLOD.LODType == EGeometryScriptLODType::SourceModel)
+	{
+		bResult = RequestedLOD.LODIndex >= 0
+			&& RequestedLOD.LODIndex < FromStaticMeshAsset->GetNumSourceModels()
+			&& FromStaticMeshAsset->IsSourceModelValid(RequestedLOD.LODIndex);
+	}
+	else if (RequestedLOD.LODType == EGeometryScriptLODType::MaxAvailable)
+	{
+		bResult = (FromStaticMeshAsset->GetNumSourceModels() > 0);
+	}
+	Outcome = (bResult) ? EGeometryScriptSearchOutcomePins::Found : EGeometryScriptSearchOutcomePins::NotFound;
+	return bResult;
+
+#else
+	Outcome = EGeometryScriptSearchOutcomePins::NotFound;
+	return false;
+#endif
+}
+
+
+
+int UGeometryScriptLibrary_StaticMeshFunctions::GetNumStaticMeshLODsOfType(
+	UStaticMesh* FromStaticMeshAsset,
+	EGeometryScriptLODType LODType)
+{
+	if (FromStaticMeshAsset == nullptr) return 0;
+
+#if WITH_EDITOR
+	if (LODType == EGeometryScriptLODType::RenderData)
+	{
+		return FromStaticMeshAsset->GetNumLODs();
+	}
+	if (LODType == EGeometryScriptLODType::HiResSourceModel)
+	{
+		return FromStaticMeshAsset->IsHiResMeshDescriptionValid() ? 1 : 0;
+	}
+	if (LODType == EGeometryScriptLODType::SourceModel || LODType == EGeometryScriptLODType::MaxAvailable)
+	{
+		return FromStaticMeshAsset->GetNumSourceModels();
+	}
+#else
+	if (LODType == EGeometryScriptLODType::RenderData && FromStaticMeshAsset->bAllowCPUAccess)
+	{
+		return FromStaticMeshAsset->GetNumLODs();
+	}
+#endif
+
+	return 0;
+}
+
 
 
 void UGeometryScriptLibrary_StaticMeshFunctions::GetSectionMaterialListFromStaticMesh(
@@ -413,6 +497,7 @@ void UGeometryScriptLibrary_StaticMeshFunctions::GetSectionMaterialListFromStati
 	FGeometryScriptMeshReadLOD RequestedLOD,
 	TArray<UMaterialInterface*>& MaterialList,
 	TArray<int32>& MaterialIndex,
+	TArray<FName>& MaterialSlotNames,
 	EGeometryScriptOutcomePins& Outcome,
 	UGeometryScriptDebug* Debug)
 {
@@ -423,23 +508,49 @@ void UGeometryScriptLibrary_StaticMeshFunctions::GetSectionMaterialListFromStati
 		UE::Geometry::AppendError(Debug, EGeometryScriptErrorType::InvalidInputs, LOCTEXT("GetSectionMaterialListFromStaticMesh_InvalidInput1", "GetSectionMaterialListFromStaticMesh: FromStaticMeshAsset is Null"));
 		return;
 	}
+
+	// RenderData mesh sections directly reference a Material Index, which is set as the MaterialID in CopyMeshFromStaticMesh_RenderData
+	if (RequestedLOD.LODType == EGeometryScriptLODType::RenderData)
+	{
+		MaterialList.Reset();
+		MaterialIndex.Reset();
+		MaterialSlotNames.Reset();
+		const TArray<FStaticMaterial>& AssetMaterials = FromStaticMeshAsset->GetStaticMaterials();
+		for (int32 k = 0; k < AssetMaterials.Num(); ++k)
+		{
+			MaterialList.Add(AssetMaterials[k].MaterialInterface);
+			MaterialIndex.Add(k);
+			MaterialSlotNames.Add(AssetMaterials[k].MaterialSlotName);
+		}
+
+		Outcome = EGeometryScriptOutcomePins::Success;
+		return;
+	}
+
+#if WITH_EDITOR
+
 	if (RequestedLOD.LODType != EGeometryScriptLODType::MaxAvailable && RequestedLOD.LODType != EGeometryScriptLODType::SourceModel)
 	{
 		UE::Geometry::AppendError(Debug, EGeometryScriptErrorType::InvalidInputs, LOCTEXT("GetSectionMaterialListFromStaticMesh_LODNotAvailable", "GetSectionMaterialListFromStaticMesh: Requested LOD is not available"));
 		return;
 	}
 
-	int32 UseLODIndex = FMath::Clamp(RequestedLOD.LODIndex, 0, FromStaticMeshAsset->GetNumLODs() - 1);
+	int32 UseLODIndex = FMath::Clamp(RequestedLOD.LODIndex, 0, FromStaticMeshAsset->GetNumSourceModels() - 1);
 
 	MaterialList.Reset();
 	MaterialIndex.Reset();
-	if (UE::AssetUtils::GetStaticMeshLODMaterialListBySection(FromStaticMeshAsset, UseLODIndex, MaterialList, MaterialIndex) == false)
+	MaterialSlotNames.Reset();
+	if (UE::AssetUtils::GetStaticMeshLODMaterialListBySection(FromStaticMeshAsset, UseLODIndex, MaterialList, MaterialIndex, MaterialSlotNames) == false)
 	{
 		UE::Geometry::AppendError(Debug, EGeometryScriptErrorType::InvalidInputs, LOCTEXT("GetSectionMaterialListFromStaticMesh_QueryFailed", "GetSectionMaterialListFromStaticMesh: Could not fetch Material Set from Asset"));
 		return;
 	}
 
 	Outcome = EGeometryScriptOutcomePins::Success;
+
+#else
+	UE::Geometry::AppendError(Debug, EGeometryScriptErrorType::InvalidInputs, LOCTEXT("GetSectionMaterialListFromStaticMesh_EditorOnly", "GetSectionMaterialListFromStaticMesh: Source Models are not available at Runtime"));
+#endif
 }
 
 
diff --git a/Engine/Plugins/Experimental/GeometryScripting/Source/GeometryScriptingCore/Private/MeshMaterialFunctions.cpp b/Engine/Plugins/Experimental/GeometryScripting/Source/GeometryScriptingCore/Private/MeshMaterialFunctions.cpp
--- a/Engine/Plugins/Experimental/GeometryScripting/Source/GeometryScriptingCore/Private/MeshMaterialFunctions.cpp
+++ b/Engine/Plugins/Experimental/GeometryScripting/Source/GeometryScriptingCore/Private/MeshMaterialFunctions.cpp
@@ -161,6 +161,60 @@ UDynamicMesh* UGeometryScriptLibrary_MeshMaterialFunctions::RemapMaterialIDs(
 }
 
 
+UDynamicMesh* UGeometryScriptLibrary_MeshMaterialFunctions::RemapToNewMaterialIDsByMaterial(
+	UDynamicMesh* TargetMesh,
+	const TArray<UMaterialInterface*>& FromMaterialList,
+	const TArray<UMaterialInterface*>& ToMaterialList,
+	int MissingMaterialID,
+	UGeometryScriptDebug* Debug)
+{
+	if (TargetMesh == nullptr)
+	{
+		UE::Geometry::AppendError(Debug, EGeometryScriptErrorType::InvalidInputs, LOCTEXT("RemapToNewMaterialIDsByMaterial_InvalidInput", "RemapToNewMaterialIDsByMaterial: TargetMesh is Null"));
+		return TargetMesh;
+	}
+
+	TArray<int32> ToMaterialID;
+	ToMaterialID.SetNum(FromMaterialList.Num());
+	for (int32 k = 0; k < FromMaterialList.Num(); ++k)
+	{
+		int32 FoundIdx = ToMaterialList.IndexOfByKey(FromMaterialList[k]);
+		if (FoundIdx == INDEX_NONE)
+		{
+			if (MissingMaterialID >= 0)
+			{
+				ToMaterialID[k] = MissingMaterialID;
+			}
+			else
+			{
+				UE::Geometry::AppendWarning(Debug, EGeometryScriptErrorType::InvalidInputs, LOCTEXT("RemapToNewMaterialIDsByMaterial_MaterialMissing", "RemapToNewMaterialIDsByMaterial: Material in FromMaterialList not found in ToMaterialList, skipping"));
+				ToMaterialID[k] = k;
+			}
+		}
+		else
+		{
+			ToMaterialID[k] = FoundIdx;
+		}
+	}
+
+	bool bHasMaterialIDs;
+	SimpleMeshMaterialEdit(TargetMesh, true, bHasMaterialIDs, 
+		[&](FDynamicMesh3& Mesh, FDynamicMeshMaterialAttribute& MaterialIDs) 
+	{
+		for (int32 TriangleID : Mesh.TriangleIndicesItr())
+		{
+			int32 CurID = MaterialIDs.GetValue(TriangleID);
+			int32 NewID = ToMaterialID[CurID];
+			MaterialIDs.SetValue(TriangleID, NewID);
+		}
+	});
+
+	return TargetMesh;
+}
+
+
+
+
 UDynamicMesh* UGeometryScriptLibrary_MeshMaterialFunctions::GetMaterialIDsOfTriangles( 
 	UDynamicMesh* TargetMesh, 
 	FGeometryScriptIndexList TriangleIDList,
diff --git a/Engine/Plugins/Experimental/GeometryScripting/Source/GeometryScriptingCore/Public/GeometryScript/GeometryScriptTypes.h b/Engine/Plugins/Experimental/GeometryScripting/Source/GeometryScriptingCore/Public/GeometryScript/GeometryScriptTypes.h
--- a/Engine/Plugins/Experimental/GeometryScripting/Source/GeometryScriptingCore/Public/GeometryScript/GeometryScriptTypes.h
+++ b/Engine/Plugins/Experimental/GeometryScripting/Source/GeometryScriptingCore/Public/GeometryScript/GeometryScriptTypes.h
@@ -38,12 +38,27 @@ enum class EGeometryScriptContainmentOutcomePins : uint8
 };
 
 
+/**
+ * The Type of LOD in a Mesh Asset (ie a StaticMesh Asset)
+ */
 UENUM(BlueprintType)
 enum class EGeometryScriptLODType : uint8
 {
+	/** The Maximum-quality available SourceModel LOD (HiResSourceModel if it is available, otherwise SourceModel LOD0) */
 	MaxAvailable,
+	/** The HiRes SourceModel. LOD Index is ignored. HiResSourceModel is not available at Runtime. */
 	HiResSourceModel,
+	/** 
+	 * The SourceModel mesh at a given LOD Index. Note that a StaticMesh Asset with Auto-Generated LODs may not have a valid SourceModel for every LOD Index 
+	 * SourceModel meshes are not available at Runtime.
+	 */
 	SourceModel,
+	/** 
+	 * The Render mesh at at given LOD Index. 
+	 * A StaticMesh Asset derives its RenderData LODs from it's SourceModel LODs. RenderData LODs always exist for every valid LOD Index.
+	 * However the RenderData LODs are not identical to SourceModel LODs, in particular they will be split at UV seams, Hard Normal creases, etc.
+	 * RenderData LODs in a StaticMesh Asset are only available at Runtime if the bAllowCPUAccess flag was enabled on the Asset at Cook time.
+	 */
 	RenderData
 };
 
diff --git a/Engine/Plugins/Experimental/GeometryScripting/Source/GeometryScriptingCore/Public/GeometryScript/MeshAssetFunctions.h b/Engine/Plugins/Experimental/GeometryScripting/Source/GeometryScriptingCore/Public/GeometryScript/MeshAssetFunctions.h
--- a/Engine/Plugins/Experimental/GeometryScripting/Source/GeometryScriptingCore/Public/GeometryScript/MeshAssetFunctions.h
+++ b/Engine/Plugins/Experimental/GeometryScripting/Source/GeometryScriptingCore/Public/GeometryScript/MeshAssetFunctions.h
@@ -105,8 +105,31 @@ class GEOMETRYSCRIPTINGCORE_API UGeometryScriptLibrary_StaticMeshFunctions : pub
 	GENERATED_BODY()
 public:
 	
+	/** 
+	* Check if a Static Mesh Asset has the RequestedLOD available, ie if CopyMeshFromStaticMesh will be able to
+	* succeed for the given LODType and LODIndex. 
+	*/
+	UFUNCTION(BlueprintCallable, Category = "GeometryScript|StaticMesh", meta = (ExpandEnumAsExecs = "Outcome"))
+	static bool
+	CheckStaticMeshHasAvailableLOD(
+		UStaticMesh* StaticMeshAsset, 
+		FGeometryScriptMeshReadLOD RequestedLOD,
+		EGeometryScriptSearchOutcomePins& Outcome,
+		UGeometryScriptDebug* Debug = nullptr);
+
+	/**
+	 * Determine the number of available LODs of the requested LODType in a Static Mesh Asset
+	 */
+	UFUNCTION(BlueprintCallable, Category = "GeometryScript|StaticMesh")
+	static int
+	GetNumStaticMeshLODsOfType(
+		UStaticMesh* StaticMeshAsset, 
+		EGeometryScriptLODType LODType = EGeometryScriptLODType::SourceModel);
+
 	/** 
 	* Extracts a Dynamic Mesh from a Static Mesh Asset. 
+	* 
+	* Note that the LOD Index in RequestedLOD will be silently clamped to the available number of LODs (SourceModel or RenderData)
 	*/
 	UFUNCTION(BlueprintCallable, Category = "GeometryScript|StaticMesh", meta = (ExpandEnumAsExecs = "Outcome"))
 	static UPARAM(DisplayName = "Dynamic Mesh") UDynamicMesh* 
@@ -134,6 +157,14 @@ public:
 
     /** 
 	* Extracts the Material List and corresponding Material Indices from the specified LOD of the Static Mesh Asset. 
+	* The MaterialList is sorted by Section, so if CopyMeshToStaticMesh was used to create a DynamicMesh, then the returned
+	* MaterialList here will correspond to the MaterialIDs in that DynamicMesh (as each Static Mesh Section becomes a MaterialID, in-order). 
+	* So, the returned MaterialList can be passed directly to (eg) a DynamicMeshComponent.
+	* 
+	* @param MaterialIndex this returned array is the same size as MaterialList, with each value the index of that Material in the StaticMesh Material List
+	* @param MateriaSlotNames this returned array is the same size as MaterialList, with each value the Slot Name of that Material in the StaticMesh Material List
+	*
+	* Note that the LOD Index in RequestedLOD will be silently clamped to the available number of LODs (SourceModel or RenderData)
 	*/
 	UFUNCTION(BlueprintCallable, Category = "GeometryScript|StaticMesh", meta = (ExpandEnumAsExecs = "Outcome"))
 	static void
@@ -142,6 +173,7 @@ public:
 		FGeometryScriptMeshReadLOD RequestedLOD,
 		TArray<UMaterialInterface*>& MaterialList,
 		TArray<int32>& MaterialIndex,
+		TArray<FName>& MaterialSlotNames,
 		EGeometryScriptOutcomePins& Outcome,
 		UGeometryScriptDebug* Debug = nullptr);
 
diff --git a/Engine/Plugins/Experimental/GeometryScripting/Source/GeometryScriptingCore/Public/GeometryScript/MeshMaterialFunctions.h b/Engine/Plugins/Experimental/GeometryScripting/Source/GeometryScriptingCore/Public/GeometryScript/MeshMaterialFunctions.h
--- a/Engine/Plugins/Experimental/GeometryScripting/Source/GeometryScriptingCore/Public/GeometryScript/MeshMaterialFunctions.h
+++ b/Engine/Plugins/Experimental/GeometryScripting/Source/GeometryScriptingCore/Public/GeometryScript/MeshMaterialFunctions.h
@@ -47,6 +47,25 @@ public:
 		int ToMaterialID,
 		UGeometryScriptDebug* Debug = nullptr );
 
+
+	/**
+	 * Remap the Material IDs of the TargetMesh to a new set of Material IDs based on a 'From'/Current Material List, and a New Material List.
+	 * For each triangle, the current Material is determined as FromMaterialList[MaterialID], and then the first index of this Material is found
+	 * in the ToMaterialList, and this index is used as the new MaterialID 
+	 * 
+	 * If a Material cannot be found in ToMaterialList, a warning will be printed and the MaterialID left unmodified, 
+	 * unless MissingMaterialID is set to a value >= 0, in which case MissingMaterialID will be assigned
+	 */
+	UFUNCTION(BlueprintCallable, Category = "GeometryScript|Materials", meta=(ScriptMethod))
+	static UPARAM(DisplayName = "Target Mesh") UDynamicMesh* 
+	RemapToNewMaterialIDsByMaterial( 
+		UDynamicMesh* TargetMesh, 
+		const TArray<UMaterialInterface*>& FromMaterialList,
+		const TArray<UMaterialInterface*>& ToMaterialList,
+		int MissingMaterialID = -1,
+		UGeometryScriptDebug* Debug = nullptr );
+
+
 	/**
 	 * Returns the current Material ID for a Triangle.  
 	 * If the mesh does not have Material IDs enabled or if the Triangle ID is not an element of the mesh, the value 0 will be returned and bIsValidTriangle will be false.
diff --git a/Engine/Plugins/Runtime/MeshModelingToolset/Source/ModelingComponents/Private/AssetUtils/StaticMeshMaterialUtil.cpp b/Engine/Plugins/Runtime/MeshModelingToolset/Source/ModelingComponents/Private/AssetUtils/StaticMeshMaterialUtil.cpp
--- a/Engine/Plugins/Runtime/MeshModelingToolset/Source/ModelingComponents/Private/AssetUtils/StaticMeshMaterialUtil.cpp
+++ b/Engine/Plugins/Runtime/MeshModelingToolset/Source/ModelingComponents/Private/AssetUtils/StaticMeshMaterialUtil.cpp
@@ -100,7 +100,8 @@ bool UE::AssetUtils::GetStaticMeshLODMaterialListBySection(
 	UStaticMesh* StaticMeshAsset,
 	int32 LODIndex,
 	TArray<UMaterialInterface*>& MaterialListOut,
-	TArray<int32>& MaterialIndexOut)
+	TArray<int32>& MaterialIndexOut,
+	TArray<FName>& MaterialSlotNameOut)
 {
 #if WITH_EDITOR
 	// need valid MeshDescription in Editor path
@@ -141,11 +142,13 @@ bool UE::AssetUtils::GetStaticMeshLODMaterialListBySection(
 		{
 			MaterialIndexOut.Add(UseSlotIndex);
 			MaterialListOut.Add(MaterialSetInfo.MaterialSlots[UseSlotIndex].Material);
+			MaterialSlotNameOut.Add(MaterialSetInfo.MaterialSlots[UseSlotIndex].SlotName);
 		}
 		else
 		{
 			MaterialIndexOut.Add(-1);
 			MaterialListOut.Add(nullptr);
+			MaterialSlotNameOut.Add(NAME_None);
 		}
 	}
 
diff --git a/Engine/Plugins/Runtime/MeshModelingToolset/Source/ModelingComponents/Public/AssetUtils/StaticMeshMaterialUtil.h b/Engine/Plugins/Runtime/MeshModelingToolset/Source/ModelingComponents/Public/AssetUtils/StaticMeshMaterialUtil.h
--- a/Engine/Plugins/Runtime/MeshModelingToolset/Source/ModelingComponents/Public/AssetUtils/StaticMeshMaterialUtil.h
+++ b/Engine/Plugins/Runtime/MeshModelingToolset/Source/ModelingComponents/Public/AssetUtils/StaticMeshMaterialUtil.h
@@ -59,12 +59,14 @@ namespace AssetUtils
 	 * Construct the linear per-section material list for a given LODIndex of a StaticMeshAsset
 	 * @param MaterialListOut the list of linear per-section indices into the Asset Material List
 	 * @param MaterialIndexOut the corresponding list of linear per-section indices into the Asset Material List
+	 * @param MaterialSlotNameOut the corresponding list of material slot names
 	 */
 	MODELINGCOMPONENTS_API bool GetStaticMeshLODMaterialListBySection(
 		UStaticMesh* StaticMeshAsset,
 		int32 LODIndex,
 		TArray<UMaterialInterface*>& MaterialListOut,
-		TArray<int32>& MaterialIndexOut);
+		TArray<int32>& MaterialIndexOut,
+		TArray<FName>& MaterialSlotNameOut);
 
 	/**
 	 * Generate a new unique material slot name for the given SlotMaterial and NewSlotIndex,
