fix crash on attempt to combine collision shapes with themselves + log a warning in this case, since it is unlikely to be intentional/desired

diff --git a/Engine/Plugins/Experimental/GeometryScripting/Source/GeometryScriptingCore/Private/CollisionFunctions.cpp b/Engine/Plugins/Experimental/GeometryScripting/Source/GeometryScriptingCore/Private/CollisionFunctions.cpp
--- a/Engine/Plugins/Experimental/GeometryScripting/Source/GeometryScriptingCore/Private/CollisionFunctions.cpp
+++ b/Engine/Plugins/Experimental/GeometryScripting/Source/GeometryScriptingCore/Private/CollisionFunctions.cpp
@@ -814,9 +814,17 @@ FGeometryScriptSimpleCollision UGeometryScriptLibrary_CollisionFunctions::Transf
 
 void UGeometryScriptLibrary_CollisionFunctions::CombineSimpleCollision(
 	FGeometryScriptSimpleCollision& Collision,
-	const FGeometryScriptSimpleCollision& AppendCollision
+	const FGeometryScriptSimpleCollision& AppendCollision,
+	UGeometryScriptDebug* Debug
 )
 {
+	// specially handle appending to self
+	if (&Collision.AggGeom == &AppendCollision.AggGeom)
+	{
+		// No apparent reason to combine collision shapes with themselves aside from a bug, so rather than adding duplicates of each shape, log a warning and return
+		UE::Geometry::AppendWarning(Debug, EGeometryScriptErrorType::InvalidInputs, LOCTEXT("CombineSimpleCollision Combining Collision Shapes with Self", "CombineSimpleCollision: Collision and AppendCollision must be different Simple Collision objects."));
+		return;
+	}
 	Collision.AggGeom.BoxElems.Append(AppendCollision.AggGeom.BoxElems);
 	Collision.AggGeom.ConvexElems.Append(AppendCollision.AggGeom.ConvexElems);
 	Collision.AggGeom.SphylElems.Append(AppendCollision.AggGeom.SphylElems);
diff --git a/Engine/Plugins/Experimental/GeometryScripting/Source/GeometryScriptingCore/Public/GeometryScript/CollisionFunctions.h b/Engine/Plugins/Experimental/GeometryScripting/Source/GeometryScriptingCore/Public/GeometryScript/CollisionFunctions.h
--- a/Engine/Plugins/Experimental/GeometryScripting/Source/GeometryScriptingCore/Public/GeometryScript/CollisionFunctions.h
+++ b/Engine/Plugins/Experimental/GeometryScripting/Source/GeometryScriptingCore/Public/GeometryScript/CollisionFunctions.h
@@ -413,7 +413,8 @@ public:
 	UFUNCTION(BlueprintCallable, Category = "GeometryScript|Collision", meta = (ScriptMethod))
 	static void CombineSimpleCollision(
 		UPARAM(ref) FGeometryScriptSimpleCollision& CollisionToUpdate,
-		const FGeometryScriptSimpleCollision& AppendCollision
+		const FGeometryScriptSimpleCollision& AppendCollision,
+		UGeometryScriptDebug* Debug = nullptr
 	);
 	
 	/**
