deprecate convex decomposition's negative space compute's assumption of default-flipped triangle orientations by making the triangle orientation an explicit function parameter (rather than a settings struct member) and deprecating the old function w/out this parameter.

diff --git a/Engine/Plugins/Experimental/Fracture/Source/FractureEngine/Private/FractureEngineConvex.cpp b/Engine/Plugins/Experimental/Fracture/Source/FractureEngine/Private/FractureEngineConvex.cpp
--- a/Engine/Plugins/Experimental/Fracture/Source/FractureEngine/Private/FractureEngineConvex.cpp
+++ b/Engine/Plugins/Experimental/Fracture/Source/FractureEngine/Private/FractureEngineConvex.cpp
@@ -428,7 +428,7 @@ namespace UE::FractureEngine::Convex
 
 		UE::Geometry::FDynamicMeshAABBTree3 Tree(&CombinedMesh, true);
 		UE::Geometry::TFastWindingTree<UE::Geometry::FDynamicMesh3> Winding(&Tree, true);
-		OutNegativeSpace.AddNegativeSpace(Winding, Settings);
+		OutNegativeSpace.AddNegativeSpace(Winding, Settings, true);
 
 		return bNoFailures;
 	}
diff --git a/Engine/Plugins/Experimental/GeometryScripting/Source/GeometryScriptingCore/Private/CollisionFunctions.cpp b/Engine/Plugins/Experimental/GeometryScripting/Source/GeometryScriptingCore/Private/CollisionFunctions.cpp
--- a/Engine/Plugins/Experimental/GeometryScripting/Source/GeometryScriptingCore/Private/CollisionFunctions.cpp
+++ b/Engine/Plugins/Experimental/GeometryScripting/Source/GeometryScriptingCore/Private/CollisionFunctions.cpp
@@ -252,7 +252,6 @@ static FNegativeSpaceSampleSettings ConvertNegativeSpaceOptions(const FComputeNe
 	NegativeSpaceSettings.bRequireSearchSampleCoverage = NegativeSpaceOptions.bRequireSearchSampleCoverage;
 	NegativeSpaceSettings.bOnlyConnectedToHull = NegativeSpaceOptions.bOnlyConnectedToHull;
 	NegativeSpaceSettings.MaxVoxelsPerDim = NegativeSpaceOptions.MaxVoxelsPerDim;
-	NegativeSpaceSettings.bReferenceMeshHasNegativeWinding = false;
 	NegativeSpaceSettings.Sanitize();
 	return NegativeSpaceSettings;
 }
@@ -1187,7 +1186,7 @@ FGeometryScriptSimpleCollision UGeometryScriptLibrary_CollisionFunctions::MergeS
 		MinProximityOverlapTolerance = FMath::Max(SampleSettings.ReduceRadiusMargin * .5, MinProximityOverlapTolerance);
 		FDynamicMeshAABBTree3 CollisionAABBTree(CollisionMesh.Get(), true);
 		TFastWindingTree<FDynamicMesh3> CollisionFastWinding(&CollisionAABBTree, true);
-		NegativeSpace.AddNegativeSpace(CollisionFastWinding, SampleSettings);
+		NegativeSpace.AddNegativeSpace(CollisionFastWinding, SampleSettings, false);
 	}
 	// Add any precomputed negative space, if valid/non-empty
 	if (MergeOptions.PrecomputedNegativeSpace.Spheres.IsValid() && MergeOptions.PrecomputedNegativeSpace.Spheres->Num() > 0)
@@ -1288,7 +1287,7 @@ FGeometryScriptSphereCovering UGeometryScriptLibrary_CollisionFunctions::Compute
 	
 	FNegativeSpaceSampleSettings UseSettings = UELocal::ConvertNegativeSpaceOptions(NegativeSpaceOptions);
 	ToRet.Reset();
-	ToRet.Spheres->AddNegativeSpace(*MeshBVH.FWNTree, UseSettings);
+	ToRet.Spheres->AddNegativeSpace(*MeshBVH.FWNTree, UseSettings, false);
 	return ToRet;
 }
 
diff --git a/Engine/Source/Runtime/Experimental/Chaos/Private/GeometryCollection/GeometryCollectionConvexUtility.cpp b/Engine/Source/Runtime/Experimental/Chaos/Private/GeometryCollection/GeometryCollectionConvexUtility.cpp
--- a/Engine/Source/Runtime/Experimental/Chaos/Private/GeometryCollection/GeometryCollectionConvexUtility.cpp
+++ b/Engine/Source/Runtime/Experimental/Chaos/Private/GeometryCollection/GeometryCollectionConvexUtility.cpp
@@ -2308,7 +2308,7 @@ void FGeometryCollectionConvexUtility::MergeHullsOnTransforms(FManagedArrayColle
 			}
 			UE::Geometry::FDynamicMeshAABBTree3 HullsTree(&HullsMesh, true);
 			UE::Geometry::TFastWindingTree<UE::Geometry::FDynamicMesh3> HullsWinding(&HullsTree, true);
-			LocalCovering.AddNegativeSpace(HullsWinding, *Settings.ComputeEmptySpacePerBoneSettings);
+			LocalCovering.AddNegativeSpace(HullsWinding, *Settings.ComputeEmptySpacePerBoneSettings, true);
 
 			if (Settings.EmptySpace)
 			{
diff --git a/Engine/Source/Runtime/GeometryCore/Private/CompGeom/ConvexDecomposition3.cpp b/Engine/Source/Runtime/GeometryCore/Private/CompGeom/ConvexDecomposition3.cpp
--- a/Engine/Source/Runtime/GeometryCore/Private/CompGeom/ConvexDecomposition3.cpp
+++ b/Engine/Source/Runtime/GeometryCore/Private/CompGeom/ConvexDecomposition3.cpp
@@ -24,7 +24,7 @@ namespace UE
 namespace Geometry
 {
 
-bool FSphereCovering::AddNegativeSpace(const TFastWindingTree<FDynamicMesh3>& Spatial, const FNegativeSpaceSampleSettings& SampleSettings)
+bool FSphereCovering::AddNegativeSpace(const TFastWindingTree<FDynamicMesh3>& Spatial, const FNegativeSpaceSampleSettings& SampleSettings, bool bHasFlippedTriangles)
 {
 	bool bAddedPoints = false;
 
@@ -32,7 +32,7 @@ bool FSphereCovering::AddNegativeSpace(const TFastWindingTree<FDynamicMesh3>& Sp
 
 	FAxisAlignedBox3d Bounds = Spatial.GetTree()->GetBoundingBox();
 
-	double WindingSign = SampleSettings.bReferenceMeshHasNegativeWinding ? -1 : 1;
+	double WindingSign = bHasFlippedTriangles ? -1 : 1;
 
 	if (SampleSettings.SampleMethod == FNegativeSpaceSampleSettings::ESampleMethod::Uniform)
 	{
diff --git a/Engine/Source/Runtime/GeometryCore/Public/CompGeom/ConvexDecomposition3.h b/Engine/Source/Runtime/GeometryCore/Public/CompGeom/ConvexDecomposition3.h
--- a/Engine/Source/Runtime/GeometryCore/Public/CompGeom/ConvexDecomposition3.h
+++ b/Engine/Source/Runtime/GeometryCore/Public/CompGeom/ConvexDecomposition3.h
@@ -63,9 +63,6 @@ struct FNegativeSpaceSampleSettings
 	// Ignore spheres with smaller radius than this
 	double MinRadius = 10.0;
 
-	// Whether the reference mesh used has flipped orientation, so winding < -.5 is 'inside'
-	bool bReferenceMeshHasNegativeWinding = true;
-
 	// Below options currently only apply to VoxelSearch.
 	
 	// Whether to require that all candidate sample locations identified by Voxel Search are covered by negative space samples, up to the specified Min Sample Spacing.
@@ -108,8 +105,18 @@ public:
 	}
 
 	// Add spheres covering the negative space of the given fast winding tree
+	// @param Spatial	Fast winding tree of the reference mesh, used to compute negative space
+	// @param Settings	Settings controlling how negative space is computed
+	// @param bHasFlippedTriangles	Whether the mesh referenced by Spatial has reversed triangle windings
 	// @return true if any spheres were added
-	GEOMETRYCORE_API bool AddNegativeSpace(const TFastWindingTree<FDynamicMesh3>& Spatial, const FNegativeSpaceSampleSettings& Settings);
+	GEOMETRYCORE_API bool AddNegativeSpace(const TFastWindingTree<FDynamicMesh3>& Spatial, const FNegativeSpaceSampleSettings& Settings, bool bHasFlippedTriangles);
+	
+	// Note: This version of AddNegativeSpace assumed the input had flipped triangle orientations
+	UE_DEPRECATED(5.4, "Use the version of this function with a bHasFlippedTriangles parameter")
+	bool AddNegativeSpace(const TFastWindingTree<FDynamicMesh3>& Spatial, const FNegativeSpaceSampleSettings& Settings)
+	{
+		return AddNegativeSpace(Spatial, Settings, true);
+	}
 
 	void Reset()
 	{
