Optionally allow empty results for geometry script Mesh Boolean operations

diff --git a/Engine/Plugins/Experimental/GeometryScripting/Source/GeometryScriptingCore/Private/MeshBooleanFunctions.cpp b/Engine/Plugins/Experimental/GeometryScripting/Source/GeometryScriptingCore/Private/MeshBooleanFunctions.cpp
--- a/Engine/Plugins/Experimental/GeometryScripting/Source/GeometryScriptingCore/Private/MeshBooleanFunctions.cpp
+++ b/Engine/Plugins/Experimental/GeometryScripting/Source/GeometryScriptingCore/Private/MeshBooleanFunctions.cpp
@@ -77,12 +77,12 @@ UDynamicMesh* UGeometryScriptLibrary_MeshBooleanFunctions::ApplyMeshBoolean(
 		});
 	});
 
-	// bSuccess comes back false even if we only had small errors...
-	bSuccess = (NewResultMesh.TriangleCount() > 0);
-	 
-	if (bSuccess == false)
+	// Note: ignore bSuccess, as it comes back false even if we only had small errors...
+
+	bool bFailDueToEmptyResult = (!Options.bAllowEmptyResult && NewResultMesh.TriangleCount() == 0);
+	if (bFailDueToEmptyResult)
 	{
-		UE::Geometry::AppendError(Debug, EGeometryScriptErrorType::InvalidInputs, LOCTEXT("BooleanUnion_Failed", "BooleanUnion: Boolean operation failed"));
+		UE::Geometry::AppendError(Debug, EGeometryScriptErrorType::InvalidInputs, LOCTEXT("BooleanUnion_Failed", "BooleanUnion: Boolean operation failed due to an empty result; enable Allow Empty Result if empty results should be accepted."));
 		return TargetMesh;
 	}
 
diff --git a/Engine/Plugins/Experimental/GeometryScripting/Source/GeometryScriptingCore/Public/GeometryScript/MeshBooleanFunctions.h b/Engine/Plugins/Experimental/GeometryScripting/Source/GeometryScriptingCore/Public/GeometryScript/MeshBooleanFunctions.h
--- a/Engine/Plugins/Experimental/GeometryScripting/Source/GeometryScriptingCore/Public/GeometryScript/MeshBooleanFunctions.h
+++ b/Engine/Plugins/Experimental/GeometryScripting/Source/GeometryScriptingCore/Public/GeometryScript/MeshBooleanFunctions.h
@@ -32,6 +32,10 @@ public:
 
 	UPROPERTY(EditAnywhere, BlueprintReadWrite, Category = Options)
 	float SimplifyPlanarTolerance = 0.01f;
+
+	// Whether to allow the Mesh Boolean operation to generate an empty mesh as its result
+	UPROPERTY(EditAnywhere, BlueprintReadWrite, Category = Options)
+	bool bAllowEmptyResult = false;
 };
 
 
