Make geometry script collision processing functions use the FKConvexElem's Transform

diff --git a/Engine/Plugins/Experimental/GeometryScripting/Source/GeometryScriptingCore/Private/CollisionFunctions.cpp b/Engine/Plugins/Experimental/GeometryScripting/Source/GeometryScriptingCore/Private/CollisionFunctions.cpp
--- a/Engine/Plugins/Experimental/GeometryScripting/Source/GeometryScriptingCore/Private/CollisionFunctions.cpp
+++ b/Engine/Plugins/Experimental/GeometryScripting/Source/GeometryScriptingCore/Private/CollisionFunctions.cpp
@@ -268,7 +268,7 @@ static bool AppendConvexElemToCompactDynamicMesh(const FKConvexElem& Elem, FDyna
 	int32 StartV = Mesh.MaxVertexID();
 	for (FVector V : Elem.VertexData)
 	{
-		Mesh.AppendVertex(V);
+		Mesh.AppendVertex(Elem.GetTransform().TransformPosition(V));
 	}
 	for (int32 TriStart = 0; TriStart + 2 < Elem.IndexData.Num(); TriStart += 3)
 	{
@@ -284,7 +284,7 @@ static double GetConvexElemVolume(const FKConvexElem& Convex)
 	// Note: Not reliable to use the FKConvexElem::GetVolume function because it depends on the chaos convex being allocated, and also is not currently exported
 	TIndexMeshArrayAdapter<int32, double, FVector3d> HullMeshAdapter(&Convex.VertexData, &Convex.IndexData);
 	// Note: We take the negative volume because the hull triangles have opposite winding from ordinary meshes
-	double Volume = -TMeshQueries<TIndexMeshArrayAdapter<int32, double, FVector3d>>::GetVolumeArea(HullMeshAdapter).X;
+	double Volume = -TMeshQueries<TIndexMeshArrayAdapter<int32, double, FVector3d>>::GetVolumeArea(HullMeshAdapter).X * Convex.GetTransform().GetDeterminant();
 	return Volume;
 }
 
@@ -1072,11 +1072,19 @@ FGeometryScriptSimpleCollision UGeometryScriptLibrary_CollisionFunctions::MergeS
 		}
 	};
 	auto AppendHullVertices = [&HullToShapeElem, &HullVolumes, &HullVertices, &HullVertexCounts]
-				(TArrayView<const FVector3d> Vertices, double Volume, const FKShapeElem* ShapeElem)
+				(TArrayView<const FVector3d> Vertices, double Volume, const FKShapeElem* ShapeElem, const FTransform* ShapeTransform = nullptr)
 	{
 		check(HullToShapeElem.Num() == HullVolumes.Num());
 		HullToShapeElem.Add(ShapeElem);
+		int32 HullIdxStart = HullVertices.Num();
 		HullVertices.Append(Vertices);
+		if (ShapeTransform)
+		{
+			for (int32 Idx = HullIdxStart; Idx < HullVertices.Num(); ++Idx)
+			{
+				ShapeTransform->TransformPosition(HullVertices[Idx]);
+			}
+		}
 		HullVertexCounts.Add(Vertices.Num());
 		HullVolumes.Add(Volume);
 	};
@@ -1138,7 +1146,8 @@ FGeometryScriptSimpleCollision UGeometryScriptLibrary_CollisionFunctions::MergeS
 	for (const FKConvexElem& Convex : SimpleCollision.AggGeom.ConvexElems)
 	{
 		double Volume = UELocal::GetConvexElemVolume(Convex);
-		AppendHullVertices(Convex.VertexData, Volume, &Convex);
+		FTransform ConvexTransform = Convex.GetTransform();
+		AppendHullVertices(Convex.VertexData, Volume, &Convex, &ConvexTransform);
 		if (CollisionMesh)
 		{
 			UELocal::AppendConvexElemToCompactDynamicMesh(Convex, *CollisionMesh);
diff --git a/Engine/Plugins/Experimental/GeometryScripting/Source/GeometryScriptingCore/Private/MeshPrimitiveFunctions.cpp b/Engine/Plugins/Experimental/GeometryScripting/Source/GeometryScriptingCore/Private/MeshPrimitiveFunctions.cpp
--- a/Engine/Plugins/Experimental/GeometryScripting/Source/GeometryScriptingCore/Private/MeshPrimitiveFunctions.cpp
+++ b/Engine/Plugins/Experimental/GeometryScripting/Source/GeometryScriptingCore/Private/MeshPrimitiveFunctions.cpp
@@ -1553,9 +1553,10 @@ UDynamicMesh* UGeometryScriptLibrary_MeshPrimitiveFunctions::AppendSimpleCollisi
 	for (const FKConvexElem& Convex : SimpleCollision.AggGeom.ConvexElems)
 	{
 		FDynamicMesh3 ConvexMesh;
+		FTransform ConvexTransform = Convex.GetTransform();
 		for (FVector V : Convex.VertexData)
 		{
-			ConvexMesh.AppendVertex(V);
+			ConvexMesh.AppendVertex(ConvexTransform.TransformPosition(V));
 		}
 		for (int32 TriStart = 0; TriStart + 2 < Convex.IndexData.Num(); TriStart += 3)
 		{
