Warn if Geometry Script 'Recompute Normals' is called on a mesh with no attributes or no overlay normals, and default to a per-vertex overlay in that case.

diff --git a/Engine/Plugins/Experimental/GeometryScripting/Source/GeometryScriptingCore/Private/MeshNormalsFunctions.cpp b/Engine/Plugins/Experimental/GeometryScripting/Source/GeometryScriptingCore/Private/MeshNormalsFunctions.cpp
--- a/Engine/Plugins/Experimental/GeometryScripting/Source/GeometryScriptingCore/Private/MeshNormalsFunctions.cpp
+++ b/Engine/Plugins/Experimental/GeometryScripting/Source/GeometryScriptingCore/Private/MeshNormalsFunctions.cpp
@@ -134,6 +134,11 @@ UDynamicMesh* UGeometryScriptLibrary_MeshNormalsFunctions::RecomputeNormals(
 			EditMesh.EnableAttributes();
 		}
 		FMeshNormals MeshNormals(&EditMesh);
+		if (EditMesh.Attributes()->PrimaryNormals()->ElementCount() == 0)
+		{
+			UE::Geometry::AppendWarning(Debug, EGeometryScriptErrorType::InvalidInputs, LOCTEXT("RecomputeNormals_NothingToRecompute", "RecomputeNormals: TargetMesh did not have normals to recompute; falling back to per-vertex normals. Consider using 'Set Mesh To Per Vertex Normals' or 'Compute Split Normals' instead."));
+			EditMesh.Attributes()->PrimaryNormals()->CreateFromPredicate([](int, int, int)->bool {return true;}, 0.0f);
+		}
 		MeshNormals.RecomputeOverlayNormals(EditMesh.Attributes()->PrimaryNormals(), CalculateOptions.bAreaWeighted, CalculateOptions.bAngleWeighted);
 		MeshNormals.CopyToOverlay(EditMesh.Attributes()->PrimaryNormals(), false);
 
