Better display names to distingish the two geometry script functions to set per vertex normals -- clarifying which one computes them and which one takes them as input

diff --git a/Engine/Plugins/Experimental/GeometryScripting/Source/GeometryScriptingCore/Public/GeometryScript/MeshNormalsFunctions.h b/Engine/Plugins/Experimental/GeometryScripting/Source/GeometryScriptingCore/Public/GeometryScript/MeshNormalsFunctions.h
--- a/Engine/Plugins/Experimental/GeometryScripting/Source/GeometryScriptingCore/Public/GeometryScript/MeshNormalsFunctions.h
+++ b/Engine/Plugins/Experimental/GeometryScripting/Source/GeometryScriptingCore/Public/GeometryScript/MeshNormalsFunctions.h
@@ -97,7 +97,7 @@ public:
 	 * Recompute the normals of TargetMesh by averaging the triangle/face normals around each vertex, using combined area and angle weighting.
 	 * Each vertex will have a single normal, ie there will be no hard edges.
 	 */
-	UFUNCTION(BlueprintCallable, Category = "GeometryScript|Normals", meta=(ScriptMethod, DisplayName="Set Mesh To Per Vertex Normals"))
+	UFUNCTION(BlueprintCallable, Category = "GeometryScript|Normals", meta=(ScriptMethod, DisplayName="Set Mesh To Per Vertex Normals (Computed)"))
 	static UPARAM(DisplayName = "Target Mesh") UDynamicMesh* 
 	SetPerVertexNormals( 
 		UDynamicMesh* TargetMesh, 
@@ -173,7 +173,7 @@ public:
 	 * Set all normals in the TargetMesh Normals Overlay to the specified per-vertex normals
 	 * @param VertexNormalList per-vertex normals. Size must be equal to the MaxVertexID of TargetMesh  (ie non-compact TargetMesh is supported)
 	 */
-	UFUNCTION(BlueprintCallable, Category = "GeometryScript|Normals", meta=(ScriptMethod))
+	UFUNCTION(BlueprintCallable, Category = "GeometryScript|Normals", meta=(ScriptMethod, DisplayName="Set Mesh To Per Vertex Normals (From List)"))
 	static UPARAM(DisplayName = "Target Mesh") UDynamicMesh*
 	SetMeshPerVertexNormals(
 		UDynamicMesh* TargetMesh,
