GeometryScript: handle null DynamicMeshComponent in GeneratedDynamicMeshActor, not clear how this can happen but it will lead to a crash otherwise
#rnx

diff --git a/Engine/Plugins/Experimental/GeometryScripting/Source/GeometryScriptingEditor/Private/GeometryActors/GeneratedDynamicMeshActor.cpp b/Engine/Plugins/Experimental/GeometryScripting/Source/GeometryScriptingEditor/Private/GeometryActors/GeneratedDynamicMeshActor.cpp
--- a/Engine/Plugins/Experimental/GeometryScripting/Source/GeometryScriptingEditor/Private/GeometryActors/GeneratedDynamicMeshActor.cpp
+++ b/Engine/Plugins/Experimental/GeometryScripting/Source/GeometryScriptingEditor/Private/GeometryActors/GeneratedDynamicMeshActor.cpp
@@ -174,18 +174,24 @@ void AGeneratedDynamicMeshActor::ExecuteRebuildGeneratedMeshIfPending()
 		return;
 	}
 
+	UDynamicMeshComponent* Component = GetDynamicMeshComponent();
+	if (Component == nullptr)
+	{
+		return;
+	}
+
 	// Automatically defer collision updates during generated mesh rebuild. If we do not do this, then
 	// each mesh change will result in collision being rebuilt, which is very expensive !!
 	bool bEnabledDeferredCollision = false;
-	if (DynamicMeshComponent->bDeferCollisionUpdates == false)
+	if (Component->bDeferCollisionUpdates == false)
 	{
-		DynamicMeshComponent->SetDeferredCollisionUpdatesEnabled(true, false);
+		Component->SetDeferredCollisionUpdatesEnabled(true, false);
 		bEnabledDeferredCollision = true;
 	}
 
-	if (bResetOnRebuild && DynamicMeshComponent && DynamicMeshComponent->GetDynamicMesh())
+	if (bResetOnRebuild && Component->GetDynamicMesh())
 	{
-		DynamicMeshComponent->GetDynamicMesh()->Reset();
+		Component->GetDynamicMesh()->Reset();
 	}
 
 	FEditorScriptExecutionGuard Guard;
@@ -196,19 +202,19 @@ void AGeneratedDynamicMeshActor::ExecuteRebuildGeneratedMeshIfPending()
 		Progress.MakeDialogDelayed(this->DialogDelay, true);
 		ActiveSlowTask = &Progress;
 		CurProgressAccumSteps = 0;
-		OnRebuildGeneratedMesh(DynamicMeshComponent->GetDynamicMesh());
+		OnRebuildGeneratedMesh(Component->GetDynamicMesh());
 		ActiveSlowTask = nullptr;
 	}
 	else
 	{
-		OnRebuildGeneratedMesh(DynamicMeshComponent->GetDynamicMesh());
+		OnRebuildGeneratedMesh(Component->GetDynamicMesh());
 	}
 
 	bGeneratedMeshRebuildPending = false;
 
 	if (bEnabledDeferredCollision)
 	{
-		DynamicMeshComponent->SetDeferredCollisionUpdatesEnabled(false, true);
+		Component->SetDeferredCollisionUpdatesEnabled(false, true);
 	}
 }
 
