ModelingTools: BakeVertex per channel bakes will only write over the specified channels.
Exposed new One/Zero evaluators for per channel vertex bakes to recreate the previous behavior.
This behavior & functionality is mirrored to the GeometryScript BakeVertex node.

diff --git a/Engine/Plugins/Editor/ModelingToolsEditorMode/Source/ModelingToolsEditorMode/Private/DetailsCustomizations/BakeMeshAttributeToolCustomizations.cpp b/Engine/Plugins/Editor/ModelingToolsEditorMode/Source/ModelingToolsEditorMode/Private/DetailsCustomizations/BakeMeshAttributeToolCustomizations.cpp
--- a/Engine/Plugins/Editor/ModelingToolsEditorMode/Source/ModelingToolsEditorMode/Private/DetailsCustomizations/BakeMeshAttributeToolCustomizations.cpp
+++ b/Engine/Plugins/Editor/ModelingToolsEditorMode/Source/ModelingToolsEditorMode/Private/DetailsCustomizations/BakeMeshAttributeToolCustomizations.cpp
@@ -64,7 +64,9 @@ namespace BakeCustomizationLocals
 			{
 				EBakeMapType::VertexColor,
 				EBakeMapType::MaterialID,
-				EBakeMapType::PolyGroupID
+				EBakeMapType::PolyGroupID,
+				EBakeMapType::One,
+				EBakeMapType::Zero
 			}
 		};
 
diff --git a/Engine/Plugins/Experimental/GeometryScripting/Source/GeometryScriptingCore/Private/MeshBakeFunctions.cpp b/Engine/Plugins/Experimental/GeometryScripting/Source/GeometryScriptingCore/Private/MeshBakeFunctions.cpp
--- a/Engine/Plugins/Experimental/GeometryScripting/Source/GeometryScriptingCore/Private/MeshBakeFunctions.cpp
+++ b/Engine/Plugins/Experimental/GeometryScripting/Source/GeometryScriptingCore/Private/MeshBakeFunctions.cpp
@@ -171,6 +171,11 @@ namespace GeometryScriptBakeLocals
 			}
 			break;
 		}
+		case EMeshMapEvaluatorType::Constant:
+		{
+			TexType = FTexture2DBuilder::ETextureType::ColorLinear;
+			break;
+		}
 		}
 		return TexType;
 	}
@@ -469,6 +474,13 @@ namespace GeometryScriptBakeLocals
 			Result = PropertyEval;
 			break;
 		}
+		case EGeometryScriptBakeTypes::Constant:
+		{
+			FGeometryScriptBakeType_Constant* ConstantOptions = static_cast<FGeometryScriptBakeType_Constant*>(Options.Options.Get());
+			TSharedPtr<FMeshConstantMapEvaluator> ConstantEval = MakeShared<FMeshConstantMapEvaluator>(ConstantOptions->Value);
+			Result = ConstantEval;
+			break;
+		}
 		default:
 			break;
 		}
@@ -548,9 +560,15 @@ namespace GeometryScriptBakeLocals
 			Baker.SetCorrespondenceStrategy(FMeshBaseBaker::ECorrespondenceStrategy::Identity);
 		}
 
-		auto IsValidBakeType = [](EGeometryScriptBakeTypes, TArray<FGeometryScriptDebugMessage>*)
+		auto IsValidBakeType = [&BakeTexturePrefix](EGeometryScriptBakeTypes BakeType, TArray<FGeometryScriptDebugMessage>* Debug)
 		{
-			return true;
+			const bool bIsValid = BakeType != EGeometryScriptBakeTypes::None;
+			if (!bIsValid)
+			{
+				const FText BakeTypeName = FText::FromName(StaticEnum<EGeometryScriptBakeTypes>()->GetNameByIndex(static_cast<int32>(BakeType)));
+				UE::Geometry::AppendError(Debug, EGeometryScriptErrorType::InvalidInputs, FText::Format(LOCTEXT("BakeTexture_InvalidBakeType", "{0}: {1} bake type is not a supported RGBA evaluator."), BakeTexturePrefix, BakeTypeName));
+			}
+			return bIsValid;
 		};
 
 		FEvaluatorState EvalState;
@@ -688,10 +706,13 @@ namespace GeometryScriptBakeLocals
 				return bCanShare;
 			}, 0.0f);
 
-		if (bIsBakeToSelf)
+		if (bIsBakeToSelf || BakeTypes.OutputMode == EGeometryScriptBakeOutputMode::PerChannel)
 		{
-			// Copy source vertex colors onto new color overlay topology for identity bakes.
+			// Copy source vertex colors onto new color overlay topology for identity bakes
 			// This is necessary when sampling vertex color data.
+			//
+			// Also copy source vertex colors for PerChannel bakes so users can opt to bake
+			// a single channel of an existing vertex color dataset.
 			const FDynamicMeshColorOverlay* SourceColorOverlay = SourceMeshToUse->HasAttributes() ? SourceMeshToUse->Attributes()->PrimaryColors() : nullptr;
 			FDynamicMeshColorOverlay* TargetColorOverlay = TargetMeshRef.Attributes()->PrimaryColors(); 
 			if (SourceColorOverlay)
@@ -727,9 +748,15 @@ namespace GeometryScriptBakeLocals
 
 		if (BakeTypes.OutputMode == EGeometryScriptBakeOutputMode::RGBA)
 		{
-			auto IsValidBakeType = [](EGeometryScriptBakeTypes, TArray<FGeometryScriptDebugMessage>*)
+			auto IsValidBakeType = [&BakeVertexPrefix](EGeometryScriptBakeTypes BakeType, TArray<FGeometryScriptDebugMessage>* Debug)
 			{
-				return true;
+				const bool bIsValid = BakeType != EGeometryScriptBakeTypes::None;
+				if (!bIsValid)
+				{
+					const FText BakeTypeName = FText::FromName(StaticEnum<EGeometryScriptBakeTypes>()->GetNameByIndex(static_cast<int32>(BakeType)));
+					UE::Geometry::AppendError(Debug, EGeometryScriptErrorType::InvalidInputs, FText::Format(LOCTEXT("BakeVertex_InvalidBakeType", "{0}: {1} bake type is not a supported RGBA evaluator."), BakeVertexPrefix, BakeTypeName));
+				}
+				return bIsValid;
 			};
 			
 			Baker.BakeMode = FMeshVertexBaker::EBakeMode::RGBA;
@@ -741,7 +768,8 @@ namespace GeometryScriptBakeLocals
 
 			auto IsValidBakeType = [&BakeVertexPrefix](EGeometryScriptBakeTypes BakeType, TArray<FGeometryScriptDebugMessage>* Debug)
 			{
-				const bool bIsValid = (BakeType == EGeometryScriptBakeTypes::AmbientOcclusion || BakeType == EGeometryScriptBakeTypes::Curvature);
+				// The None bake type is permitted for PerChannel bakes since it allows users to specify a channel that should not be overwritten.
+				const bool bIsValid = (BakeType == EGeometryScriptBakeTypes::AmbientOcclusion || BakeType == EGeometryScriptBakeTypes::Curvature || BakeType == EGeometryScriptBakeTypes::Constant || BakeType == EGeometryScriptBakeTypes::None);
 				if (!bIsValid)
 				{
 					const FText BakeTypeName = FText::FromName(StaticEnum<EGeometryScriptBakeTypes>()->GetNameByIndex(static_cast<int32>(BakeType)));
@@ -774,21 +802,67 @@ namespace GeometryScriptBakeLocals
 		return MoveTemp(Result);
 	}
 
-	bool ApplyVertexBakeToMesh(const FMeshVertexBaker* Baker, UDynamicMesh* Mesh)
+	bool ApplyVertexBakeToMesh(const FMeshVertexBaker* Baker, const FGeometryScriptBakeOutputType& BakeTypes, UDynamicMesh* Mesh)
 	{
 		if (!Baker || !Mesh)
 		{
 			return false;
 		}
 
-		FDynamicMesh3& MeshRef = Mesh->GetMeshRef();		
-		const TImageBuilder<FVector4f>* ImageResult = Baker->GetBakeResult();
-		const int NumColors = MeshRef.Attributes()->PrimaryColors()->ElementCount();
-		check(NumColors == ImageResult->GetDimensions().GetWidth());
-		for (int Idx = 0; Idx < NumColors; ++Idx)
+		if (BakeTypes.OutputMode == EGeometryScriptBakeOutputMode::PerChannel)
 		{
-			const FVector4f& Pixel = ImageResult->GetPixel(Idx);
-			MeshRef.Attributes()->PrimaryColors()->SetElement(Idx, Pixel);
+			// Precompute scale vectors for source and image pixel data to merge
+			// the data according to the populated channels.
+			FVector4f SrcScale = FVector4f::Zero();
+			FVector4f ImgScale = FVector4f::Zero();
+
+			const bool bOutputR = BakeTypes.R.BakeType != EGeometryScriptBakeTypes::None;
+			const bool bOutputG = BakeTypes.G.BakeType != EGeometryScriptBakeTypes::None;
+			const bool bOutputB = BakeTypes.B.BakeType != EGeometryScriptBakeTypes::None;
+			const bool bOutputA = BakeTypes.A.BakeType != EGeometryScriptBakeTypes::None;
+
+			SrcScale[0] = static_cast<float>(!bOutputR);
+			ImgScale[0] = static_cast<float>(bOutputR);
+
+			SrcScale[1] = static_cast<float>(!bOutputG);
+			ImgScale[1] = static_cast<float>(bOutputG);
+
+			SrcScale[2] = static_cast<float>(!bOutputB);
+			ImgScale[2] = static_cast<float>(bOutputB);
+
+			SrcScale[3] = static_cast<float>(!bOutputA);
+			ImgScale[3] = static_cast<float>(bOutputA);
+
+			FDynamicMesh3& MeshRef = Mesh->GetMeshRef();
+			const int NumColors = MeshRef.Attributes()->PrimaryColors()->ElementCount();
+			const TImageBuilder<FVector4f>* ImageResult = Baker->GetBakeResult();
+			check(NumColors == ImageResult->GetDimensions().GetWidth());
+			for (int Idx = 0; Idx < NumColors; ++Idx)
+			{
+				if (const FDynamicMeshColorOverlay* ColorOverlay = MeshRef.Attributes()->PrimaryColors())
+				{
+					FVector4f Pixel;
+					ColorOverlay->GetElement(Idx, Pixel);
+					Pixel *= SrcScale;
+
+					// Swizzle the ImageResult pixels based on the requested channels.
+					const FVector4f& ImagePixel = ImageResult->GetPixel(Idx);
+					Pixel += ImagePixel * ImgScale;
+					MeshRef.Attributes()->PrimaryColors()->SetElement(Idx, Pixel);
+				}
+			}
+		}
+		else
+		{
+			FDynamicMesh3& MeshRef = Mesh->GetMeshRef();		
+			const TImageBuilder<FVector4f>* ImageResult = Baker->GetBakeResult();
+			const int NumColors = MeshRef.Attributes()->PrimaryColors()->ElementCount();
+			check(NumColors == ImageResult->GetDimensions().GetWidth());
+			for (int Idx = 0; Idx < NumColors; ++Idx)
+			{
+				const FVector4f& Pixel = ImageResult->GetPixel(Idx);
+				MeshRef.Attributes()->PrimaryColors()->SetElement(Idx, Pixel);
+			}
 		}
 		return true;
 	}
@@ -995,6 +1069,16 @@ FGeometryScriptBakeTypeOptions UGeometryScriptLibrary_MeshBakeFunctions::MakeBak
 	return Output;
 }
 
+FGeometryScriptBakeTypeOptions UGeometryScriptLibrary_MeshBakeFunctions::MakeBakeTypeConstant(
+	float Value)
+{
+	FGeometryScriptBakeTypeOptions Output;
+	Output.BakeType = EGeometryScriptBakeTypes::Constant;
+	const TSharedPtr<FGeometryScriptBakeType_Constant> ConstantOptions = MakeShared<FGeometryScriptBakeType_Constant>();
+	ConstantOptions->Value = Value;
+	Output.Options = ConstantOptions;
+	return Output;
+}
 
 TArray<UTexture2D*> UGeometryScriptLibrary_MeshBakeFunctions::BakeTexture(
 	UDynamicMesh* TargetMesh,
@@ -1047,7 +1131,7 @@ UDynamicMesh* UGeometryScriptLibrary_MeshBakeFunctions::BakeVertex(
 		Debug ? &Debug->Messages : nullptr);
 
 	// Extract the vertex bake data and apply to target mesh.
-	GeometryScriptBakeLocals::ApplyVertexBakeToMesh(Baker.Get(), TargetMesh);
+	GeometryScriptBakeLocals::ApplyVertexBakeToMesh(Baker.Get(), BakeTypes, TargetMesh);
 
 	return TargetMesh;
 }
diff --git a/Engine/Plugins/Experimental/GeometryScripting/Source/GeometryScriptingCore/Public/GeometryScript/MeshBakeFunctions.h b/Engine/Plugins/Experimental/GeometryScripting/Source/GeometryScriptingCore/Public/GeometryScript/MeshBakeFunctions.h
--- a/Engine/Plugins/Experimental/GeometryScripting/Source/GeometryScriptingCore/Public/GeometryScript/MeshBakeFunctions.h
+++ b/Engine/Plugins/Experimental/GeometryScripting/Source/GeometryScriptingCore/Public/GeometryScript/MeshBakeFunctions.h
@@ -51,6 +51,7 @@ enum class EGeometryScriptBakeFilteringType : uint8
 UENUM(BlueprintType)
 enum class EGeometryScriptBakeTypes : uint8
 {
+	None,
 	/* Normals in tangent space */
 	TangentSpaceNormal     UMETA(DisplayName = "Tangent Normal"),
 	/* Interpolated normals in object space */
@@ -73,6 +74,8 @@ enum class EGeometryScriptBakeTypes : uint8
 	VertexColor            ,
 	/* Material IDs as unique colors */
 	MaterialID             UMETA(DisplayName = "Material ID"),
+	/* Constant value */
+	Constant
 };
 
 UENUM(BlueprintType)
@@ -183,6 +186,12 @@ struct GEOMETRYSCRIPTINGCORE_API FGeometryScriptBakeType_MultiTexture : public F
 	int SourceUVLayer = 0;
 };
 
+struct GEOMETRYSCRIPTINGCORE_API FGeometryScriptBakeType_Constant : public FGeometryScriptBakeTypes
+{
+	/** Constant value to set */
+	float Value = 0.0f;
+};
+
 /**
  * Opaque struct for storing bake type options.
  */
@@ -193,7 +202,7 @@ struct GEOMETRYSCRIPTINGCORE_API FGeometryScriptBakeTypeOptions
 
 	/** The bake output type to generate */
 	UPROPERTY(BlueprintReadOnly, Category = Type)
-	EGeometryScriptBakeTypes BakeType = EGeometryScriptBakeTypes::TangentSpaceNormal;
+	EGeometryScriptBakeTypes BakeType = EGeometryScriptBakeTypes::None;
 
 	TSharedPtr<FGeometryScriptBakeTypes> Options;
 };
@@ -525,6 +534,10 @@ public:
 
 	UFUNCTION(BlueprintPure, Category = "GeometryScript|Bake")
 	static UPARAM(DisplayName="Bake Type Out") FGeometryScriptBakeTypeOptions MakeBakeTypeMaterialID();
+
+	UFUNCTION(BlueprintPure, Category = "GeometryScript|Bake")
+	static UPARAM(DisplayName="Bake Type Out") FGeometryScriptBakeTypeOptions MakeBakeTypeConstant(
+		float Value = 0.0f);
 	
 	UFUNCTION(BlueprintCallable, Category = "GeometryScript|Bake")
 	static UPARAM(DisplayName="Textures Out") TArray<UTexture2D*> BakeTexture(
diff --git a/Engine/Plugins/Experimental/MeshModelingToolsetExp/Source/MeshModelingToolsExp/Private/BakeMeshAttributeVertexTool.cpp b/Engine/Plugins/Experimental/MeshModelingToolsetExp/Source/MeshModelingToolsExp/Private/BakeMeshAttributeVertexTool.cpp
--- a/Engine/Plugins/Experimental/MeshModelingToolsetExp/Source/MeshModelingToolsExp/Private/BakeMeshAttributeVertexTool.cpp
+++ b/Engine/Plugins/Experimental/MeshModelingToolsetExp/Source/MeshModelingToolsExp/Private/BakeMeshAttributeVertexTool.cpp
@@ -134,6 +134,14 @@ public:
 					Baker->ChannelEvaluators[ChannelIdx] = CurvatureEval;
 					break;
 				}
+				case EBakeMapType::One:
+				case EBakeMapType::Zero:
+				{
+					const float Value = BakeSettings.OutputTypePerChannel[ChannelIdx] == EBakeMapType::One ? 1.0f : 0.0f;
+					TSharedPtr<FMeshConstantMapEvaluator> ConstantEval = MakeShared<FMeshConstantMapEvaluator>(Value);
+					Baker->ChannelEvaluators[ChannelIdx] = ConstantEval;
+					break;
+				}
 				default:
 				case EBakeMapType::None:
 				{
@@ -576,6 +584,7 @@ void UBakeMeshAttributeVertexTool::UpdateVisualization()
 	}
 }
 
+/** Regenerates the VertexColorOverlay topology of the PreviewMesh per our Settings */
 void UBakeMeshAttributeVertexTool::UpdateColorTopology()
 {
 	// Update PreviewMesh color topology
@@ -606,8 +615,21 @@ void UBakeMeshAttributeVertexTool::UpdateColorTopology()
 				}
 				return bCanShare;
 			}, 0.0f);
+	});
+
+	// Copy source vertex colors onto new color overlay topology.
+	UpdateSourceVertexColors();
+	
+	NumColorElements = PreviewMesh->GetMesh()->Attributes()->PrimaryColors()->ElementCount();
 
-		// Copy source vertex colors onto new color overlay topology.
+	bColorTopologyValid = true;
+}
+
+/**	Copies the vertex colors from the TargetMesh to the PreviewMesh */
+void UBakeMeshAttributeVertexTool::UpdateSourceVertexColors()
+{
+	PreviewMesh->EditMesh([this](FDynamicMesh3& Mesh)
+	{
 		const FDynamicMeshColorOverlay* TargetColorOverlay = TargetMesh->HasAttributes() ? TargetMesh->Attributes()->PrimaryColors() : nullptr;
 		FDynamicMeshColorOverlay* PreviewColorOverlay = Mesh.Attributes()->PrimaryColors(); 
 		if (TargetColorOverlay)
@@ -623,11 +645,9 @@ void UBakeMeshAttributeVertexTool::UpdateColorTopology()
 			}
 		}
 	});
-	NumColorElements = PreviewMesh->GetMesh()->Attributes()->PrimaryColors()->ElementCount();
-
-	bColorTopologyValid = true;
 }
 
+
 void UBakeMeshAttributeVertexTool::UpdateResult()
 {
 	if (static_cast<bool>(OpState & EBakeOpState::EvaluateDetailMesh))
@@ -696,7 +716,7 @@ void UBakeMeshAttributeVertexTool::UpdateResult()
 			OpState |= UpdateResult_MeshProperty(Dimensions);
 			// Force copy the original vertex colors to our PreviewMesh so that
 			// the baker samples the source vertex colors for identity bakes.
-			UpdateColorTopology();
+			UpdateSourceVertexColors();
 			break;
 		case EBakeMapType::Texture:
 			OpState |= UpdateResult_Texture2DImage(Dimensions, DetailMesh.Get());
@@ -721,6 +741,11 @@ void UBakeMeshAttributeVertexTool::UpdateResult()
 		{
 			OpState |= UpdateResult_Curvature(Dimensions);
 		}
+		
+		// Always force copy the original vertex colors to our PreviewMesh for
+		// PerChannel bakes so that channels that are not targeted persist
+		// through the bake.
+		UpdateSourceVertexColors();
 	}
 
 	// Early exit if op input parameters are invalid.
@@ -743,12 +768,45 @@ void UBakeMeshAttributeVertexTool::OnResultUpdated(const TUniquePtr<FMeshVertexB
 
 	PreviewMesh->DeferredEditMesh([this, &ImageResult](FDynamicMesh3& Mesh)
 	{
-		const int NumColors = Mesh.Attributes()->PrimaryColors()->ElementCount();
-		check(NumColors == ImageResult->GetDimensions().GetWidth());
-		for (int Idx = 0; Idx < NumColors; ++Idx)
+		if (CachedBakeSettings.OutputMode == EBakeVertexOutput::PerChannel)
+		{
+			// Precompute scale vectors for source and image pixel data to merge
+			// the data according to the populated channels.
+			FVector4f SrcScale = FVector4f::Zero();
+			FVector4f ImgScale = FVector4f::Zero();
+			for (int ChannelIdx = 0; ChannelIdx < 4; ++ChannelIdx)
+			{
+				const bool bOutputChannel = CachedBakeSettings.OutputTypePerChannel[ChannelIdx] != EBakeMapType::None;
+				SrcScale[ChannelIdx] = static_cast<float>(!bOutputChannel);
+				ImgScale[ChannelIdx] = static_cast<float>(bOutputChannel);
+			}
+			
+			const int NumColors = Mesh.Attributes()->PrimaryColors()->ElementCount();
+			check(NumColors == ImageResult->GetDimensions().GetWidth());
+			for (int Idx = 0; Idx < NumColors; ++Idx)
+			{
+				if (const FDynamicMeshColorOverlay* ColorOverlay = Mesh.Attributes()->PrimaryColors())
+				{
+					FVector4f Pixel;
+					ColorOverlay->GetElement(Idx, Pixel);
+					Pixel *= SrcScale;
+
+					// Merge the ImageResult pixels with the source vertex colors based on the requested channels.
+					const FVector4f& ImagePixel = ImageResult->GetPixel(Idx);
+					Pixel += ImagePixel * ImgScale;
+					Mesh.Attributes()->PrimaryColors()->SetElement(Idx, Pixel);
+				}
+			}
+		}
+		else //if (Settings->OutputMode == EBakeVertexOutput::RGBA)
 		{
-			const FVector4f& Pixel = ImageResult->GetPixel(Idx);
-			Mesh.Attributes()->PrimaryColors()->SetElement(Idx, Pixel);
+			const int NumColors = Mesh.Attributes()->PrimaryColors()->ElementCount();
+			check(NumColors == ImageResult->GetDimensions().GetWidth());
+			for (int Idx = 0; Idx < NumColors; ++Idx)
+			{
+				const FVector4f& Pixel = ImageResult->GetPixel(Idx);
+				Mesh.Attributes()->PrimaryColors()->SetElement(Idx, Pixel);
+			}
 		}
 	}, false);
 	PreviewMesh->NotifyDeferredEditCompleted(UPreviewMesh::ERenderUpdateMode::FastUpdate, EMeshRenderAttributeFlags::VertexColors, false);
diff --git a/Engine/Plugins/Experimental/MeshModelingToolsetExp/Source/MeshModelingToolsExp/Public/BakeMeshAttributeTool.h b/Engine/Plugins/Experimental/MeshModelingToolsetExp/Source/MeshModelingToolsExp/Public/BakeMeshAttributeTool.h
--- a/Engine/Plugins/Experimental/MeshModelingToolsetExp/Source/MeshModelingToolsExp/Public/BakeMeshAttributeTool.h
+++ b/Engine/Plugins/Experimental/MeshModelingToolsetExp/Source/MeshModelingToolsExp/Public/BakeMeshAttributeTool.h
@@ -49,8 +49,12 @@ enum class EBakeMapType
 	MaterialID             = 1 << 10 UMETA(DisplayName = "Material ID"),
 	/* PolyGroup IDs as unique colors */
 	PolyGroupID            = 1 << 11 UMETA(DisplayName = "PolyGroup ID"),
+	/* Constant value of One */
+	One                    = 1 << 12,
+	/* Constant value of Zero */
+	Zero                   = 1 << 13,
 
-	All                    = 0x7FF UMETA(Hidden)
+	All                    = 0x1FFF UMETA(Hidden)
 };
 ENUM_CLASS_FLAGS(EBakeMapType);
 
diff --git a/Engine/Plugins/Experimental/MeshModelingToolsetExp/Source/MeshModelingToolsExp/Public/BakeMeshAttributeVertexTool.h b/Engine/Plugins/Experimental/MeshModelingToolsetExp/Source/MeshModelingToolsExp/Public/BakeMeshAttributeVertexTool.h
--- a/Engine/Plugins/Experimental/MeshModelingToolsetExp/Source/MeshModelingToolsExp/Public/BakeMeshAttributeVertexTool.h
+++ b/Engine/Plugins/Experimental/MeshModelingToolsetExp/Source/MeshModelingToolsExp/Public/BakeMeshAttributeVertexTool.h
@@ -71,25 +71,25 @@ public:
 
 	/** The bake output type to generate in the Red channel */
 	UPROPERTY(EditAnywhere, Category = BakeOutput, meta=(Bitmask, BitmaskEnum = "/Script/MeshModelingToolsExp.EBakeMapType",
-		ValidEnumValues="None, AmbientOcclusion, Curvature",
+		ValidEnumValues="None, AmbientOcclusion, Curvature, One, Zero",
 		EditCondition="OutputMode == EBakeVertexOutput::PerChannel", EditConditionHides))
 	int32 OutputTypeR = static_cast<int32>(EBakeMapType::None);
 
 	/** The bake output type to generate in the Green channel */
 	UPROPERTY(EditAnywhere, Category = BakeOutput, meta=(Bitmask, BitmaskEnum = "/Script/MeshModelingToolsExp.EBakeMapType",
-		ValidEnumValues="None, AmbientOcclusion, Curvature",
+		ValidEnumValues="None, AmbientOcclusion, Curvature, One, Zero",
 		EditCondition="OutputMode == EBakeVertexOutput::PerChannel", EditConditionHides))
 	int32 OutputTypeG = static_cast<int32>(EBakeMapType::None);
 
 	/** The bake output type to generate in the Blue channel */
 	UPROPERTY(EditAnywhere, Category = BakeOutput, meta=(Bitmask, BitmaskEnum = "/Script/MeshModelingToolsExp.EBakeMapType",
-		ValidEnumValues="None, AmbientOcclusion, Curvature",
+		ValidEnumValues="None, AmbientOcclusion, Curvature, One, Zero",
 		EditCondition="OutputMode == EBakeVertexOutput::PerChannel", EditConditionHides))
 	int32 OutputTypeB = static_cast<int32>(EBakeMapType::None);
 
 	/** The bake output type to generate in the Alpha channel */
 	UPROPERTY(EditAnywhere, Category = BakeOutput, meta=(Bitmask, BitmaskEnum = "/Script/MeshModelingToolsExp.EBakeMapType",
-		ValidEnumValues="None, AmbientOcclusion, Curvature",
+		ValidEnumValues="None, AmbientOcclusion, Curvature, One, Zero",
 		EditCondition="OutputMode == EBakeVertexOutput::PerChannel", EditConditionHides))
 	int32 OutputTypeA = static_cast<int32>(EBakeMapType::None);
 
@@ -167,6 +167,7 @@ protected:
 	void UpdateOnModeChange();
 	void UpdateVisualization();
 	void UpdateColorTopology();
+	void UpdateSourceVertexColors();
 	void UpdateResult();
 
 	struct FBakeSettings
