Optionally (and by default) apply Build Scale for when using Geometry Script functions to Copy from Static Mesh, and likewise apply inverse Build Scale when copying back to Static Mesh

diff --git a/Engine/Plugins/Experimental/GeometryScripting/Source/GeometryScriptingCore/Private/MeshAssetFunctions.cpp b/Engine/Plugins/Experimental/GeometryScripting/Source/GeometryScriptingCore/Private/MeshAssetFunctions.cpp
--- a/Engine/Plugins/Experimental/GeometryScripting/Source/GeometryScriptingCore/Private/MeshAssetFunctions.cpp
+++ b/Engine/Plugins/Experimental/GeometryScripting/Source/GeometryScriptingCore/Private/MeshAssetFunctions.cpp
@@ -86,33 +86,46 @@ static UDynamicMesh* CopyMeshFromStaticMesh_SourceData(
 
 	bool bHasDirtyBuildSettings = BuildSettings->bRecomputeNormals
 		|| (BuildSettings->bRecomputeTangents && AssetOptions.bRequestTangents);
+	bool bNeedsBuildScale = AssetOptions.bUseBuildScale && BuildSettings && !BuildSettings->BuildScale3D.Equals(FVector::OneVector);
+	bool bNeedsOtherBuildSettings = AssetOptions.bApplyBuildSettings && bHasDirtyBuildSettings;
 
 	FMeshDescription LocalSourceMeshCopy;
-	if (AssetOptions.bApplyBuildSettings && bHasDirtyBuildSettings )
+	if (bNeedsBuildScale || bNeedsOtherBuildSettings)
 	{
 		LocalSourceMeshCopy = *SourceMesh;
 
 		FStaticMeshAttributes Attributes(LocalSourceMeshCopy);
-		if (!Attributes.GetTriangleNormals().IsValid() || !Attributes.GetTriangleTangents().IsValid())
-		{
-			// If these attributes don't exist, create them and compute their values for each triangle
-			FStaticMeshOperations::ComputeTriangleTangentsAndNormals(LocalSourceMeshCopy);
-		}
 
-		EComputeNTBsFlags ComputeNTBsOptions = EComputeNTBsFlags::BlendOverlappingNormals;
-		ComputeNTBsOptions |= BuildSettings->bRecomputeNormals ? EComputeNTBsFlags::Normals : EComputeNTBsFlags::None;
-		if (AssetOptions.bRequestTangents)
+		if (bNeedsBuildScale)
 		{
-			ComputeNTBsOptions |= BuildSettings->bRecomputeTangents ? EComputeNTBsFlags::Tangents : EComputeNTBsFlags::None;
-			ComputeNTBsOptions |= BuildSettings->bUseMikkTSpace ? EComputeNTBsFlags::UseMikkTSpace : EComputeNTBsFlags::None;
+			FTransform BuildScaleTransform = FTransform::Identity;
+			BuildScaleTransform.SetScale3D(BuildSettings->BuildScale3D);
+			FStaticMeshOperations::ApplyTransform(LocalSourceMeshCopy, BuildScaleTransform, true /*use correct normal transforms*/);
 		}
-		ComputeNTBsOptions |= BuildSettings->bComputeWeightedNormals ? EComputeNTBsFlags::WeightedNTBs : EComputeNTBsFlags::None;
-		if (AssetOptions.bIgnoreRemoveDegenerates == false)
+
+		if (bNeedsOtherBuildSettings)
 		{
-			ComputeNTBsOptions |= BuildSettings->bRemoveDegenerates ? EComputeNTBsFlags::IgnoreDegenerateTriangles : EComputeNTBsFlags::None;
-		}
+			if (!Attributes.GetTriangleNormals().IsValid() || !Attributes.GetTriangleTangents().IsValid())
+			{
+				// If these attributes don't exist, create them and compute their values for each triangle
+				FStaticMeshOperations::ComputeTriangleTangentsAndNormals(LocalSourceMeshCopy);
+			}
 
-		FStaticMeshOperations::ComputeTangentsAndNormals(LocalSourceMeshCopy, ComputeNTBsOptions);
+			EComputeNTBsFlags ComputeNTBsOptions = EComputeNTBsFlags::BlendOverlappingNormals;
+			ComputeNTBsOptions |= BuildSettings->bRecomputeNormals ? EComputeNTBsFlags::Normals : EComputeNTBsFlags::None;
+			if (AssetOptions.bRequestTangents)
+			{
+				ComputeNTBsOptions |= BuildSettings->bRecomputeTangents ? EComputeNTBsFlags::Tangents : EComputeNTBsFlags::None;
+				ComputeNTBsOptions |= BuildSettings->bUseMikkTSpace ? EComputeNTBsFlags::UseMikkTSpace : EComputeNTBsFlags::None;
+			}
+			ComputeNTBsOptions |= BuildSettings->bComputeWeightedNormals ? EComputeNTBsFlags::WeightedNTBs : EComputeNTBsFlags::None;
+			if (AssetOptions.bIgnoreRemoveDegenerates == false)
+			{
+				ComputeNTBsOptions |= BuildSettings->bRemoveDegenerates ? EComputeNTBsFlags::IgnoreDegenerateTriangles : EComputeNTBsFlags::None;
+			}
+
+			FStaticMeshOperations::ComputeTangentsAndNormals(LocalSourceMeshCopy, ComputeNTBsOptions);
+		}
 
 		SourceMesh = &LocalSourceMeshCopy;
 	}
@@ -171,9 +184,17 @@ static UDynamicMesh* CopyMeshFromStaticMesh_RenderData(
 
 	FStaticMeshLODResourcesToDynamicMesh::ConversionOptions ConvertOptions;
 #if WITH_EDITOR
-	// respect BuildScale build setting
-	const FMeshBuildSettings& LODBuildSettings = FromStaticMeshAsset->GetSourceModel(UseLODIndex).BuildSettings;
-	ConvertOptions.BuildScale = (FVector3d)LODBuildSettings.BuildScale3D;
+	if (AssetOptions.bUseBuildScale)
+	{
+		// respect BuildScale build setting
+		const FMeshBuildSettings& LODBuildSettings = FromStaticMeshAsset->GetSourceModel(UseLODIndex).BuildSettings;
+		ConvertOptions.BuildScale = (FVector3d)LODBuildSettings.BuildScale3D;
+	}
+#else
+	if (!AssetOptions.bUseBuildScale)
+	{
+		UE::Geometry::AppendError(Debug, EGeometryScriptErrorType::InvalidInputs, LOCTEXT("CopyMeshFromStaticMesh_BuildScaleAlreadyBaked", "CopyMeshFromStaticMesh: Requested mesh without BuildScale, but BuildScale is already baked into the RenderData."));
+	}
 #endif
 
 	FDynamicMesh3 NewMesh;
@@ -282,18 +303,40 @@ UDynamicMesh*  UGeometryScriptLibrary_StaticMeshFunctions::CopyMeshToStaticMesh(
 	// mark as modified
 	ToStaticMeshAsset->Modify();
 
-	auto ConfigureBuildSettingsFromOptions = [](FStaticMeshSourceModel& SourceModel, FGeometryScriptCopyMeshToAssetOptions& Options)
+	auto ConfigureBuildSettingsFromOptions = [](FStaticMeshSourceModel& SourceModel, FGeometryScriptCopyMeshToAssetOptions& Options) -> FVector
 												{
 													FMeshBuildSettings& BuildSettings = SourceModel.BuildSettings;
 													BuildSettings.bRecomputeNormals  = Options.bEnableRecomputeNormals;
 													BuildSettings.bRecomputeTangents = Options.bEnableRecomputeTangents;
 													BuildSettings.bRemoveDegenerates = Options.bEnableRemoveDegenerates;
+													if (!Options.bUseBuildScale) // if we're not using build scale, set asset BuildScale to 1,1,1
+													{
+														BuildSettings.BuildScale3D = FVector::OneVector;
+													}
+													return BuildSettings.BuildScale3D;
 												};
+	
+	auto ApplyInverseBuildScale = [](FMeshDescription& MeshDescription, FVector BuildScale)
+	{
+		if (BuildScale.Equals(FVector::OneVector))
+		{
+			return;
+		}
+		FTransform InverseBuildScaleTransform = FTransform::Identity;
+		FVector InverseBuildScale;
+		// Safely invert BuildScale
+		for (int32 Idx = 0; Idx < 3; ++Idx)
+		{
+			InverseBuildScale[Idx] = FMath::IsNearlyZero(BuildScale[Idx], FMathd::Epsilon) ? 1.0 : 1.0 / BuildScale[Idx];
+		}
+		InverseBuildScaleTransform.SetScale3D(InverseBuildScale);
+		FStaticMeshOperations::ApplyTransform(MeshDescription, InverseBuildScaleTransform, true /*use correct normal transforms*/);
+	};
 
 	if (TargetLOD.bWriteHiResSource)
 	{
 		// update model build settings
-		ConfigureBuildSettingsFromOptions(ToStaticMeshAsset->GetHiResSourceModel(), Options);
+		FVector BuildScale = ConfigureBuildSettingsFromOptions(ToStaticMeshAsset->GetHiResSourceModel(), Options);
 
 		ToStaticMeshAsset->ModifyHiResMeshDescription();
 		FMeshDescription* NewHiResMD = ToStaticMeshAsset->CreateHiResMeshDescription();
@@ -311,6 +354,7 @@ UDynamicMesh*  UGeometryScriptLibrary_StaticMeshFunctions::CopyMeshToStaticMesh(
 				Converter.Convert(&ReadMesh, *NewHiResMD, !Options.bEnableRecomputeTangents);
 			});
 
+		ApplyInverseBuildScale(*NewHiResMD, BuildScale);
 
 		ToStaticMeshAsset->CommitHiResMeshDescription();
 	}
@@ -323,7 +367,7 @@ UDynamicMesh*  UGeometryScriptLibrary_StaticMeshFunctions::CopyMeshToStaticMesh(
 		}
 
 		// update model build settings
-		ConfigureBuildSettingsFromOptions(ToStaticMeshAsset->GetSourceModel(UseLODIndex), Options);
+		FVector BuildScale = ConfigureBuildSettingsFromOptions(ToStaticMeshAsset->GetSourceModel(UseLODIndex), Options);
 
 		FMeshDescription* MeshDescription = ToStaticMeshAsset->GetMeshDescription(UseLODIndex);
 		if (MeshDescription == nullptr)
@@ -348,6 +392,8 @@ UDynamicMesh*  UGeometryScriptLibrary_StaticMeshFunctions::CopyMeshToStaticMesh(
 			Converter.Convert(&ReadMesh, *MeshDescription, !Options.bEnableRecomputeTangents);
 		});
 
+		ApplyInverseBuildScale(*MeshDescription, BuildScale);
+
 		// Setting to prevent the standard static mesh reduction from running and replacing the render LOD.
 		FStaticMeshSourceModel& ThisSourceModel = ToStaticMeshAsset->GetSourceModel(UseLODIndex);
 		ThisSourceModel.ResetReductionSetting();
diff --git a/Engine/Plugins/Experimental/GeometryScripting/Source/GeometryScriptingCore/Public/GeometryScript/MeshAssetFunctions.h b/Engine/Plugins/Experimental/GeometryScripting/Source/GeometryScriptingCore/Public/GeometryScript/MeshAssetFunctions.h
--- a/Engine/Plugins/Experimental/GeometryScripting/Source/GeometryScriptingCore/Public/GeometryScript/MeshAssetFunctions.h
+++ b/Engine/Plugins/Experimental/GeometryScripting/Source/GeometryScriptingCore/Public/GeometryScript/MeshAssetFunctions.h
@@ -19,14 +19,21 @@ struct GEOMETRYSCRIPTINGCORE_API FGeometryScriptCopyMeshFromAssetOptions
 {
 	GENERATED_BODY()
 public:
+	// Whether to apply Build Settings during the mesh copy.
 	UPROPERTY(EditAnywhere, BlueprintReadWrite, Category = Options)
 	bool bApplyBuildSettings = true;
 
+	// Whether to request tangents on the copied mesh. If tangents are not requested, tangent-related build settings will also be ignored.
 	UPROPERTY(EditAnywhere, BlueprintReadWrite, Category = Options)
 	bool bRequestTangents = true;
 
+	// Whether to ignore the 'remove degenerates' option from Build Settings. Note: Only applies if 'Apply Build Settings' is enabled.
 	UPROPERTY(EditAnywhere, BlueprintReadWrite, Category = Options)
 	bool bIgnoreRemoveDegenerates = true;
+
+	// Whether to scale the copied mesh by the Build Setting's 'Build Scale'. Note: This is considered separately from the 'Apply Build Settings' option.
+	UPROPERTY(EditAnywhere, BlueprintReadWrite, Category = Options)
+	bool bUseBuildScale = true;
 };
 
 /**
@@ -65,6 +72,9 @@ public:
 	UPROPERTY(EditAnywhere, BlueprintReadWrite, Category = Options)
 	bool bEnableRemoveDegenerates = false;
 
+	// Whether to use the build scale on the target asset. If enabled, the inverse scale will be applied when saving to the asset, and the BuildScale will be preserved. Otherwise, BuildScale will be set to 1.0 on the asset BuildSettings.
+	UPROPERTY(EditAnywhere, BlueprintReadWrite, Category = Options)
+	bool bUseBuildScale = true;
 	
 	UPROPERTY(EditAnywhere, BlueprintReadWrite, Category = Options)
 	bool bReplaceMaterials = false;
