Add a method to remove unused vertices from a dynamic mesh, and expose it to geometry script

diff --git a/Engine/Plugins/Experimental/GeometryScripting/Source/GeometryScriptingCore/Private/MeshRepairFunctions.cpp b/Engine/Plugins/Experimental/GeometryScripting/Source/GeometryScriptingCore/Private/MeshRepairFunctions.cpp
--- a/Engine/Plugins/Experimental/GeometryScripting/Source/GeometryScriptingCore/Private/MeshRepairFunctions.cpp
+++ b/Engine/Plugins/Experimental/GeometryScripting/Source/GeometryScriptingCore/Private/MeshRepairFunctions.cpp
@@ -46,6 +46,24 @@ UDynamicMesh* UGeometryScriptLibrary_MeshRepairFunctions::CompactMesh(
 	return TargetMesh;
 }
 
+UDynamicMesh* UGeometryScriptLibrary_MeshRepairFunctions::RemoveUnusedVertices(
+	UDynamicMesh* TargetMesh,
+	UGeometryScriptDebug* Debug)
+{
+	if (TargetMesh == nullptr)
+	{
+		UE::Geometry::AppendError(Debug, EGeometryScriptErrorType::InvalidInputs, LOCTEXT("RemoveUnusedVertices_InvalidInput", "RemoveUnusedVertices: TargetMesh is Null"));
+		return TargetMesh;
+	}
+
+	TargetMesh->EditMesh([&](FDynamicMesh3& EditMesh) 
+	{
+		EditMesh.RemoveUnusedVertices();
+	}, EDynamicMeshChangeType::GeneralEdit, EDynamicMeshAttributeChangeFlags::Unknown, false);
+
+	return TargetMesh;
+}
+
 
 UDynamicMesh* UGeometryScriptLibrary_MeshRepairFunctions::ResolveMeshTJunctions(
 	UDynamicMesh* TargetMesh,
diff --git a/Engine/Plugins/Experimental/GeometryScripting/Source/GeometryScriptingCore/Public/GeometryScript/MeshRepairFunctions.h b/Engine/Plugins/Experimental/GeometryScripting/Source/GeometryScriptingCore/Public/GeometryScript/MeshRepairFunctions.h
--- a/Engine/Plugins/Experimental/GeometryScripting/Source/GeometryScriptingCore/Public/GeometryScript/MeshRepairFunctions.h
+++ b/Engine/Plugins/Experimental/GeometryScripting/Source/GeometryScriptingCore/Public/GeometryScript/MeshRepairFunctions.h
@@ -174,6 +174,15 @@ public:
 		UDynamicMesh* TargetMesh, 
 		UGeometryScriptDebug* Debug = nullptr);
 
+	/**
+	 * Remove vertices that are not used by any triangles. Note: Does not update the IDs of any remaining vertices; use CompactMesh to do so.
+	 */
+	UFUNCTION(BlueprintCallable, Category = "GeometryScript|Repair", meta = (ScriptMethod))
+	static UPARAM(DisplayName = "Target Mesh") UDynamicMesh*
+	RemoveUnusedVertices(
+		UDynamicMesh* TargetMesh,
+		UGeometryScriptDebug* Debug = nullptr);
+
 	/**
 	* Attempts to resolve T-Junctions in the mesh by addition of vertices and welding.
 	*/
diff --git a/Engine/Source/Runtime/GeometryCore/Private/DynamicMesh/DynamicMesh3_Edits.cpp b/Engine/Source/Runtime/GeometryCore/Private/DynamicMesh/DynamicMesh3_Edits.cpp
--- a/Engine/Source/Runtime/GeometryCore/Private/DynamicMesh/DynamicMesh3_Edits.cpp
+++ b/Engine/Source/Runtime/GeometryCore/Private/DynamicMesh/DynamicMesh3_Edits.cpp
@@ -286,7 +286,21 @@ EMeshResult FDynamicMesh3::InsertTriangle(int tid, const FIndex3i& tv, int gid,
 
 
 
+void FDynamicMesh3::RemoveUnusedVertices()
+{
+	for (int32 VID = 0; VID < MaxVertexID(); ++VID)
+	{
+		// If vertex exists but is not referenced by any triangles
+		if (VertexRefCounts.GetRefCount(VID) == 1)
+		{
+			VertexRefCounts.Decrement(VID);
+			checkSlow(VertexRefCounts.IsValid(VID) == false); // vertex should now not be valid
+			checkSlow(VertexEdgeLists.GetCount(VID) == 0); // vertex should not have had any edges attached
+		}
+	}
 
+	UpdateChangeStamps(true, true);
+}
 
 
 
diff --git a/Engine/Source/Runtime/GeometryCore/Public/DynamicMesh/DynamicMesh3.h b/Engine/Source/Runtime/GeometryCore/Public/DynamicMesh/DynamicMesh3.h
--- a/Engine/Source/Runtime/GeometryCore/Public/DynamicMesh/DynamicMesh3.h
+++ b/Engine/Source/Runtime/GeometryCore/Public/DynamicMesh/DynamicMesh3.h
@@ -1173,6 +1173,11 @@ public:
 	 */
 	GEOMETRYCORE_API void CompactInPlace(FCompactMaps* CompactInfo = nullptr);
 
+	/**
+	 * Remove unused vertices. Note: Does not compact the remaining vertices.
+	 */
+	GEOMETRYCORE_API void RemoveUnusedVertices();
+
 	/**
 	 * Reverse the ccw/cw orientation of all triangles in the mesh, and
 	 * optionally flip the vertex normals if they exist
