PR #11363: Subclasses of UGeneratedDynamicMeshActor can rebuild generated mesh by overriding the function `RebuildGeneratedMesh`

diff --git a/Engine/Plugins/Experimental/GeometryScripting/Source/GeometryScriptingEditor/Private/GeometryActors/GeneratedDynamicMeshActor.cpp b/Engine/Plugins/Experimental/GeometryScripting/Source/GeometryScriptingEditor/Private/GeometryActors/GeneratedDynamicMeshActor.cpp
--- a/Engine/Plugins/Experimental/GeometryScripting/Source/GeometryScriptingEditor/Private/GeometryActors/GeneratedDynamicMeshActor.cpp
+++ b/Engine/Plugins/Experimental/GeometryScripting/Source/GeometryScriptingEditor/Private/GeometryActors/GeneratedDynamicMeshActor.cpp
@@ -202,12 +202,12 @@ void AGeneratedDynamicMeshActor::ExecuteRebuildGeneratedMeshIfPending()
 		Progress.MakeDialogDelayed(this->DialogDelay, true);
 		ActiveSlowTask = &Progress;
 		CurProgressAccumSteps = 0;
-		OnRebuildGeneratedMesh(Component->GetDynamicMesh());
+		RebuildGeneratedMesh(Component->GetDynamicMesh());
 		ActiveSlowTask = nullptr;
 	}
 	else
 	{
-		OnRebuildGeneratedMesh(Component->GetDynamicMesh());
+		RebuildGeneratedMesh(Component->GetDynamicMesh());
 	}
 
 	bGeneratedMeshRebuildPending = false;
@@ -218,6 +218,10 @@ void AGeneratedDynamicMeshActor::ExecuteRebuildGeneratedMeshIfPending()
 	}
 }
 
+void AGeneratedDynamicMeshActor::RebuildGeneratedMesh(UDynamicMesh* TargetMesh)
+{
+	OnRebuildGeneratedMesh(TargetMesh);
+}
 
 void AGeneratedDynamicMeshActor::MarkForMeshRebuild(bool bImmediate, bool bImmediateEvenIfFrozen)
 {
diff --git a/Engine/Plugins/Experimental/GeometryScripting/Source/GeometryScriptingEditor/Public/GeometryActors/GeneratedDynamicMeshActor.h b/Engine/Plugins/Experimental/GeometryScripting/Source/GeometryScriptingEditor/Public/GeometryActors/GeneratedDynamicMeshActor.h
--- a/Engine/Plugins/Experimental/GeometryScripting/Source/GeometryScriptingEditor/Public/GeometryActors/GeneratedDynamicMeshActor.h
+++ b/Engine/Plugins/Experimental/GeometryScripting/Source/GeometryScriptingEditor/Public/GeometryActors/GeneratedDynamicMeshActor.h
@@ -35,7 +35,7 @@ public:
 	UPROPERTY(Category = "DynamicMeshActor", EditAnywhere, BlueprintReadWrite)
 	bool bFrozen = false;
 
-	/** If true, the DynamicMeshComponent will be cleared before the OnRebuildGeneratedMesh event is executed. */
+	/** If true, the DynamicMeshComponent will be cleared before the RebuildGeneratedMesh event is executed. */
 	UPROPERTY(Category = "DynamicMeshActor|Advanced", EditAnywhere, BlueprintReadWrite)
 	bool bResetOnRebuild = true;
 
@@ -55,12 +55,20 @@ public:
 
 
 	/**
-	 * This function will fire the OnRebuildGeneratedMesh event if the actor has been
+	 * This function will fire the RebuildGeneratedMesh function if the actor has been
 	 * marked for a pending rebuild (eg via OnConstruction)
 	 */
 	virtual void ExecuteRebuildGeneratedMeshIfPending();
 
 
+protected:
+	/**
+	 * Overridable native event for when the generated Mesh should be rebuilt.
+	 * This function will fire the OnRebuildGeneratedMesh event to the BP.
+	 */
+	virtual void RebuildGeneratedMesh(UDynamicMesh* TargetMesh);
+
+
 public:
 
 	/** 
@@ -101,7 +109,7 @@ public:
 
 protected:
 	// this internal flag is set in OnConstruction, and will cause ExecuteRebuildGeneratedMesh to
-	// fire the OnRebuildGeneratedMesh event, after which the flag will be cleared
+	// fire the RebuildGeneratedMesh event, after which the flag will be cleared
 	bool bGeneratedMeshRebuildPending = false;
 
 	// indicates that this Actor is registered with the UEditorGeometryGenerationSubsystem, which 
