add level set option to geometry script 'collision from mesh' options

diff --git a/Engine/Plugins/Experimental/GeometryScripting/Source/GeometryScriptingCore/Private/CollisionFunctions.cpp b/Engine/Plugins/Experimental/GeometryScripting/Source/GeometryScriptingCore/Private/CollisionFunctions.cpp
--- a/Engine/Plugins/Experimental/GeometryScripting/Source/GeometryScriptingCore/Private/CollisionFunctions.cpp
+++ b/Engine/Plugins/Experimental/GeometryScripting/Source/GeometryScriptingCore/Private/CollisionFunctions.cpp
@@ -136,6 +136,9 @@ void ComputeCollisionFromMesh(
 	case EGeometryScriptCollisionGenerationMethod::MinVolumeShapes:
 		ShapeGenerator.Generate_MinVolume(NewCollision.Geometry);
 		break;
+	case EGeometryScriptCollisionGenerationMethod::LevelSets:
+		ShapeGenerator.Generate_LevelSets(NewCollision.Geometry);
+		break;
 	}
 
 	if (Options.bRemoveFullyContainedShapes && Components.Num() > 1)
diff --git a/Engine/Plugins/Experimental/GeometryScripting/Source/GeometryScriptingCore/Public/GeometryScript/CollisionFunctions.h b/Engine/Plugins/Experimental/GeometryScripting/Source/GeometryScriptingCore/Public/GeometryScript/CollisionFunctions.h
--- a/Engine/Plugins/Experimental/GeometryScripting/Source/GeometryScriptingCore/Public/GeometryScript/CollisionFunctions.h
+++ b/Engine/Plugins/Experimental/GeometryScripting/Source/GeometryScriptingCore/Public/GeometryScript/CollisionFunctions.h
@@ -20,7 +20,8 @@ enum class EGeometryScriptCollisionGenerationMethod : uint8
 	Capsules = 3,
 	ConvexHulls = 4,
 	SweptHulls = 5,
-	MinVolumeShapes = 6
+	MinVolumeShapes = 6,
+	LevelSets = 7
 };
 
 
