fix geometry script's ray plane intersection function

diff --git a/Engine/Plugins/Experimental/GeometryScripting/Source/GeometryScriptingCore/Private/ShapeFunctions.cpp b/Engine/Plugins/Experimental/GeometryScripting/Source/GeometryScriptingCore/Private/ShapeFunctions.cpp
--- a/Engine/Plugins/Experimental/GeometryScripting/Source/GeometryScriptingCore/Private/ShapeFunctions.cpp
+++ b/Engine/Plugins/Experimental/GeometryScripting/Source/GeometryScriptingCore/Private/ShapeFunctions.cpp
@@ -142,9 +142,10 @@ bool UGeometryScriptLibrary_RayFunctions::GetRayPlaneIntersection(FRay Ray, FPla
 {
 	const FVector PlaneNormal = FVector(Plane.X, Plane.Y, Plane.Z);
 	const FVector PlaneOrigin = PlaneNormal * Plane.W;
-	if (FMathd::Abs(FVector::DotProduct(Ray.Direction, PlaneNormal)) > FMathd::ZeroTolerance)
+	const double DirDotN = Ray.Direction.Dot(PlaneNormal);
+	if (FMathd::Abs(DirDotN) > FMathd::ZeroTolerance)
 	{
-		HitDistance = FVector::DotProduct((PlaneOrigin - Ray.Origin), PlaneNormal);
+		HitDistance = (Plane.W - Ray.Origin.Dot(PlaneNormal)) / DirDotN;
 		return true;
 	}
 	HitDistance = 0;
