fix crash on appending to an empty polygon list
+ fix polygon path offset functions not properly updating the success flag, and not handling the copy-input-on-failure setting

diff --git a/Engine/Plugins/Experimental/GeometryScripting/Source/GeometryScriptingCore/Private/PolygonFunctions.cpp b/Engine/Plugins/Experimental/GeometryScripting/Source/GeometryScriptingCore/Private/PolygonFunctions.cpp
--- a/Engine/Plugins/Experimental/GeometryScripting/Source/GeometryScriptingCore/Private/PolygonFunctions.cpp
+++ b/Engine/Plugins/Experimental/GeometryScripting/Source/GeometryScriptingCore/Private/PolygonFunctions.cpp
@@ -418,7 +418,7 @@ void UGeometryScriptLibrary_PolygonListFunctions::AppendPolygonList(FGeometryScr
 {
 	if (!PolygonList.Polygons)
 	{
-		PolygonList.Polygons.Reset();
+		PolygonList.Reset();
 	}
 	if (!PolygonsToAppend.Polygons)
 	{
@@ -587,10 +587,17 @@ FGeometryScriptGeneralPolygonList UGeometryScriptLibrary_PolygonListFunctions::C
 	PolygonOffset.DefaultStepsPerRadianScale = OffsetOptions.StepsPerRadianScale;
 	PolygonOffset.Offset = Offset;
 	PolygonOffset.Polygons.Add(TArrayView<FVector2D>(Path.GetData(), Path.Num()));
-	PolygonOffset.ComputeResult();
+	bOperationSuccess = PolygonOffset.ComputeResult();
 	FGeometryScriptGeneralPolygonList PolygonListResult;
 	PolygonListResult.Reset();
-	*PolygonListResult.Polygons = MoveTemp(PolygonOffset.Result);
+	if (!bOperationSuccess && bCopyInputOnFailure)
+	{
+		PolygonListResult.Polygons->Emplace(FPolygon2d(Path));
+	}
+	else
+	{
+		*PolygonListResult.Polygons = MoveTemp(PolygonOffset.Result);
+	}
 	return PolygonListResult;
 }
 
@@ -630,10 +637,23 @@ FGeometryScriptGeneralPolygonList UGeometryScriptLibrary_PolygonListFunctions::C
 	{
 		PolygonOffset.Polygons.Add(TArrayView<FVector2D>(AllPathPts.GetData() + PathPtIdx, PathLens[Idx]));
 	}
-	PolygonOffset.ComputeResult();
+	bOperationSuccess = PolygonOffset.ComputeResult();
 	FGeometryScriptGeneralPolygonList PolygonListResult;
 	PolygonListResult.Reset();
-	*PolygonListResult.Polygons = MoveTemp(PolygonOffset.Result);
+	if (!bOperationSuccess && bCopyInputOnFailure)
+	{
+		for (const FGeometryScriptPolyPath& Path : PolyPaths)
+		{
+			if (Path.Path)
+			{
+				PolygonListResult.Polygons->Emplace(FPolygon2d(*Path.Path));
+			}
+		}
+	}
+	else
+	{
+		*PolygonListResult.Polygons = MoveTemp(PolygonOffset.Result);
+	}
 	return PolygonListResult;
 }
 
