Geometry Script Node: Copy Mesh From Component.  Fiix bug where Outcome Pin was always set to failure.

diff --git a/Engine/Plugins/Runtime/GeometryScripting/Source/GeometryScriptingCore/Private/SceneUtilityFunctions.cpp b/Engine/Plugins/Runtime/GeometryScripting/Source/GeometryScriptingCore/Private/SceneUtilityFunctions.cpp
--- a/Engine/Plugins/Runtime/GeometryScripting/Source/GeometryScriptingCore/Private/SceneUtilityFunctions.cpp
+++ b/Engine/Plugins/Runtime/GeometryScripting/Source/GeometryScriptingCore/Private/SceneUtilityFunctions.cpp
@@ -79,6 +79,7 @@ UDynamicMesh* UGeometryScriptLibrary_SceneUtilityFunctions::CopyMeshFromComponen
 	if (bSuccess)
 	{
 		ToDynamicMesh->SetMesh(MoveTemp(NewMesh));
+		Outcome = EGeometryScriptOutcomePins::Success;
 	}
 	else // failed
 	{
