When executing CopyMeshToSkeletalMesh, append an error and return if we failed to generate the Mesh Description for the TargetLOD.

#rnx

diff --git a/Engine/Plugins/Runtime/GeometryScripting/Source/GeometryScriptingCore/Private/MeshAssetFunctions.cpp b/Engine/Plugins/Runtime/GeometryScripting/Source/GeometryScriptingCore/Private/MeshAssetFunctions.cpp
--- a/Engine/Plugins/Runtime/GeometryScripting/Source/GeometryScriptingCore/Private/MeshAssetFunctions.cpp
+++ b/Engine/Plugins/Runtime/GeometryScripting/Source/GeometryScriptingCore/Private/MeshAssetFunctions.cpp
@@ -753,10 +753,21 @@ UDynamicMesh* UGeometryScriptLibrary_StaticMeshFunctions::CopyMeshToSkeletalMesh
 
 	verify(ToSkeletalMeshAsset->Modify());
 
-	FMeshDescription* MeshDescription = ToSkeletalMeshAsset->GetMeshDescription(TargetLOD.LODIndex);
+	// Ensure we have enough LODInfos to cover up to the requested LOD.
+	for (int32 LODIndex = ToSkeletalMeshAsset->GetLODInfoArray().Num(); LODIndex <= TargetLOD.LODIndex; LODIndex++)
+	{
+		FSkeletalMeshLODInfo& LODInfo = ToSkeletalMeshAsset->AddLODInfo();
+		
+		ToSkeletalMeshAsset->GetImportedModel()->LODModels.Add(new FSkeletalMeshLODModel);
+		LODInfo.ReductionSettings.BaseLOD = 0;
+	}
+
+	FMeshDescription* MeshDescription = ToSkeletalMeshAsset->CreateMeshDescription(TargetLOD.LODIndex);
+
 	if (MeshDescription == nullptr)
 	{
-		MeshDescription = ToSkeletalMeshAsset->CreateMeshDescription(TargetLOD.LODIndex);
+		UE::Geometry::AppendError(Debug, EGeometryScriptErrorType::InvalidInputs, LOCTEXT("CopyMeshToSkeletalMesh_TargetMeshDescription", "CopyMeshToSkeletalMesh: Failed to generate the mesh data for the Target LOD Index"));
+		return FromDynamicMesh;
 	}
 	
 	FSkeletalMeshAttributes MeshAttributes(*MeshDescription);
@@ -771,15 +782,6 @@ UDynamicMesh* UGeometryScriptLibrary_StaticMeshFunctions::CopyMeshToSkeletalMesh
 		Converter.Convert(&ReadMesh, *MeshDescription, !Options.bEnableRecomputeTangents);
 	});
 
-	// Ensure we have enough LODInfos to cover up to the requested LOD.
-	for (int32 LODIndex = ToSkeletalMeshAsset->GetLODInfoArray().Num(); LODIndex <= TargetLOD.LODIndex; LODIndex++)
-	{
-		FSkeletalMeshLODInfo& LODInfo = ToSkeletalMeshAsset->AddLODInfo();
-		
-		ToSkeletalMeshAsset->GetImportedModel()->LODModels.Add(new FSkeletalMeshLODModel);
-		LODInfo.ReductionSettings.BaseLOD = 0;
-	}
-
 	FSkeletalMeshLODInfo* SkeletalLODInfo = ToSkeletalMeshAsset->GetLODInfo(TargetLOD.LODIndex);
 	SkeletalLODInfo->BuildSettings.bRecomputeNormals = Options.bEnableRecomputeNormals;
 	SkeletalLODInfo->BuildSettings.bRecomputeTangents = Options.bEnableRecomputeTangents;
