document Hull Mesh parameter for geometry script containment functions

diff --git a/Engine/Plugins/Runtime/GeometryScripting/Source/GeometryScriptingCore/Public/GeometryScript/ContainmentFunctions.h b/Engine/Plugins/Runtime/GeometryScripting/Source/GeometryScriptingCore/Public/GeometryScript/ContainmentFunctions.h
--- a/Engine/Plugins/Runtime/GeometryScripting/Source/GeometryScriptingCore/Public/GeometryScript/ContainmentFunctions.h
+++ b/Engine/Plugins/Runtime/GeometryScripting/Source/GeometryScriptingCore/Public/GeometryScript/ContainmentFunctions.h
@@ -91,8 +91,10 @@ class GEOMETRYSCRIPTINGCORE_API UGeometryScriptLibrary_ContainmentFunctions : pu
 public:
 
 	/**
-	 * Compute the Convex Hull of a given Mesh, or part of the mesh if an optional Selection is provided
-	 * @param Selection selection of mesh faces/vertices to contain in the convex hull. If not provided, entire mesh is used.
+	 * Compute the Convex Hull of a given Target Mesh, or part of the mesh if an optional Selection is provided, and put the result in Hull Mesh
+	 * @param CopyToMesh	The Dynamic Mesh to store the convex hull geometry.
+	 * @param CopyToMeshOut	The resulting convex hull.
+	 * @param Selection	Selection of mesh faces/vertices to contain in the convex hull. If not provided, entire mesh is used.
 	 */
 	UFUNCTION(BlueprintCallable, Category = "GeometryScript|Containment", meta=(ScriptMethod))
 	static UPARAM(DisplayName = "Target Mesh") UDynamicMesh* 
@@ -105,8 +107,10 @@ public:
 		UGeometryScriptDebug* Debug = nullptr);
 
 	/**
-	* Compute the Swept Hull of a given Mesh for a given 3D Plane defined by ProjectionFrame.
+	* Compute the Swept Hull of a given Target Mesh for a given 3D Plane defined by ProjectionFrame, and put the result in Hull Mesh
 	* The Swept Hull is a linear sweep of the 2D convex hull of the mesh vertices projected onto the plane (the sweep precisely contains the mesh extents along the plane normal)
+	* @param CopyToMesh	The Dynamic Mesh to store the swept hull geometry.
+	* @param CopyToMeshOut	The resulting swept hull.
 	*/
 	UFUNCTION(BlueprintCallable, Category = "GeometryScript|Containment", meta=(ScriptMethod))
 	static UPARAM(DisplayName = "Target Mesh") UDynamicMesh* 
@@ -120,9 +124,12 @@ public:
 
 	/**
 	 * Compute a Convex Hull Decomposition of the given TargetMesh. Assuming more than one hull is requested,
-	 * multiple hulls will be returned that attempt to approximate the mesh. There is no guarantee that the entire
-	 * mesh is contained in the hulls.
+	 * multiple hulls will be returned that attempt to approximate the mesh. If simplification settings are enabled,
+	 * there is no guarantee that the entire mesh is contained in the hulls.
+	 * 
 	 * @warning this function can be quite expensive, and the results are expected to change in the future as the Convex Decomposition algorithm is improved
+	 * @param CopyToMesh The Dynamic Mesh to store the convex hulls as a single, combined mesh. Note: SplitMeshByComponents can separate this result into its convex parts.
+	 * @param CopyToMeshOut A combined mesh of the convex hulls. Note: SplitMeshByComponents can separate this result into its convex parts.
 	 */
 	UFUNCTION(BlueprintCallable, Category = "GeometryScript|Containment|Experimental", meta=(ScriptMethod))
 	static UPARAM(DisplayName = "Target Mesh") UDynamicMesh* 
