add navigable convex decomposition method to geometry script collision functions
+ expand preprocessing options for convex decomposition algorithm

diff --git a/Engine/Plugins/Runtime/GeometryScripting/Source/GeometryScriptingCore/Private/CollisionFunctions.cpp b/Engine/Plugins/Runtime/GeometryScripting/Source/GeometryScriptingCore/Private/CollisionFunctions.cpp
--- a/Engine/Plugins/Runtime/GeometryScripting/Source/GeometryScriptingCore/Private/CollisionFunctions.cpp
+++ b/Engine/Plugins/Runtime/GeometryScripting/Source/GeometryScriptingCore/Private/CollisionFunctions.cpp
@@ -239,6 +239,8 @@ static UE::Geometry::FNegativeSpaceSampleSettings::ESampleMethod ConvertNegative
 		return UE::Geometry::FNegativeSpaceSampleSettings::ESampleMethod::Uniform;
 	case ENegativeSpaceSampleMethod::VoxelSearch:
 		return UE::Geometry::FNegativeSpaceSampleSettings::ESampleMethod::VoxelSearch;
+	case ENegativeSpaceSampleMethod::NavigableVoxelSearch:
+		return UE::Geometry::FNegativeSpaceSampleSettings::ESampleMethod::NavigableVoxelSearch;
 	}
 	return UE::Geometry::FNegativeSpaceSampleSettings::ESampleMethod::Uniform;
 }
@@ -1303,6 +1305,94 @@ FGeometryScriptSphereCovering UGeometryScriptLibrary_CollisionFunctions::Compute
 	return ToRet;
 }
 
+FGeometryScriptSimpleCollision UGeometryScriptLibrary_CollisionFunctions::ComputeNavigableConvexDecomposition(
+	const UDynamicMesh* TargetMesh,
+	const FNavigableConvexDecompositionOptions& NegativeSpaceOptions,
+	UGeometryScriptDebug* Debug)
+{
+	FGeometryScriptSimpleCollision ToReturn;
+	
+	if (!TargetMesh)
+	{
+		UE::Geometry::AppendError(Debug, EGeometryScriptErrorType::InvalidInputs, LOCTEXT("ComputeNavigableConvexDecomposition_NullInput", "ComputeNavigableConvexDecomposition: Dynamic Mesh is Null"));
+		return ToReturn;
+	}
+
+	TargetMesh->ProcessMesh([&NegativeSpaceOptions, &ToReturn](const FDynamicMesh3& Mesh) -> void
+	{
+		FConvexDecomposition3::FPreprocessMeshOptions PreprocessOptions;
+		PreprocessOptions.bMergeEdges = true;
+		PreprocessOptions.CustomPreprocess = [](FDynamicMesh3& ProcessMesh, const FAxisAlignedBox3d& Bounds) -> void
+		{
+			// for solid inputs, flip orientation if the initial volume is negative
+			if (ProcessMesh.IsClosed())
+			{
+				double InitialVolume = TMeshQueries<FDynamicMesh3>::GetVolumeArea(ProcessMesh).X;
+				if (InitialVolume < 0)
+				{
+					ProcessMesh.ReverseOrientation();
+				}
+			}
+			// Note: If we add options to simplify the input mesh, should be applied here.
+		};
+
+		FConvexDecomposition3 ConvexDecomposition(Mesh, PreprocessOptions);
+		const bool bIsSolid = ConvexDecomposition.IsInputSolid();
+		ConvexDecomposition.bTreatAsSolid = bIsSolid;
+
+		FNegativeSpaceSampleSettings NegativeSpaceSettings;
+		NegativeSpaceSettings.ApplyDefaults();
+		NegativeSpaceSettings.SampleMethod = FNegativeSpaceSampleSettings::ESampleMethod::NavigableVoxelSearch;
+		NegativeSpaceSettings.bDeterministic = true;
+		NegativeSpaceSettings.bRequireSearchSampleCoverage = true;
+		NegativeSpaceSettings.bOnlyConnectedToHull = NegativeSpaceOptions.bIgnoreUnreachableInternalSpace;
+		NegativeSpaceSettings.TargetNumSamples = 0;
+		NegativeSpaceSettings.bAllowSamplesInsideMesh = !bIsSolid;
+
+		NegativeSpaceSettings.ReduceRadiusMargin = NegativeSpaceOptions.Tolerance;
+		NegativeSpaceSettings.MinRadius = NegativeSpaceOptions.MinRadius;
+		NegativeSpaceSettings.MinSpacing = 0;
+
+		ConvexDecomposition.InitializeNegativeSpace(NegativeSpaceSettings, NegativeSpaceOptions.CustomNavigablePositions);
+
+		ConvexDecomposition.MaxConvexEdgePlanes = 4;
+		ConvexDecomposition.bSplitDisconnectedComponents = false;
+		ConvexDecomposition.ConvexEdgeAngleMoreSamplesThreshold = 180;
+		ConvexDecomposition.ThickenAfterHullFailure = FMath::Max(FMathd::ZeroTolerance, NegativeSpaceSettings.ReduceRadiusMargin * .01);
+		int32 NumSplits = -1;
+		for (int32 Split = 0; Split < NumSplits || NumSplits < 0; Split++)
+		{
+
+			int32 NumSplit = ConvexDecomposition.SplitWorst(false, -1, true, NegativeSpaceSettings.ReduceRadiusMargin * .5);
+
+			if (NumSplit == 0)
+			{
+				break;
+			}
+		}
+
+		ConvexDecomposition.FixHullOverlapsInNegativeSpace();
+			
+		int32 NumHullsBefore = ConvexDecomposition.NumHulls();
+		constexpr double MinThicknessToleranceWorldSpace = 0;
+		int32 NumMerged = ConvexDecomposition.MergeBest(-1, 0, MinThicknessToleranceWorldSpace, true);
+
+		// transfer to output
+		for (int32 HullIdx = 0; HullIdx < ConvexDecomposition.Decomposition.Num(); ++HullIdx)
+		{
+			const FConvexDecomposition3::FConvexPart& Part = ConvexDecomposition.Decomposition[HullIdx];
+			// Add the merged part
+			FKConvexElem& Convex = ToReturn.AggGeom.ConvexElems.Emplace_GetRef();
+			Convex.VertexData = ConvexDecomposition.GetVertices<double>(HullIdx);
+			Convex.UpdateElemBox(); // Note: In addition to updating the bounding box, this also re-computes hull indices.
+		}
+
+	});
+	return ToReturn;
+
+}
+
+
 TArray<FSphere> UGeometryScriptLibrary_CollisionFunctions::Conv_GeometryScriptSphereCoveringToSphereArray(const FGeometryScriptSphereCovering& SphereCovering)
 {
 	TArray<FSphere> ToRet;
diff --git a/Engine/Plugins/Runtime/GeometryScripting/Source/GeometryScriptingCore/Public/GeometryScript/CollisionFunctions.h b/Engine/Plugins/Runtime/GeometryScripting/Source/GeometryScriptingCore/Public/GeometryScript/CollisionFunctions.h
--- a/Engine/Plugins/Runtime/GeometryScripting/Source/GeometryScriptingCore/Public/GeometryScript/CollisionFunctions.h
+++ b/Engine/Plugins/Runtime/GeometryScripting/Source/GeometryScriptingCore/Public/GeometryScript/CollisionFunctions.h
@@ -123,7 +123,9 @@ enum class ENegativeSpaceSampleMethod : uint8
 	// Place sample spheres in a uniform grid pattern
 	Uniform,
 	// Use voxel-based subtraction and offsetting methods to specifically target concavities
-	VoxelSearch
+	VoxelSearch,
+	// A more-principled version of VoxelSearch that attempts to target specifically the space that is reachable by characters at least as large as a MinRadius sphere
+	NavigableVoxelSearch
 };
 
 // Options controlling how to sample the negative space of shapes, e.g. to define a region that must be avoided when merging collision shapes
@@ -166,6 +168,35 @@ public:
 	double MinRadius = 10.0;
 };
 
+
+/**
+ * Settings to define the important regions for a convex decomposition to preserve for a given input shape.
+ * 
+ * Note: this is similar to FComputeNegativeSpaceOptions, but with better default behavior and more intuitive parameters.
+ */
+USTRUCT(BlueprintType)
+struct GEOMETRYSCRIPTINGCORE_API FNavigableConvexDecompositionOptions
+{
+	GENERATED_BODY()
+public:
+	/// Minimum radius of characters/manipulators that should be able to navigate an input shape
+	UPROPERTY(EditAnywhere, BlueprintReadWrite, Category = Options)
+	double MinRadius = 100;
+
+	/// Tolerance distance: convex decomposition should be no further than this from an input shape, in the navigable regions
+	UPROPERTY(EditAnywhere, BlueprintReadWrite, Category = Options)
+	double Tolerance = 10;
+
+	/// Whether to only consider navigable space that is accessible from outside the shape. (Note this parameter is called bOnlyConnectedToHull elsewhere.)
+	UPROPERTY(EditAnywhere, BlueprintReadWrite, Category = Options)
+	bool bIgnoreUnreachableInternalSpace = true;
+
+	/// Optional list of locations that we expect to be navigable
+	UPROPERTY(EditAnywhere, BlueprintReadWrite, Category = Options)
+	TArray<FVector> CustomNavigablePositions;
+};
+
+
 // Options controlling how collision shapes can be merged together
 USTRUCT(BlueprintType)
 struct GEOMETRYSCRIPTINGCORE_API FGeometryScriptMergeSimpleCollisionOptions
@@ -483,6 +514,22 @@ public:
 		UGeometryScriptDebug* Debug = nullptr
 	);
 
+	/**
+	 * Compute the 'navigable' convex decomposition of an input mesh surface, i.e. a convex decomposition 
+	 * appropriate for a character of (or larger than) a given size
+	 *
+	 * @param TargetMesh	Mesh to decompose to convex hulls
+	 * @param Options		Options controlling the convex decomposition
+	 * @return				The resulting convex hulls as simple collision shapes
+	 */
+	UFUNCTION(BlueprintCallable, Category = "GeometryScript|Collision")
+	static UPARAM(DisplayName = "Convex Decomposition") FGeometryScriptSimpleCollision ComputeNavigableConvexDecomposition(
+		const UDynamicMesh* TargetMesh,
+		const FNavigableConvexDecompositionOptions& Options,
+		UGeometryScriptDebug* Debug = nullptr);
+
+
+
 	/**
 	 * @return An array of the spheres in the given Sphere Covering
 	 */
diff --git a/Engine/Source/Runtime/GeometryCore/Private/CompGeom/ConvexDecomposition3.cpp b/Engine/Source/Runtime/GeometryCore/Private/CompGeom/ConvexDecomposition3.cpp
--- a/Engine/Source/Runtime/GeometryCore/Private/CompGeom/ConvexDecomposition3.cpp
+++ b/Engine/Source/Runtime/GeometryCore/Private/CompGeom/ConvexDecomposition3.cpp
@@ -1512,9 +1512,15 @@ void FConvexDecomposition3::FConvexPart::Compact()
 }
 
 void FConvexDecomposition3::InitializeFromMesh(const FDynamicMesh3& SourceMesh, bool bMergeEdges)
+{
+	FPreprocessMeshOptions Options;
+	Options.bMergeEdges = bMergeEdges;
+	InitializeFromMesh(SourceMesh, Options);
+}
+void FConvexDecomposition3::InitializeFromMesh(const FDynamicMesh3& SourceMesh, const FPreprocessMeshOptions& Options)
 {
 	Decomposition.Empty();
-	FConvexPart* Convex = new FConvexPart(SourceMesh, bMergeEdges, ResultTransform);
+	FConvexPart* Convex = new FConvexPart(SourceMesh, Options, ResultTransform);
 	if (Convex->IsFailed())
 	{
 		delete Convex;
@@ -1528,7 +1534,17 @@ FConvexDecomposition3::FConvexPart::FConvexPart(const FDynamicMesh3& SourceMesh,
 	// Copy out the source mesh
 	InternalGeo.Copy(SourceMesh, false, false, false, false);
 
-	InitializeFromInternalGeo(bMergeEdges, TransformOut);
+	FConvexDecomposition3::FPreprocessMeshOptions Options;
+	Options.bMergeEdges = bMergeEdges;
+	InitializeFromInternalGeo(Options, TransformOut);
+}
+
+FConvexDecomposition3::FConvexPart::FConvexPart(const FDynamicMesh3& SourceMesh, const FConvexDecomposition3::FPreprocessMeshOptions& Options, FTransformSRT3d& TransformOut)
+{
+	// Copy out the source mesh
+	InternalGeo.Copy(SourceMesh, false, false, false, false);
+
+	InitializeFromInternalGeo(Options, TransformOut);
 }
 
 FConvexDecomposition3::FConvexPart::FConvexPart(TArrayView<const FVector3f> Vertices, TArrayView<const FIntVector3> Faces, bool bMergeEdges, FTransformSRT3d& TransformOut, int32 FaceVertexOffset)
@@ -1542,28 +1558,44 @@ FConvexDecomposition3::FConvexPart::FConvexPart(TArrayView<const FVector3f> Vert
 		InternalGeo.AppendTriangle(FIndex3i(F.X + FaceVertexOffset, F.Y + FaceVertexOffset, F.Z + FaceVertexOffset));
 	}
 
-	InitializeFromInternalGeo(bMergeEdges, TransformOut);
+	FConvexDecomposition3::FPreprocessMeshOptions Options;
+	Options.bMergeEdges = bMergeEdges;
+	InitializeFromInternalGeo(Options, TransformOut);
 }
 
-void FConvexDecomposition3::FConvexPart::InitializeFromInternalGeo(bool bMergeEdges, FTransformSRT3d& TransformOut)
+void FConvexDecomposition3::FConvexPart::InitializeFromInternalGeo(const FConvexDecomposition3::FPreprocessMeshOptions& Preprocess, FTransformSRT3d& TransformOut)
 {
-	// Transform the mesh to a standard unit-cube-at-origin space, so threshold have a consistent meaning
 	FAxisAlignedBox3d InitialBounds = InternalGeo.GetBounds();
 	double InvScaleFactor = FMath::Clamp(InitialBounds.MaxDim(), KINDA_SMALL_NUMBER, 1e8);
 	double ScaleFactor = 1.0 / InvScaleFactor;
+
+	// Weld close edges so we can sample convex edges
+	if (Preprocess.bMergeEdges)
+	{
+		FMergeCoincidentMeshEdges MergeEdges(&InternalGeo);
+		// scale the tolerance by the bounds size, for consistency across input scales
+		MergeEdges.MergeVertexTolerance = FMath::Max(MergeEdges.MergeVertexTolerance * ScaleFactor, FMathd::Epsilon);
+		MergeEdges.Apply();
+	}
+
+	// Apply any custom preprocessing (e.g., simplification, mesh repair, etc)
+	if (Preprocess.CustomPreprocess)
+	{
+		Preprocess.CustomPreprocess(InternalGeo, InitialBounds);
+		
+		// Re-compute bounds and scale factors, as the preprocess might have changed them
+		InitialBounds = InternalGeo.GetBounds();
+		InvScaleFactor = FMath::Clamp(InitialBounds.MaxDim(), KINDA_SMALL_NUMBER, 1e8);
+		ScaleFactor = 1.0 / InvScaleFactor;
+	}
+
+	// Transform the mesh to a standard unit-cube-at-origin space, so thresholds have a consistent meaning
 	FTransformSRT3d MeshTransform(-ScaleFactor * InitialBounds.Center());
 	MeshTransform.SetScale(ScaleFactor * MeshTransform.GetScale());
 	MeshTransforms::ApplyTransform(InternalGeo, MeshTransform);
 	// Return the inverse transform by reference, so we can put the results back into the original space
 	TransformOut = FTransformSRT3d(InitialBounds.Center());
 	TransformOut.SetScale(FVector3d(InvScaleFactor, InvScaleFactor, InvScaleFactor));
-	
-	// Weld close edges so we can sample convex edges
-	if (bMergeEdges)
-	{
-		FMergeCoincidentMeshEdges MergeEdges(&InternalGeo);
-		MergeEdges.Apply();
-	}
 
 	// Compute hull and standard measurements (volume, center, bounds, etc)
 	ComputeHull();
@@ -1786,7 +1818,7 @@ bool FConvexDecomposition3::SplitWorstHelper(bool bCanSkipUnreliableGeoVolumes,
 	// stop early if there are no negative-space overlaps, and we're only splitting in those cases
 	if (!bHasOverlapsNegative && bOnlySplitIfNegativeSpaceCovered)
 	{
-		return 0;
+		return true;
 	}
 	for (int32 PartIdx = 0; PartIdx < Decomposition.Num(); PartIdx++)
 	{
diff --git a/Engine/Source/Runtime/GeometryCore/Public/CompGeom/ConvexDecomposition3.h b/Engine/Source/Runtime/GeometryCore/Public/CompGeom/ConvexDecomposition3.h
--- a/Engine/Source/Runtime/GeometryCore/Public/CompGeom/ConvexDecomposition3.h
+++ b/Engine/Source/Runtime/GeometryCore/Public/CompGeom/ConvexDecomposition3.h
@@ -275,6 +275,34 @@ public:
 		InitializeFromMesh(SourceMesh, bMergeEdges);
 	}
 
+	struct FPreprocessMeshOptions
+	{
+		// Whether to attempt to merge boundary edges before processing
+		bool bMergeEdges = true;
+
+		// CustomPreprocess takes the mesh to process, and that mesh's bounding box, and updates the mesh in place
+		// Applied after discarding attributes and applying Merge Edges (if enabled)
+		// Can be used e.g. to simplify/repair/solidify the input
+		TUniqueFunction<void(FDynamicMesh3& Mesh, const FAxisAlignedBox3d& Bounds)> CustomPreprocess = nullptr;
+	};
+
+	FConvexDecomposition3(const FDynamicMesh3& SourceMesh, const FPreprocessMeshOptions& Options)
+	{
+		InitializeFromMesh(SourceMesh, Options);
+	}
+
+	/**
+	 * Test whether the input mesh (after preprocessing) is solid. Must be called after initialization, but before decomposition is computed.
+	 */
+	bool IsInputSolid()
+	{
+		if (!ensureMsgf(Decomposition.Num() == 1 && !Decomposition[0].IsCompact(), TEXT("This method must be called after initializing with the input mesh, and before the decomposition has been computed.")))
+		{
+			return false;
+		}
+		return Decomposition[0].InternalGeo.IsClosed();
+	}
+
 	/**
 	 * Initialize from convex hulls allows the caller to only use the hull merging phase of the algorithm
 	 * @param NumHulls			Number of convex hulls in the initial decomposition
@@ -294,6 +322,7 @@ public:
 	GEOMETRYCORE_API void InitializeProximityFromDecompositionBoundingBoxOverlaps(double BoundsExpandByMinDimFactor, double BoundsExpandByMaxDimFactor, double MinBoundsExpand);
 
 	GEOMETRYCORE_API void InitializeFromMesh(const FDynamicMesh3& SourceMesh, bool bMergeEdges);
+	GEOMETRYCORE_API void InitializeFromMesh(const FDynamicMesh3& SourceMesh, const FPreprocessMeshOptions& Options);
 
 	/**
 	 * Initialize convex decomposition with a triangle index mesh
@@ -537,6 +566,7 @@ public:
 	{
 		FConvexPart() {}
 		FConvexPart(const FDynamicMesh3& SourceMesh, bool bMergeEdges, FTransformSRT3d& TransformOut);
+		FConvexPart(const FDynamicMesh3& SourceMesh, const FConvexDecomposition3::FPreprocessMeshOptions& Options, FTransformSRT3d& TransformOut);
 		FConvexPart(TArrayView<const FVector3f> Vertices, TArrayView<const FIntVector3> Faces, bool bMergeEdges, FTransformSRT3d& TransformOut, int32 FaceVertexOffset = 0);
 
 		// Allow direct construction of a compact part (e.g. to allow construction of a pre-existing convex hull)
@@ -637,7 +667,7 @@ public:
 
 	private:
 		// helper to initialize the part from the already-set InternalGeo member; used to support constructors
-		void InitializeFromInternalGeo(bool bMergeEdges, FTransformSRT3d& TransformOut);
+		void InitializeFromInternalGeo(const FConvexDecomposition3::FPreprocessMeshOptions& Preprocess, FTransformSRT3d& TransformOut);
 	};
 
 	// All convex hulls parts in the convex decomposition
