CIS fix for incorrect loop conditions on geometry script convex decomposition method + change the condition to guard against infinite loop

diff --git a/Engine/Plugins/Runtime/GeometryScripting/Source/GeometryScriptingCore/Private/CollisionFunctions.cpp b/Engine/Plugins/Runtime/GeometryScripting/Source/GeometryScriptingCore/Private/CollisionFunctions.cpp
--- a/Engine/Plugins/Runtime/GeometryScripting/Source/GeometryScriptingCore/Private/CollisionFunctions.cpp
+++ b/Engine/Plugins/Runtime/GeometryScripting/Source/GeometryScriptingCore/Private/CollisionFunctions.cpp
@@ -1359,16 +1359,20 @@ FGeometryScriptSimpleCollision UGeometryScriptLibrary_CollisionFunctions::Comput
 		ConvexDecomposition.bSplitDisconnectedComponents = false;
 		ConvexDecomposition.ConvexEdgeAngleMoreSamplesThreshold = 180;
 		ConvexDecomposition.ThickenAfterHullFailure = FMath::Max(FMathd::ZeroTolerance, NegativeSpaceSettings.ReduceRadiusMargin * .01);
-		int32 NumSplits = -1;
-		for (int32 Split = 0; Split < NumSplits || NumSplits < 0; Split++)
+		constexpr int32 MaxAllowedSplits = 1000000; // more parts than any expected / reasonable decomposition
+		for (int32 Split = 0; ; Split++)
 		{
-
 			int32 NumSplit = ConvexDecomposition.SplitWorst(false, -1, true, NegativeSpaceSettings.ReduceRadiusMargin * .5);
 
 			if (NumSplit == 0)
 			{
 				break;
 			}
+
+			if (!ensureMsgf(Split < MaxAllowedSplits, TEXT("Convex decomposition split the input %d times; likely stuck in a loop"), Split))
+			{
+				break;
+			}
 		}
 
 		ConvexDecomposition.FixHullOverlapsInNegativeSpace();
