make mesh asset functions' static mesh to dynamic mesh conversion use the same conversion code as the scene component conversion

diff --git a/Engine/Plugins/Runtime/GeometryScripting/Source/GeometryScriptingCore/Private/MeshAssetFunctions.cpp b/Engine/Plugins/Runtime/GeometryScripting/Source/GeometryScriptingCore/Private/MeshAssetFunctions.cpp
--- a/Engine/Plugins/Runtime/GeometryScripting/Source/GeometryScriptingCore/Private/MeshAssetFunctions.cpp
+++ b/Engine/Plugins/Runtime/GeometryScripting/Source/GeometryScriptingCore/Private/MeshAssetFunctions.cpp
@@ -22,6 +22,7 @@
 #include "DynamicMeshToMeshDescription.h"
 #include "StaticMeshLODResourcesToDynamicMesh.h"
 #include "AssetUtils/StaticMeshMaterialUtil.h"
+#include "ConversionUtils/SceneComponentToDynamicMesh.h"
 
 #include UE_INLINE_GENERATED_CPP_BY_NAME(MeshAssetFunctions)
 
@@ -36,178 +37,37 @@ using namespace UE::Geometry;
 #define LOCTEXT_NAMESPACE "UGeometryScriptLibrary_MeshAssetFunctions"
 
 
-
-
-static UDynamicMesh* CopyMeshFromStaticMesh_SourceData(	
-	UStaticMesh* FromStaticMeshAsset, 
-	UDynamicMesh* ToDynamicMesh, 
-	FGeometryScriptCopyMeshFromAssetOptions AssetOptions,
-	FGeometryScriptMeshReadLOD RequestedLOD,
-	EGeometryScriptOutcomePins& Outcome,
-	UGeometryScriptDebug* Debug
-)
+static void ConvertGeometryScriptReadLOD(const FGeometryScriptMeshReadLOD& ReadLOD, UE::Conversion::EMeshLODType& OutLODType, int32& OutLODIndex)
 {
-	if (RequestedLOD.LODType != EGeometryScriptLODType::MaxAvailable && RequestedLOD.LODType != EGeometryScriptLODType::SourceModel && RequestedLOD.LODType != EGeometryScriptLODType::HiResSourceModel)
-	{
-		UE::Geometry::AppendError(Debug, EGeometryScriptErrorType::InvalidInputs, LOCTEXT("CopyMeshFromStaticMesh_LODNotAvailable", "CopyMeshFromStaticMesh: Requested LOD Type is not available"));
-		return ToDynamicMesh;
-	}
-
-#if WITH_EDITOR
-	if (RequestedLOD.LODType == EGeometryScriptLODType::HiResSourceModel && FromStaticMeshAsset->IsHiResMeshDescriptionValid() == false)
-	{
-		UE::Geometry::AppendError(Debug, EGeometryScriptErrorType::InvalidInputs, LOCTEXT("CopyMeshFromStaticMesh_HiResLODNotAvailable", "CopyMeshFromStaticMesh: HiResSourceModel LOD Type is not available"));
-		return ToDynamicMesh;
-	}
-
-	const FMeshDescription* SourceMesh = nullptr;
-	const FMeshBuildSettings* BuildSettings = nullptr;
-
-	if ((RequestedLOD.LODType == EGeometryScriptLODType::HiResSourceModel) ||
-		(RequestedLOD.LODType == EGeometryScriptLODType::MaxAvailable && FromStaticMeshAsset->IsHiResMeshDescriptionValid()))
-	{
-		SourceMesh = FromStaticMeshAsset->GetHiResMeshDescription();
-		const FStaticMeshSourceModel& SourceModel = FromStaticMeshAsset->GetHiResSourceModel();
-		BuildSettings = &SourceModel.BuildSettings;
-	}
-	else
-	{
-		int32 UseLODIndex = FMath::Clamp(RequestedLOD.LODIndex, 0, FromStaticMeshAsset->GetNumSourceModels() - 1);
-		SourceMesh = FromStaticMeshAsset->GetMeshDescription(UseLODIndex);
-		const FStaticMeshSourceModel& SourceModel = FromStaticMeshAsset->GetSourceModel(UseLODIndex);
-		BuildSettings = &SourceModel.BuildSettings;
-	}
-
-	if (SourceMesh == nullptr)
-	{
-		UE::Geometry::AppendError(Debug, EGeometryScriptErrorType::InvalidInputs, LOCTEXT("CopyMeshFromStaticMesh_SourceLODIsNull", "CopyMeshFromStaticMesh: Requested SourceModel LOD is null, only RenderData Mesh is available"));
-		return ToDynamicMesh;
-	}
-
-	bool bHasDirtyBuildSettings = BuildSettings->bRecomputeNormals
-		|| (BuildSettings->bRecomputeTangents && AssetOptions.bRequestTangents);
-	bool bNeedsBuildScale = AssetOptions.bUseBuildScale && BuildSettings && !BuildSettings->BuildScale3D.Equals(FVector::OneVector);
-	bool bNeedsOtherBuildSettings = AssetOptions.bApplyBuildSettings && bHasDirtyBuildSettings;
-
-	FMeshDescription LocalSourceMeshCopy;
-	if (bNeedsBuildScale || bNeedsOtherBuildSettings)
-	{
-		LocalSourceMeshCopy = *SourceMesh;
-
-		FStaticMeshAttributes Attributes(LocalSourceMeshCopy);
-
-		if (bNeedsBuildScale)
+	using namespace UE::Conversion;
+	OutLODType = [](const EGeometryScriptLODType& LODType)
 		{
-			FTransform BuildScaleTransform = FTransform::Identity;
-			BuildScaleTransform.SetScale3D(BuildSettings->BuildScale3D);
-			FStaticMeshOperations::ApplyTransform(LocalSourceMeshCopy, BuildScaleTransform, true /*use correct normal transforms*/);
-		}
-
-		if (bNeedsOtherBuildSettings)
-		{
-			if (!Attributes.GetTriangleNormals().IsValid() || !Attributes.GetTriangleTangents().IsValid())
-			{
-				// If these attributes don't exist, create them and compute their values for each triangle
-				FStaticMeshOperations::ComputeTriangleTangentsAndNormals(LocalSourceMeshCopy);
-			}
-
-			EComputeNTBsFlags ComputeNTBsOptions = EComputeNTBsFlags::BlendOverlappingNormals;
-			ComputeNTBsOptions |= BuildSettings->bRecomputeNormals ? EComputeNTBsFlags::Normals : EComputeNTBsFlags::None;
-			if (AssetOptions.bRequestTangents)
-			{
-				ComputeNTBsOptions |= BuildSettings->bRecomputeTangents ? EComputeNTBsFlags::Tangents : EComputeNTBsFlags::None;
-				ComputeNTBsOptions |= BuildSettings->bUseMikkTSpace ? EComputeNTBsFlags::UseMikkTSpace : EComputeNTBsFlags::None;
-			}
-			ComputeNTBsOptions |= BuildSettings->bComputeWeightedNormals ? EComputeNTBsFlags::WeightedNTBs : EComputeNTBsFlags::None;
-			if (AssetOptions.bIgnoreRemoveDegenerates == false)
+			switch (LODType)
 			{
-				ComputeNTBsOptions |= BuildSettings->bRemoveDegenerates ? EComputeNTBsFlags::IgnoreDegenerateTriangles : EComputeNTBsFlags::None;
+			case EGeometryScriptLODType::MaxAvailable:
+				return EMeshLODType::MaxAvailable;
+			case EGeometryScriptLODType::HiResSourceModel:
+				return EMeshLODType::HiResSourceModel;
+			case EGeometryScriptLODType::SourceModel:
+				return EMeshLODType::SourceModel;
+			case EGeometryScriptLODType::RenderData:
+				return EMeshLODType::RenderData;
+			default:
+				checkNoEntry();
+				return EMeshLODType::RenderData;
 			}
-
-			FStaticMeshOperations::ComputeTangentsAndNormals(LocalSourceMeshCopy, ComputeNTBsOptions);
-		}
-
-		SourceMesh = &LocalSourceMeshCopy;
-	}
-
-	FDynamicMesh3 NewMesh;
-	FMeshDescriptionToDynamicMesh Converter;
-	Converter.Convert(SourceMesh, NewMesh, AssetOptions.bRequestTangents);
-
-	ToDynamicMesh->SetMesh(MoveTemp(NewMesh));
-
-	Outcome = EGeometryScriptOutcomePins::Success;
-#else
-	UE::Geometry::AppendError(Debug, EGeometryScriptErrorType::InvalidInputs, LOCTEXT("CopyMeshFromAsset_EditorOnly", "CopyMeshFromStaticMesh: Source Models are not available at Runtime"));
-#endif
-
-	return ToDynamicMesh;
+		}(ReadLOD.LODType);
+	OutLODIndex = ReadLOD.LODIndex;
 }
 
-
-
-static UDynamicMesh* CopyMeshFromStaticMesh_RenderData(	
-	UStaticMesh* FromStaticMeshAsset, 
-	UDynamicMesh* ToDynamicMesh, 
-	FGeometryScriptCopyMeshFromAssetOptions AssetOptions,
-	FGeometryScriptMeshReadLOD RequestedLOD,
-	EGeometryScriptOutcomePins& Outcome,
-	UGeometryScriptDebug* Debug
-)
+static void ConvertGeometryScriptWriteLOD(const FGeometryScriptMeshWriteLOD& WriteLOD, UE::Conversion::EMeshLODType& OutLODType, int32& OutLODIndex)
 {
-	if (RequestedLOD.LODType != EGeometryScriptLODType::MaxAvailable && RequestedLOD.LODType != EGeometryScriptLODType::RenderData)
-	{
-		UE::Geometry::AppendError(Debug, EGeometryScriptErrorType::InvalidInputs, LOCTEXT("CopyMeshFromStaticMesh_LODNotAvailable", "CopyMeshFromStaticMesh: Requested LOD Type is not available"));
-		return ToDynamicMesh;
-	}
-
-#if !WITH_EDITOR
-	if (FromStaticMeshAsset->bAllowCPUAccess == false)
-	{
-		UE::Geometry::AppendError(Debug, EGeometryScriptErrorType::InvalidInputs, LOCTEXT("CopyMeshFromStaticMesh_CPUAccess", "CopyMeshFromStaticMesh: StaticMesh bAllowCPUAccess must be set to true to read mesh data at Runtime"));
-		return ToDynamicMesh;
-	}
-#endif
-
-	int32 UseLODIndex = FMath::Clamp(RequestedLOD.LODIndex, 0, FromStaticMeshAsset->GetNumLODs() - 1);
-
-	const FStaticMeshLODResources* LODResources = nullptr;
-	if (FStaticMeshRenderData* RenderData = FromStaticMeshAsset->GetRenderData())
-	{
-		LODResources = &RenderData->LODResources[UseLODIndex];
-	}
-	if (LODResources == nullptr)
-	{
-		UE::Geometry::AppendError(Debug, EGeometryScriptErrorType::InvalidInputs, LOCTEXT("CopyMeshFromStaticMesh_NoLODResources", "CopyMeshFromStaticMesh: LOD Data is not available"));
-		return ToDynamicMesh;
-	}
-
-	FStaticMeshLODResourcesToDynamicMesh::ConversionOptions ConvertOptions;
-#if WITH_EDITOR
-	if (AssetOptions.bUseBuildScale)
-	{
-		// respect BuildScale build setting
-		const FMeshBuildSettings& LODBuildSettings = FromStaticMeshAsset->GetSourceModel(UseLODIndex).BuildSettings;
-		ConvertOptions.BuildScale = (FVector3d)LODBuildSettings.BuildScale3D;
-	}
-#else
-	if (!AssetOptions.bUseBuildScale)
-	{
-		UE::Geometry::AppendError(Debug, EGeometryScriptErrorType::InvalidInputs, LOCTEXT("CopyMeshFromStaticMesh_BuildScaleAlreadyBaked", "CopyMeshFromStaticMesh: Requested mesh without BuildScale, but BuildScale is already baked into the RenderData."));
-	}
-#endif
-
-	FDynamicMesh3 NewMesh;
-	FStaticMeshLODResourcesToDynamicMesh Converter;
-	Converter.Convert(LODResources, ConvertOptions, NewMesh);
-
-	ToDynamicMesh->SetMesh(MoveTemp(NewMesh));
-	Outcome = EGeometryScriptOutcomePins::Success;
-	return ToDynamicMesh;
+	using namespace UE::Conversion;
+	OutLODType = WriteLOD.bWriteHiResSource ? EMeshLODType::HiResSourceModel : EMeshLODType::SourceModel;
+	OutLODIndex = WriteLOD.LODIndex;
 }
 
 
-
 UDynamicMesh*  UGeometryScriptLibrary_StaticMeshFunctions::CopyMeshFromStaticMesh(
 	UStaticMesh* FromStaticMeshAsset, 
 	UDynamicMesh* ToDynamicMesh, 
@@ -229,18 +89,30 @@ UDynamicMesh*  UGeometryScriptLibrary_StaticMeshFunctions::CopyMeshFromStaticMes
 		return ToDynamicMesh;
 	}
 
-#if WITH_EDITOR
-	if (RequestedLOD.LODType == EGeometryScriptLODType::RenderData)
+	using namespace UE::Conversion;
+	EMeshLODType LODType;
+	int32 LODIndex;
+	ConvertGeometryScriptReadLOD(RequestedLOD, LODType, LODIndex);
+	FStaticMeshConversionOptions ConversionOptions;
+	ConversionOptions.bApplyBuildSettings = AssetOptions.bApplyBuildSettings;
+	ConversionOptions.bRequestTangents = AssetOptions.bRequestTangents;
+	ConversionOptions.bIgnoreRemoveDegenerates = AssetOptions.bIgnoreRemoveDegenerates;
+	ConversionOptions.bUseBuildScale = AssetOptions.bUseBuildScale;
+
+	FText ErrorMessage;
+
+	FDynamicMesh3 NewMesh;
+	bool bSuccess = StaticMeshToDynamicMesh(FromStaticMeshAsset, NewMesh, ErrorMessage, ConversionOptions, LODType, LODIndex);
+	if (!bSuccess)
 	{
-		return CopyMeshFromStaticMesh_RenderData(FromStaticMeshAsset, ToDynamicMesh, AssetOptions, RequestedLOD, Outcome, Debug);
+		UE::Geometry::AppendError(Debug, EGeometryScriptErrorType::InvalidInputs, ErrorMessage);
 	}
 	else
 	{
-		return CopyMeshFromStaticMesh_SourceData(FromStaticMeshAsset, ToDynamicMesh, AssetOptions, RequestedLOD, Outcome, Debug);
+		ToDynamicMesh->SetMesh(MoveTemp(NewMesh));
+		Outcome = EGeometryScriptOutcomePins::Success;
 	}
-#else
-	return CopyMeshFromStaticMesh_RenderData(FromStaticMeshAsset, ToDynamicMesh, AssetOptions, RequestedLOD, Outcome, Debug);	
-#endif
+	return ToDynamicMesh;
 }
 
 
@@ -384,6 +256,14 @@ UDynamicMesh*  UGeometryScriptLibrary_StaticMeshFunctions::CopyMeshToStaticMesh(
 
 		FConversionToMeshDescriptionOptions ConversionOptions;
 		FDynamicMeshToMeshDescription Converter(ConversionOptions);
+		if (!Options.bReplaceMaterials)
+		{
+			UE::Conversion::EMeshLODType LODType;
+			int32 LODIndex;
+			ConvertGeometryScriptWriteLOD(TargetLOD, LODType, LODIndex);
+			TArray<int32> MaterialIDMap = UE::Conversion::GetPolygonGroupToMaterialIndexMap(ToStaticMeshAsset, LODType, LODIndex);
+			Converter.SetMaterialIDMapFromInverseMap(MaterialIDMap);
+		}
 		FromDynamicMesh->ProcessMesh([&](const FDynamicMesh3& ReadMesh)
 		{
 			Converter.Convert(&ReadMesh, *MeshDescription, !Options.bEnableRecomputeTangents);
@@ -422,6 +302,7 @@ UDynamicMesh*  UGeometryScriptLibrary_StaticMeshFunctions::CopyMeshToStaticMesh(
 
 			// Reset the section info map
 			ToStaticMeshAsset->GetSectionInfoMap().Clear();
+			// TODO: populate the section info map here
 		}
 
 		ToStaticMeshAsset->CommitMeshDescription(UseLODIndex);
diff --git a/Engine/Plugins/Runtime/MeshModelingToolset/Source/ModelingComponents/Private/ConversionUtils/SceneComponentToDynamicMesh.cpp b/Engine/Plugins/Runtime/MeshModelingToolset/Source/ModelingComponents/Private/ConversionUtils/SceneComponentToDynamicMesh.cpp
--- a/Engine/Plugins/Runtime/MeshModelingToolset/Source/ModelingComponents/Private/ConversionUtils/SceneComponentToDynamicMesh.cpp
+++ b/Engine/Plugins/Runtime/MeshModelingToolset/Source/ModelingComponents/Private/ConversionUtils/SceneComponentToDynamicMesh.cpp
@@ -94,23 +94,6 @@ namespace Private::ConversionHelper
 	// Static mesh conversion functions (from geometry script MeshAssetFunctions.cpp)
 	// TODO: these static mesh conversion helpers should be pulled out to their own StaticMeshToDynamicMesh converter method
 
-	struct FStaticMeshConversionOptions
-	{
-		// Whether to apply Build Settings during the mesh copy.
-		bool bApplyBuildSettings = true;
-
-		// Whether to request tangents on the copied mesh. If tangents are not requested, tangent-related build settings will also be ignored.
-		bool bRequestTangents = true;
-
-		// Whether to ignore the 'remove degenerates' option from Build Settings. Note: Only applies if 'Apply Build Settings' is enabled.
-		bool bIgnoreRemoveDegenerates = true;
-
-		// Whether to scale the copied mesh by the Build Setting's 'Build Scale'. Note: This is considered separately from the 'Apply Build Settings' option.
-		bool bUseBuildScale = true;
-
-		// Whether to request the vertex colors of the component instancing the static mesh, rather than the static mesh asset
-		bool bRequestInstanceVertexColors = true;
-	};
 
 	// helper for the material ID remapping used for source LODs
 	// note: returns empty array if no remapping needed (or if not WITH_EDITOR)
@@ -307,6 +290,7 @@ namespace Private::ConversionHelper
 		FStaticMeshConversionOptions AssetOptions,
 		EMeshLODType LODType,
 		int32 LODIndex,
+		bool bRequestInstanceVertexColors,
 		FDynamicMesh3& OutMesh,
 		FText& OutErrorMessage
 	)
@@ -359,7 +343,7 @@ namespace Private::ConversionHelper
 #endif
 
 		FStaticMeshLODResourcesToDynamicMesh Converter;
-		if (AssetOptions.bRequestInstanceVertexColors && StaticMeshComponent && StaticMeshComponent->LODData.IsValidIndex(UseLODIndex))
+		if (bRequestInstanceVertexColors && StaticMeshComponent && StaticMeshComponent->LODData.IsValidIndex(UseLODIndex))
 		{
 			FStaticMeshComponentLODInfo* InstanceMeshLODInfo = &StaticMeshComponent->LODData[UseLODIndex];
 			const bool bValidInstanceData = InstanceMeshLODInfo
@@ -387,6 +371,7 @@ namespace Private::ConversionHelper
 		EMeshLODType LODType,
 		int32 LODIndex,
 		bool bUseClosestLOD,
+		bool bRequestInstanceVertexColors,
 		FDynamicMesh3& OutMesh,
 		FText& OutErrorMessage
 	)
@@ -433,7 +418,7 @@ namespace Private::ConversionHelper
 
 		if (LODType == EMeshLODType::RenderData)
 		{
-			return CopyMeshFromStaticMesh_RenderData(FromStaticMeshAsset, StaticMeshComponent, AssetOptions, LODType, LODIndex, OutMesh, OutErrorMessage);
+			return CopyMeshFromStaticMesh_RenderData(FromStaticMeshAsset, StaticMeshComponent, AssetOptions, LODType, LODIndex, bRequestInstanceVertexColors, OutMesh, OutErrorMessage);
 		}
 		else
 		{
@@ -462,6 +447,15 @@ TArray<int32> GetPolygonGroupToMaterialIndexMap(const UStaticMesh* StaticMesh, E
 #endif
 }
 
+bool StaticMeshToDynamicMesh(UStaticMesh* InMesh, Geometry::FDynamicMesh3& OutMesh, FText& OutErrorMessage,
+	const FStaticMeshConversionOptions& ConversionOptions, EMeshLODType LODType, int32 LODIndex, bool bUseClosestLOD)
+{
+	constexpr UStaticMeshComponent* StaticMeshComponent = nullptr; // ok to leave this null when converting from asset
+	constexpr bool bRequestInstanceVertexColors = false; // cannot request instance colors from the asset
+	return Private::ConversionHelper::CopyMeshFromStaticMesh(
+		InMesh, StaticMeshComponent, ConversionOptions, LODType, LODIndex, bUseClosestLOD, bRequestInstanceVertexColors, OutMesh, OutErrorMessage);
+}
+
 bool SceneComponentToDynamicMesh(USceneComponent* Component, const FToMeshOptions& Options, bool bTransformToWorld, 
 	Geometry::FDynamicMesh3& OutMesh, FTransform& OutLocalToWorld, FText& OutErrorMessage,
 	TArray<UMaterialInterface*>* OutComponentMaterials, TArray<UMaterialInterface*>* OutAssetMaterials)
@@ -543,12 +537,11 @@ bool SceneComponentToDynamicMesh(USceneComponent* Component, const FToMeshOption
 		UStaticMesh* StaticMesh = SplineMeshComponent->GetStaticMesh();
 		if (StaticMesh)
 		{
-			Private::ConversionHelper::FStaticMeshConversionOptions AssetOptions;
+			FStaticMeshConversionOptions AssetOptions;
 			AssetOptions.bApplyBuildSettings = (Options.bWantNormals || Options.bWantTangents);
 			AssetOptions.bRequestTangents = Options.bWantTangents;
-			AssetOptions.bRequestInstanceVertexColors = Options.bWantInstanceColors;
 			bSuccess = Private::ConversionHelper::CopyMeshFromStaticMesh(
-				StaticMesh, SplineMeshComponent, AssetOptions, Options.LODType, Options.LODIndex, Options.bUseClosestLOD, OutMesh, OutErrorMessage);
+				StaticMesh, SplineMeshComponent, AssetOptions, Options.LODType, Options.LODIndex, Options.bUseClosestLOD, Options.bWantInstanceColors, OutMesh, OutErrorMessage);
 
 			// deform the dynamic mesh and its tangent space with the spline
 			if (bSuccess)
@@ -577,12 +570,12 @@ bool SceneComponentToDynamicMesh(USceneComponent* Component, const FToMeshOption
 		UStaticMesh* StaticMesh = StaticMeshComponent->GetStaticMesh();
 		if (StaticMesh)
 		{
-			Private::ConversionHelper::FStaticMeshConversionOptions AssetOptions;
+			FStaticMeshConversionOptions AssetOptions;
 			AssetOptions.bApplyBuildSettings = (Options.bWantNormals || Options.bWantTangents);
 			AssetOptions.bRequestTangents = Options.bWantTangents;
-			AssetOptions.bRequestInstanceVertexColors = Options.bWantInstanceColors;
+			bool bRequestInstanceVertexColors = Options.bWantInstanceColors;
 			bSuccess = Private::ConversionHelper::CopyMeshFromStaticMesh(
-				StaticMesh, StaticMeshComponent, AssetOptions, Options.LODType, Options.LODIndex, Options.bUseClosestLOD, OutMesh, OutErrorMessage);
+				StaticMesh, StaticMeshComponent, AssetOptions, Options.LODType, Options.LODIndex, Options.bUseClosestLOD, bRequestInstanceVertexColors, OutMesh, OutErrorMessage);
 
 			// if we have an ISMC, append instances
 			if (UInstancedStaticMeshComponent* ISMComponent = Cast<UInstancedStaticMeshComponent>(StaticMeshComponent))
diff --git a/Engine/Plugins/Runtime/MeshModelingToolset/Source/ModelingComponents/Public/ConversionUtils/SceneComponentToDynamicMesh.h b/Engine/Plugins/Runtime/MeshModelingToolset/Source/ModelingComponents/Public/ConversionUtils/SceneComponentToDynamicMesh.h
--- a/Engine/Plugins/Runtime/MeshModelingToolset/Source/ModelingComponents/Public/ConversionUtils/SceneComponentToDynamicMesh.h
+++ b/Engine/Plugins/Runtime/MeshModelingToolset/Source/ModelingComponents/Public/ConversionUtils/SceneComponentToDynamicMesh.h
@@ -85,6 +85,37 @@ namespace Conversion {
 		USceneComponent* SceneComponent, const FToMeshOptions& ConversionOptions, bool bTransformToWorld, 
 		Geometry::FDynamicMesh3& OutMesh, FTransform& OutTransform, FText& OutErrorMessage, 
 		TArray<UMaterialInterface*>* OutComponentMaterials = nullptr, TArray<UMaterialInterface*>* OutAssetMaterials = nullptr);
+
+	// Options for converting Static Mesh assets
+	struct FStaticMeshConversionOptions
+	{
+		// Whether to apply Build Settings during the mesh copy.
+		bool bApplyBuildSettings = true;
+
+		// Whether to request tangents on the copied mesh. If tangents are not requested, tangent-related build settings will also be ignored.
+		bool bRequestTangents = true;
+
+		// Whether to ignore the 'remove degenerates' option from Build Settings. Note: Only applies if 'Apply Build Settings' is enabled.
+		bool bIgnoreRemoveDegenerates = true;
+
+		// Whether to scale the copied mesh by the Build Setting's 'Build Scale'. Note: This is considered separately from the 'Apply Build Settings' option.
+		bool bUseBuildScale = true;
+	};
+
+	/**
+	 * Converts geometry from Static Mesh to Dynamic Mesh
+	 * 
+	 * @param InMesh			Static Mesh to convert
+	 * @param OutMesh			Dynamic Mesh to contain static mesh geometry
+	 * @param OutErrorMessage	If method returns false, contains error message
+	 * @param ConversionOptions	Static mesh conversion settings
+	 * @param LODType			Source of data to convert (e.g., render data, source data, or hi res source data)
+	 * @param LODIndex			LOD index to convert
+	 * @param bUseClosestLOD	Whether to remap to the most similar available LOD rather than fail, if the requested LOD type is not available
+	 */
+	bool MODELINGCOMPONENTS_API StaticMeshToDynamicMesh(
+		UStaticMesh* InMesh, Geometry::FDynamicMesh3& OutMesh, FText& OutErrorMessage, const FStaticMeshConversionOptions& ConversionOptions,
+		EMeshLODType LODType = EMeshLODType::SourceModel, int32 LODIndex = 0, bool bUseClosestLOD = true);
 	
 	/**
 	 * Helper to quickly test if we expect to be able to convert SceneComponent to a dynamic mesh
