repopulate material section info when replacing materials in Geometry Script copy to static mesh method

diff --git a/Engine/Plugins/Runtime/GeometryScripting/Source/GeometryScriptingCore/Private/MeshAssetFunctions.cpp b/Engine/Plugins/Runtime/GeometryScripting/Source/GeometryScriptingCore/Private/MeshAssetFunctions.cpp
--- a/Engine/Plugins/Runtime/GeometryScripting/Source/GeometryScriptingCore/Private/MeshAssetFunctions.cpp
+++ b/Engine/Plugins/Runtime/GeometryScripting/Source/GeometryScriptingCore/Private/MeshAssetFunctions.cpp
@@ -300,9 +300,47 @@ UDynamicMesh*  UGeometryScriptLibrary_StaticMeshFunctions::CopyMeshToStaticMesh(
 
 			ToStaticMeshAsset->SetStaticMaterials(NewMaterials);
 
+			// Set material slot names on the mesh description
+			FStaticMeshAttributes Attributes(*MeshDescription);
+			TPolygonGroupAttributesRef<FName> PolygonGroupImportedMaterialSlotNames = Attributes.GetPolygonGroupMaterialSlotNames();
+			for (int32 SlotIdx = 0; SlotIdx < NewMaterials.Num(); ++SlotIdx)
+			{
+				if (SlotIdx < PolygonGroupImportedMaterialSlotNames.GetNumElements())
+				{
+					PolygonGroupImportedMaterialSlotNames.Set(SlotIdx, NewMaterials[SlotIdx].ImportedMaterialSlotName);
+				}
+			}
+
 			// Reset the section info map
 			ToStaticMeshAsset->GetSectionInfoMap().Clear();
-			// TODO: populate the section info map here
+			ToStaticMeshAsset->GetOriginalSectionInfoMap().Clear();
+
+			// Repopulate section info map
+			FMeshSectionInfoMap SectionInfoMap;
+			for (int32 LODIndex = 0, NumLODs = ToStaticMeshAsset->GetNumSourceModels(); LODIndex < NumLODs; ++LODIndex)
+			{
+				if (const FMeshDescription* Mesh = (LODIndex == UseLODIndex) ? MeshDescription : ToStaticMeshAsset->GetMeshDescription(LODIndex))
+				{
+					FStaticMeshConstAttributes MeshDescriptionAttributes(*Mesh);
+					TPolygonGroupAttributesConstRef<FName> MaterialSlotNames = MeshDescriptionAttributes.GetPolygonGroupMaterialSlotNames();
+					int32 SectionIndex = 0;
+					for (FPolygonGroupID PolygonGroupID : Mesh->PolygonGroups().GetElementIDs())
+					{
+						// Material index is either from the matching material slot name or the section index if that name is not found
+						int32 MaterialIndex = ToStaticMeshAsset->GetStaticMaterials().IndexOfByPredicate(
+							[&MaterialSlotName = MaterialSlotNames[PolygonGroupID]](const FStaticMaterial& StaticMaterial) { return StaticMaterial.MaterialSlotName == MaterialSlotName; }
+						);
+						if (MaterialIndex == INDEX_NONE)
+						{
+							MaterialIndex = SectionIndex;
+						}
+						SectionInfoMap.Set(LODIndex, SectionIndex, FMeshSectionInfo(MaterialIndex));
+						SectionIndex++;
+					}
+				}
+			}
+			ToStaticMeshAsset->GetSectionInfoMap().CopyFrom(SectionInfoMap);
+			ToStaticMeshAsset->GetOriginalSectionInfoMap().CopyFrom(SectionInfoMap);
 		}
 
 		ToStaticMeshAsset->CommitMeshDescription(UseLODIndex);
