Geometry Scripting: MeshBoneWeightFunctions: avoid invalid bone indices issues due to int to uint16 (FBoneIndexType) conversion

#rnx

diff --git a/Engine/Plugins/Experimental/MeshModelingToolsetExp/Source/MeshModelingToolsEditorOnlyExp/Private/SkeletalMesh/SkinWeightsPaintTool.cpp b/Engine/Plugins/Experimental/MeshModelingToolsetExp/Source/MeshModelingToolsEditorOnlyExp/Private/SkeletalMesh/SkinWeightsPaintTool.cpp
--- a/Engine/Plugins/Experimental/MeshModelingToolsetExp/Source/MeshModelingToolsEditorOnlyExp/Private/SkeletalMesh/SkinWeightsPaintTool.cpp
+++ b/Engine/Plugins/Experimental/MeshModelingToolsetExp/Source/MeshModelingToolsEditorOnlyExp/Private/SkeletalMesh/SkinWeightsPaintTool.cpp
@@ -337,7 +337,12 @@ void FSkinToolWeights::InitializeSkinWeights(
 		for (UE::AnimationCore::FBoneWeight BoneWeight: VertexSkinWeights.Get(VertexID))
 		{
 			check(InfluenceIndex < MAX_TOTAL_INFLUENCES);
-			const int32 BoneIndex = BoneWeight.GetBoneIndex();
+			int32 BoneIndex = BoneWeight.GetBoneIndex();
+			if (!ensure(Deformer.InvCSRefPoseTransforms.IsValidIndex(BoneIndex)))
+			{
+				UE_LOG(LogMeshModelingToolsEditor, Warning, TEXT("InitializeSkinWeights: Invalid bone index provided (%d); falling back to 0 as bone index."), BoneIndex);
+				BoneIndex = 0;
+			}
 			const float Weight = BoneWeight.GetWeight();
 			const FVector& RefPoseVertexPosition = Deformer.RefPoseVertexPositions[VertexIndex];
 			const FTransform& InvRefPoseTransform = Deformer.InvCSRefPoseTransforms[BoneIndex];
diff --git a/Engine/Plugins/Experimental/MeshModelingToolsetExp/Source/SkeletalMeshModifiers/Private/SkeletonModifier.cpp b/Engine/Plugins/Experimental/MeshModelingToolsetExp/Source/SkeletalMeshModifiers/Private/SkeletonModifier.cpp
--- a/Engine/Plugins/Experimental/MeshModelingToolsetExp/Source/SkeletalMeshModifiers/Private/SkeletonModifier.cpp
+++ b/Engine/Plugins/Experimental/MeshModelingToolsetExp/Source/SkeletalMeshModifiers/Private/SkeletonModifier.cpp
@@ -343,10 +343,17 @@ bool USkeletonModifier::CommitSkeletonToSkeletalMesh()
 					{
 						const FBoneWeight& OldBoneWeight = BoneWeights[Idx];
 						const int32 BoneIndex = OldBoneWeight.GetBoneIndex();
-					
-						check(BoneIndexTracker.IsValidIndex(BoneIndex));
-					
-						const int32 NewBoneIndex = BoneIndexTracker[BoneIndex];
+
+						int32 NewBoneIndex = 0;
+						if (ensure(BoneIndexTracker.IsValidIndex(BoneIndex)))
+						{
+							NewBoneIndex = BoneIndexTracker[BoneIndex];							
+						}
+						else
+						{
+							UE_LOG(LogAnimation, Warning, TEXT("Skeleton Modifier - Commit: Invalid bone index provided (%d); falling back to 0 as bone index."), BoneIndex);
+						}
+
 						if (NewBoneIndex != INDEX_NONE)
 						{
 							NewWeights.Add(FBoneWeight(NewBoneIndex, OldBoneWeight.GetRawWeight()));
diff --git a/Engine/Plugins/Runtime/GeometryScripting/Source/GeometryScriptingCore/Private/MeshBoneWeightFunctions.cpp b/Engine/Plugins/Runtime/GeometryScripting/Source/GeometryScriptingCore/Private/MeshBoneWeightFunctions.cpp
--- a/Engine/Plugins/Runtime/GeometryScripting/Source/GeometryScriptingCore/Private/MeshBoneWeightFunctions.cpp
+++ b/Engine/Plugins/Runtime/GeometryScripting/Source/GeometryScriptingCore/Private/MeshBoneWeightFunctions.cpp
@@ -241,7 +241,8 @@ UDynamicMesh* UGeometryScriptLibrary_MeshBoneWeightFunctions::SetVertexBoneWeigh
 	int VertexID,
 	const TArray<FGeometryScriptBoneWeight>& BoneWeights,
 	bool& bHasValidBoneWeights,
-	FGeometryScriptBoneWeightProfile Profile)
+	FGeometryScriptBoneWeightProfile Profile,
+	UGeometryScriptDebug* Debug)
 {
 	bool bHasBoneWeightProfile = false;
 	bHasValidBoneWeights = SimpleMeshBoneWeightEdit<bool>(TargetMesh, Profile, bHasBoneWeightProfile, false,
@@ -253,12 +254,17 @@ UDynamicMesh* UGeometryScriptLibrary_MeshBoneWeightFunctions::SetVertexBoneWeigh
 			TArray<FBoneWeight> NewWeightsList;
 			for (int32 k = 0; k < Num; ++k)
 			{
+				if (BoneWeights[k].BoneIndex < 0)
+				{
+					AppendWarning(Debug, EGeometryScriptErrorType::InvalidInputs, LOCTEXT("SetVertexBoneWeights_InvalidInput", "SetVertexBoneWeights: Invalid bone index provided; falling back to 0 as bone index."));
+				}
+
 				FBoneWeight NewWeight;
-				NewWeight.SetBoneIndex(BoneWeights[k].BoneIndex);
+				NewWeight.SetBoneIndex(FMath::Max(0, BoneWeights[k].BoneIndex));
 				NewWeight.SetWeight(BoneWeights[k].Weight);
 				NewWeightsList.Add(NewWeight);
 			}
-			FBoneWeights NewBoneWeights = FBoneWeights::Create(NewWeightsList);
+			const FBoneWeights NewBoneWeights = FBoneWeights::Create(NewWeightsList);
 			SkinWeights.SetValue(VertexID, NewBoneWeights);
 			return true;
 		}
@@ -272,7 +278,8 @@ UDynamicMesh* UGeometryScriptLibrary_MeshBoneWeightFunctions::SetVertexBoneWeigh
 UDynamicMesh* UGeometryScriptLibrary_MeshBoneWeightFunctions::SetAllVertexBoneWeights(
 	UDynamicMesh* TargetMesh,
 	const TArray<FGeometryScriptBoneWeight>& BoneWeights, 
-	FGeometryScriptBoneWeightProfile Profile
+	FGeometryScriptBoneWeightProfile Profile,
+	UGeometryScriptDebug* Debug
 	)
 {
 	bool bHasBoneWeightProfile = false;
@@ -280,8 +287,13 @@ UDynamicMesh* UGeometryScriptLibrary_MeshBoneWeightFunctions::SetAllVertexBoneWe
 	TArray<FBoneWeight> NewWeightsList;
 	for (int32 k = 0; k < Num; ++k)
 	{
+		if (BoneWeights[k].BoneIndex < 0)
+		{
+			AppendWarning(Debug, EGeometryScriptErrorType::InvalidInputs, LOCTEXT("SetAllVertexBoneWeights_InvalidInput", "SetAllVertexBoneWeights: Invalid bone index provided; falling back to 0 as bone index."));
+		}
+		
 		FBoneWeight NewWeight;
-		NewWeight.SetBoneIndex(BoneWeights[k].BoneIndex);
+		NewWeight.SetBoneIndex(FMath::Max(0, BoneWeights[k].BoneIndex));
 		NewWeight.SetWeight(BoneWeights[k].Weight);
 		NewWeightsList.Add(NewWeight);
 	}
diff --git a/Engine/Plugins/Runtime/GeometryScripting/Source/GeometryScriptingCore/Public/GeometryScript/MeshBoneWeightFunctions.h b/Engine/Plugins/Runtime/GeometryScripting/Source/GeometryScriptingCore/Public/GeometryScript/MeshBoneWeightFunctions.h
--- a/Engine/Plugins/Runtime/GeometryScripting/Source/GeometryScriptingCore/Public/GeometryScript/MeshBoneWeightFunctions.h
+++ b/Engine/Plugins/Runtime/GeometryScripting/Source/GeometryScriptingCore/Public/GeometryScript/MeshBoneWeightFunctions.h
@@ -281,7 +281,8 @@ public:
 		int VertexID,
 		const TArray<FGeometryScriptBoneWeight>& BoneWeights,
 		bool& bIsValidVertexID,
-		FGeometryScriptBoneWeightProfile Profile = FGeometryScriptBoneWeightProfile() );
+		FGeometryScriptBoneWeightProfile Profile = FGeometryScriptBoneWeightProfile(),
+		UGeometryScriptDebug* Debug = nullptr);
 
 	/**
 	 * Set all vertices of the TargetMesh to the given Bone/Skin Weights
@@ -293,7 +294,8 @@ public:
 	SetAllVertexBoneWeights( 
 		UDynamicMesh* TargetMesh,
 		const TArray<FGeometryScriptBoneWeight>& BoneWeights,
-		FGeometryScriptBoneWeightProfile Profile = FGeometryScriptBoneWeightProfile() );
+		FGeometryScriptBoneWeightProfile Profile = FGeometryScriptBoneWeightProfile(),
+		UGeometryScriptDebug* Debug = nullptr);
 
 	/** 
 	 *  Computes a smooth skin binding for the given mesh to the skeleton provided.
