Fix issues in Geometry Script handling empty polygons and polygon lists -- AddPolygonVertex, AppendPolygonList not adding/appending, and AddPolygonToList crashing

diff --git a/Engine/Plugins/Runtime/GeometryScripting/Source/GeometryScriptingCore/Private/PolygonFunctions.cpp b/Engine/Plugins/Runtime/GeometryScripting/Source/GeometryScriptingCore/Private/PolygonFunctions.cpp
--- a/Engine/Plugins/Runtime/GeometryScripting/Source/GeometryScriptingCore/Private/PolygonFunctions.cpp
+++ b/Engine/Plugins/Runtime/GeometryScripting/Source/GeometryScriptingCore/Private/PolygonFunctions.cpp
@@ -106,7 +106,7 @@ void UGeometryScriptLibrary_SimplePolygonFunctions::SetPolygonVertex(FGeometrySc
 	}
 }
 
-int32 UGeometryScriptLibrary_SimplePolygonFunctions::AddPolygonVertex(FGeometryScriptSimplePolygon Polygon, FVector2D VertexPosition)
+int32 UGeometryScriptLibrary_SimplePolygonFunctions::AddPolygonVertex(FGeometryScriptSimplePolygon& Polygon, FVector2D VertexPosition)
 {
 	if (!Polygon.Vertices.IsValid())
 	{
@@ -368,11 +368,11 @@ FGeometryScriptGeneralPolygonList UGeometryScriptLibrary_PolygonListFunctions::C
 	return PolygonList;
 }
 
-int32 UGeometryScriptLibrary_PolygonListFunctions::AddPolygonToList(FGeometryScriptGeneralPolygonList PolygonList, FGeometryScriptSimplePolygon OuterShape, const TArray<FGeometryScriptSimplePolygon>& Holes, bool bFixHoleOrientations)
+int32 UGeometryScriptLibrary_PolygonListFunctions::AddPolygonToList(FGeometryScriptGeneralPolygonList& PolygonList, FGeometryScriptSimplePolygon OuterShape, const TArray<FGeometryScriptSimplePolygon>& Holes, bool bFixHoleOrientations)
 {
 	if (!PolygonList.Polygons.IsValid())
 	{
-		PolygonList.Polygons.Reset();
+		PolygonList.Reset();
 	}
 
 	TArray<FVector2D>* UseOuterVertices = OuterShape.Vertices.Get();
@@ -414,7 +414,7 @@ FGeometryScriptGeneralPolygonList UGeometryScriptLibrary_PolygonListFunctions::C
 	return PolygonList;
 }
 
-void UGeometryScriptLibrary_PolygonListFunctions::AppendPolygonList(FGeometryScriptGeneralPolygonList PolygonList, FGeometryScriptGeneralPolygonList PolygonsToAppend)
+void UGeometryScriptLibrary_PolygonListFunctions::AppendPolygonList(FGeometryScriptGeneralPolygonList& PolygonList, FGeometryScriptGeneralPolygonList PolygonsToAppend)
 {
 	if (!PolygonList.Polygons)
 	{
diff --git a/Engine/Plugins/Runtime/GeometryScripting/Source/GeometryScriptingCore/Public/GeometryScript/PolygonFunctions.h b/Engine/Plugins/Runtime/GeometryScripting/Source/GeometryScriptingCore/Public/GeometryScript/PolygonFunctions.h
--- a/Engine/Plugins/Runtime/GeometryScripting/Source/GeometryScriptingCore/Public/GeometryScript/PolygonFunctions.h
+++ b/Engine/Plugins/Runtime/GeometryScripting/Source/GeometryScriptingCore/Public/GeometryScript/PolygonFunctions.h
@@ -119,7 +119,7 @@ public:
 	 * Set the specified vertex of a Simple Polygon. Returns the index of the added vertex.
 	 */
 	UFUNCTION(BlueprintCallable, Category = "GeometryScript|SimplePolygon", meta = (ScriptMethod))
-	static UPARAM(DisplayName = "Vertex Index") int32 AddPolygonVertex(UPARAM(ref) FGeometryScriptSimplePolygon Polygon, FVector2D Position);
+	static UPARAM(DisplayName = "Vertex Index") int32 AddPolygonVertex(UPARAM(ref) FGeometryScriptSimplePolygon& Polygon, FVector2D Position);
 
 	/**
 	 * Returns a vertex's tangent of a Simple Polygon. VertexIndex loops around, so e.g., -1 gives the tangent of the last vertex in the polygon.
@@ -258,7 +258,7 @@ public:
 	 * Add Polygon to a Polygon List, with optional holes. Returns index of the added polygon.
 	 */
 	UFUNCTION(BlueprintCallable, Category = "GeometryScript|PolygonList", meta = (ScriptMethod, AutoCreateRefTerm = "HolePolygons"))
-	static UPARAM(DisplayName = "Polygon Index") int32 AddPolygonToList(UPARAM(ref) FGeometryScriptGeneralPolygonList PolygonList, FGeometryScriptSimplePolygon OuterPolygon, const TArray<FGeometryScriptSimplePolygon>& HolePolygons, bool bFixHoleOrientations = true);
+	static UPARAM(DisplayName = "Polygon Index") int32 AddPolygonToList(UPARAM(ref) FGeometryScriptGeneralPolygonList& PolygonList, FGeometryScriptSimplePolygon OuterPolygon, const TArray<FGeometryScriptSimplePolygon>& HolePolygons, bool bFixHoleOrientations = true);
 
 	/**
 	 * Create a Polygon List from an array of Simple Polygons
@@ -270,7 +270,7 @@ public:
 	 * Append the polygons in 'Polygons to Append' to Polygon List
 	 */
 	UFUNCTION(BlueprintCallable, Category = "GeometryScript|PolygonList", meta = (ScriptMethod))
-	static void AppendPolygonList(UPARAM(ref) FGeometryScriptGeneralPolygonList PolygonList, FGeometryScriptGeneralPolygonList PolygonsToAppend);
+	static void AppendPolygonList(UPARAM(ref) FGeometryScriptGeneralPolygonList& PolygonList, FGeometryScriptGeneralPolygonList PolygonsToAppend);
 	
 	/** Compute union of all polygons in Polygon List. Also resolves self-intersections within each polygon. */
 	UFUNCTION(BlueprintCallable, Category = "GeometryScript|PolygonList", meta = (ScriptMethod))
