add more docstring comments for geometry script FGeometryScriptCopyMeshToAssetOptions, covering the replace mesh material options

diff --git a/Engine/Plugins/Runtime/GeometryScripting/Source/GeometryScriptingCore/Public/GeometryScript/MeshAssetFunctions.h b/Engine/Plugins/Runtime/GeometryScripting/Source/GeometryScriptingCore/Public/GeometryScript/MeshAssetFunctions.h
--- a/Engine/Plugins/Runtime/GeometryScripting/Source/GeometryScriptingCore/Public/GeometryScript/MeshAssetFunctions.h
+++ b/Engine/Plugins/Runtime/GeometryScripting/Source/GeometryScriptingCore/Public/GeometryScript/MeshAssetFunctions.h
@@ -76,12 +76,15 @@ public:
 	UPROPERTY(EditAnywhere, BlueprintReadWrite, Category = Options)
 	bool bUseBuildScale = true;
 	
+	// Whether to replace the materials on the asset with those in the New Materials array
 	UPROPERTY(EditAnywhere, BlueprintReadWrite, Category = Options)
 	bool bReplaceMaterials = false;
 
+	// New materials to set if Replace Materials is enabled. Ignored otherwise.
 	UPROPERTY(EditAnywhere, BlueprintReadWrite, Category = Options)
 	TArray<TObjectPtr<UMaterialInterface>> NewMaterials;
 
+	// Optional slot names for the New Materials. Ignored if not the same length as the New Materials array.
 	UPROPERTY(EditAnywhere, BlueprintReadWrite, Category = Options)
 	TArray<FName> NewMaterialSlotNames;
 
