Add geometry script methods managing a global dynamic mesh pool (created on first access)

#virtualized

diff --git a/Engine/Plugins/Runtime/GeometryScripting/Source/GeometryScriptingCore/Private/MeshPoolFunctions.cpp b/Engine/Plugins/Runtime/GeometryScripting/Source/GeometryScriptingCore/Private/MeshPoolFunctions.cpp
new file mode 100644
--- /dev/null
+++ b/Engine/Plugins/Runtime/GeometryScripting/Source/GeometryScriptingCore/Private/MeshPoolFunctions.cpp
@@ -0,0 +1,48 @@
+// Copyright Epic Games, Inc. All Rights Reserved.
+
+#include "GeometryScript/MeshPoolFunctions.h"
+#include "UDynamicMesh.h"
+
+#include UE_INLINE_GENERATED_CPP_BY_NAME(MeshPoolFunctions)
+
+using namespace UE::Geometry;
+
+#define LOCTEXT_NAMESPACE "UGeometryScriptLibrary_MeshPoolFunctions"
+
+namespace UE::Private::MeshPoolHelper
+{
+	UDynamicMeshPool* GlobalMeshPoolHelper(bool bDestroy = false)
+	{
+		static UDynamicMeshPool* MeshPool = nullptr;
+		if (!MeshPool && !bDestroy)
+		{
+			MeshPool = NewObject<UDynamicMeshPool>();
+			MeshPool->AddToRoot();
+		}
+		else if (bDestroy)
+		{
+			MeshPool->FreeAllMeshes();
+			MeshPool->RemoveFromRoot();
+			MeshPool = nullptr;
+		}
+		return MeshPool;
+	}
+}
+
+
+
+UDynamicMeshPool* UGeometryScriptLibrary_MeshPoolFunctions::GetGlobalMeshPool()
+{
+	return UE::Private::MeshPoolHelper::GlobalMeshPoolHelper(false);
+}
+
+
+void UGeometryScriptLibrary_MeshPoolFunctions::DiscardGlobalMeshPool()
+{
+	UE::Private::MeshPoolHelper::GlobalMeshPoolHelper(true);
+}
+
+
+
+#undef LOCTEXT_NAMESPACE
+
diff --git a/Engine/Plugins/Runtime/GeometryScripting/Source/GeometryScriptingCore/Public/GeometryScript/MeshPoolFunctions.h b/Engine/Plugins/Runtime/GeometryScripting/Source/GeometryScriptingCore/Public/GeometryScript/MeshPoolFunctions.h
new file mode 100644
--- /dev/null
+++ b/Engine/Plugins/Runtime/GeometryScripting/Source/GeometryScriptingCore/Public/GeometryScript/MeshPoolFunctions.h
@@ -0,0 +1,26 @@
+// Copyright Epic Games, Inc. All Rights Reserved.
+
+#pragma once
+
+#include "Kismet/BlueprintFunctionLibrary.h"
+#include "GeometryScript/GeometryScriptTypes.h"
+#include "MeshPoolFunctions.generated.h"
+
+class UDynamicMeshPool;
+
+
+UCLASS(meta = (ScriptName = "GeometryScript_MeshPoolUtility"))
+class GEOMETRYSCRIPTINGCORE_API UGeometryScriptLibrary_MeshPoolFunctions : public UBlueprintFunctionLibrary
+{
+	GENERATED_BODY()
+public:
+
+	/** Access a global compute mesh pool (created on first access) */
+	UFUNCTION(BlueprintCallable, Category = "GeometryScript|MeshPool")
+	static UPARAM(DisplayName = "Mesh Pool") UDynamicMeshPool* 
+	GetGlobalMeshPool();
+
+	/** Fully clear/destroy the current global mesh pool, allowing it and all its meshes to be garbage collected */
+	UFUNCTION(BlueprintCallable, Category = "GeometryScript|MeshPool")
+	static void DiscardGlobalMeshPool();
+};
\ No newline at end of file
