Make geometry script's compact material ID method optionally also remove duplicate materials

diff --git a/Engine/Plugins/Runtime/GeometryScripting/Source/GeometryScriptingCore/Private/MeshMaterialFunctions.cpp b/Engine/Plugins/Runtime/GeometryScripting/Source/GeometryScriptingCore/Private/MeshMaterialFunctions.cpp
--- a/Engine/Plugins/Runtime/GeometryScripting/Source/GeometryScriptingCore/Private/MeshMaterialFunctions.cpp
+++ b/Engine/Plugins/Runtime/GeometryScripting/Source/GeometryScriptingCore/Private/MeshMaterialFunctions.cpp
@@ -643,6 +643,7 @@ UDynamicMesh* UGeometryScriptLibrary_MeshMaterialFunctions::CompactMaterialIDs(
 	UDynamicMesh* TargetMesh,
 	TArray<UMaterialInterface*> SourceMaterialList,
 	TArray<UMaterialInterface*>& CompactedMaterialList,
+	bool bRemoveDuplicateMaterials,
 	UGeometryScriptDebug* Debug)
 {
 	CompactedMaterialList = SourceMaterialList;
@@ -658,37 +659,68 @@ UDynamicMesh* UGeometryScriptLibrary_MeshMaterialFunctions::CompactMaterialIDs(
 				return;
 			}
 
-			bool bWasCompact = false;
+			// Find which material IDs are used
 			FInterval1i OldValueRange;
-			int32 NewMaxValue;
-			TArray<int32> OldToNewMap, NewToOldMap;
-			bool bOK = UE::Geometry::CompactAttributeValues(EditMesh, *MaterialIDs, OldValueRange, NewMaxValue, OldToNewMap, NewToOldMap, bWasCompact);
-			if (bOK == false)
+			TArray<bool> MaterialUsed;
+			MaterialUsed.SetNumZeroed(SourceMaterialList.Num());
+			for (int32 TID : EditMesh.TriangleIndicesItr())
 			{
-				UE::Geometry::AppendError(Debug, EGeometryScriptErrorType::InvalidInputs, LOCTEXT("CompactMaterialIDs_InvalidMaterialID", "CompactMaterialIDs: Invalid MaterialIDs found, unsafe to Compact"));
-				return;
+				int32 MID = MaterialIDs->GetValue(TID);
+				if (MID < 0)
+				{
+					UE::Geometry::AppendError(Debug, EGeometryScriptErrorType::InvalidInputs, LOCTEXT("CompactMaterialIDs_InvalidMaterialID", "CompactMaterialIDs: Invalid MaterialIDs found, unsafe to Compact"));
+					return;
+				}
+				else if (MID > MaterialUsed.Num()) // Note we allow material IDs beyond the SourceMaterialList array size
+				{
+					MaterialUsed.AddZeroed(MID + 1 - MaterialUsed.Num());
+				}
+				MaterialUsed[MID] = true;
+				OldValueRange.Contain(MID);
 			}
-			if (bWasCompact)
+			check(OldValueRange.Min >= 0 && OldValueRange.Max < MaterialUsed.Num());
+
+			CompactedMaterialList.Reset();
+
+			// Build an order-preserving mapping from original material IDs to their compacted ID
+			TArray<int32> ToCompactIdx;
+			ToCompactIdx.Init(-1, SourceMaterialList.Num());
+			TMap<UMaterialInterface*, int32> UniqueMaterials;
+			for (int32 SourceIdx = 0; SourceIdx < MaterialUsed.Num(); ++SourceIdx)
 			{
-				// this case occurs if MaterialIDs were removed such that the attribute values are still sequential/compact  (eg removed "at the end")
-				int32 NewMaterialIDCount = NewMaxValue+1;
-				CompactedMaterialList.Reset();
-				CompactedMaterialList.SetNum(NewMaterialIDCount);
-				for (int32 k = 0; k < NewMaterialIDCount; ++k)
+				if (!MaterialUsed[SourceIdx])
 				{
-					CompactedMaterialList[k] = SourceMaterialList.IsValidIndex(k) ? SourceMaterialList[k] : nullptr;
-				}		
-				
+					continue;
+				}
+				UMaterialInterface* Mat = SourceIdx < SourceMaterialList.Num() ? SourceMaterialList[SourceIdx] : nullptr;
+				int32 ToIdx = -1;
+				if (bRemoveDuplicateMaterials)
+				{
+					int32* FoundIdx = UniqueMaterials.Find(Mat);
+					if (!FoundIdx)
+					{
+						ToIdx = CompactedMaterialList.Add(Mat);
+						UniqueMaterials.Add(Mat, ToIdx);
+					}
+					else
+					{
+						ToIdx = *FoundIdx;
+					}
+				}
+				else
+				{
+					ToIdx = CompactedMaterialList.Add(Mat);
+				}
+				ToCompactIdx[SourceIdx] = ToIdx;
 			}
-			else
+
+			// Apply the remapping
+			for (int32 TID : EditMesh.TriangleIndicesItr())
 			{
-				int32 NewMaterialIDCount = NewToOldMap.Num();
-				CompactedMaterialList.Reset();
-				CompactedMaterialList.SetNum(NewMaterialIDCount);
-				for (int32 k = 0; k < NewMaterialIDCount; ++k)
+				int32 OrigMID = MaterialIDs->GetValue(TID);
+				if (ToCompactIdx.IsValidIndex(OrigMID))
 				{
-					int32 OldIndex = NewToOldMap[k];
-					CompactedMaterialList[k] = SourceMaterialList.IsValidIndex(OldIndex) ? SourceMaterialList[OldIndex] : nullptr;
+					MaterialIDs->SetValue(TID, ToCompactIdx[OrigMID]);
 				}
 			}
 
diff --git a/Engine/Plugins/Runtime/GeometryScripting/Source/GeometryScriptingCore/Public/GeometryScript/MeshMaterialFunctions.h b/Engine/Plugins/Runtime/GeometryScripting/Source/GeometryScriptingCore/Public/GeometryScript/MeshMaterialFunctions.h
--- a/Engine/Plugins/Runtime/GeometryScripting/Source/GeometryScriptingCore/Public/GeometryScript/MeshMaterialFunctions.h
+++ b/Engine/Plugins/Runtime/GeometryScripting/Source/GeometryScriptingCore/Public/GeometryScript/MeshMaterialFunctions.h
@@ -198,6 +198,7 @@ public:
 	 * N in-use MaterialIDs to the range [0,N-1]. Optionally compute a Compacted list of Materials.
 	 * @param SourceMaterialList Input Material list, assumption is that SourceMaterialList.Num() == number of MaterialIDs on mesh at input
 	 * @param CompactedMaterialList new Compacted Material list, one-to-one with new compacted MaterialIDs
+	 * @param bRemoveDuplicateMaterials Whether to also remove duplicate materials from the compacted list
 	 */
 	UFUNCTION(BlueprintCallable, Category = "GeometryScript|Materials", meta=(ScriptMethod))
 	static UPARAM(DisplayName = "Target Mesh") UDynamicMesh* 
@@ -205,6 +206,7 @@ public:
 		UDynamicMesh* TargetMesh, 
 		TArray<UMaterialInterface*> SourceMaterialList,
 		TArray<UMaterialInterface*>& CompactedMaterialList,
+		bool bRemoveDuplicateMaterials = false,
 		UGeometryScriptDebug* Debug = nullptr);
 
 };
\ No newline at end of file
