add a node to remap material IDs + construct a combined materials list, to help w/ workflows where you need multiple meshes to have consistent material IDs
+ fix crash in 'remap to new material ID by material' method when mesh has a material ID that points outside of the input materials list

diff --git a/Engine/Plugins/Runtime/GeometryScripting/Source/GeometryScriptingCore/Private/MeshMaterialFunctions.cpp b/Engine/Plugins/Runtime/GeometryScripting/Source/GeometryScriptingCore/Private/MeshMaterialFunctions.cpp
--- a/Engine/Plugins/Runtime/GeometryScripting/Source/GeometryScriptingCore/Private/MeshMaterialFunctions.cpp
+++ b/Engine/Plugins/Runtime/GeometryScripting/Source/GeometryScriptingCore/Private/MeshMaterialFunctions.cpp
@@ -204,6 +204,11 @@ UDynamicMesh* UGeometryScriptLibrary_MeshMaterialFunctions::RemapToNewMaterialID
 		for (int32 TriangleID : Mesh.TriangleIndicesItr())
 		{
 			int32 CurID = MaterialIDs.GetValue(TriangleID);
+			if (!ToMaterialID.IsValidIndex(CurID))
+			{
+				UE::Geometry::AppendWarning(Debug, EGeometryScriptErrorType::InvalidInputs, LOCTEXT("RemapToNewMaterialIDsByMaterial_InvalidMaterial", "RemapToNewMaterialIDsByMaterial: Invalid material ID in mesh was not a valid index into FromMaterialList, skipping"));
+				continue;
+			}
 			int32 NewID = ToMaterialID[CurID];
 			MaterialIDs.SetValue(TriangleID, NewID);
 		}
@@ -213,6 +218,76 @@ UDynamicMesh* UGeometryScriptLibrary_MeshMaterialFunctions::RemapToNewMaterialID
 }
 
 
+UDynamicMesh* UGeometryScriptLibrary_MeshMaterialFunctions::RemapAndCombineMaterials(
+	UDynamicMesh* TargetMesh,
+	const TArray<UMaterialInterface*>& TargetMeshMaterials,
+	const TArray<UMaterialInterface*>& RequiredMaterials,
+	TArray<UMaterialInterface*>& CombinedMaterials,
+	int RemapInvalidMaterialID,
+	bool bCompactDuplicateMaterials,
+	UGeometryScriptDebug* Debug
+)
+{
+	CombinedMaterials = RequiredMaterials;
+
+	if (TargetMesh == nullptr)
+	{
+		UE::Geometry::AppendError(Debug, EGeometryScriptErrorType::InvalidInputs, LOCTEXT("RemapAndCombineMaterials_InvalidInput", "RemapAndCombineMaterials: TargetMesh is Null"));
+		return TargetMesh;
+	}
+
+	TArray<int32> ToMaterialID;
+	if (bCompactDuplicateMaterials)
+	{
+		ToMaterialID.Init(-1, TargetMeshMaterials.Num());
+		for (int32 ToMaterialIdx = 0; ToMaterialIdx < TargetMeshMaterials.Num(); ++ToMaterialIdx)
+		{
+			int32 IdxInCombined = -1;
+			if (!CombinedMaterials.Find(TargetMeshMaterials[ToMaterialIdx], IdxInCombined))
+			{
+				IdxInCombined = CombinedMaterials.Add(TargetMeshMaterials[ToMaterialIdx]);
+			}
+			ToMaterialID[ToMaterialIdx] = IdxInCombined;
+		}
+	}
+	else
+	{
+		CombinedMaterials.Append(TargetMeshMaterials);
+	}
+
+	bool bHasMaterialIDs;
+	SimpleMeshMaterialEdit(TargetMesh, true, bHasMaterialIDs,
+		[&](FDynamicMesh3& Mesh, FDynamicMeshMaterialAttribute& MaterialIDs)
+		{
+			for (int32 TriangleID : Mesh.TriangleIndicesItr())
+			{
+				int32 CurID = MaterialIDs.GetValue(TriangleID);
+				int32 NewID = RemapInvalidMaterialID;
+				if (TargetMeshMaterials.IsValidIndex(CurID))
+				{
+					if (!bCompactDuplicateMaterials)
+					{
+						NewID = CurID + RequiredMaterials.Num();
+					}
+					else
+					{
+						NewID = ToMaterialID[CurID];
+					}
+				}
+				if (NewID < 0)
+				{
+					UE::Geometry::AppendWarning(Debug, EGeometryScriptErrorType::InvalidInputs, LOCTEXT("RemapAndCombineMaterials_InvalidMaterial", "RemapAndCombineMaterials: Invalid material ID in mesh was not a valid index into TargetMeshMaterials, skipping"));
+				}
+				else
+				{
+					MaterialIDs.SetValue(TriangleID, NewID);
+				}
+			}
+		});
+
+	return TargetMesh;
+}
+
 
 
 UDynamicMesh* UGeometryScriptLibrary_MeshMaterialFunctions::GetMaterialIDsOfTriangles( 
diff --git a/Engine/Plugins/Runtime/GeometryScripting/Source/GeometryScriptingCore/Public/GeometryScript/MeshMaterialFunctions.h b/Engine/Plugins/Runtime/GeometryScripting/Source/GeometryScriptingCore/Public/GeometryScript/MeshMaterialFunctions.h
--- a/Engine/Plugins/Runtime/GeometryScripting/Source/GeometryScriptingCore/Public/GeometryScript/MeshMaterialFunctions.h
+++ b/Engine/Plugins/Runtime/GeometryScripting/Source/GeometryScriptingCore/Public/GeometryScript/MeshMaterialFunctions.h
@@ -65,6 +65,32 @@ public:
 		int MissingMaterialID = -1,
 		UGeometryScriptDebug* Debug = nullptr );
 
+	/**
+	 * Remap material IDs to be consistent with a Required Materials list.
+	 * The Target Mesh material IDs will be remapped to reference the Combined Materials list, 
+	 * which will always start with the Required Materials.
+	 *
+	 * If a Material cannot be found in CurrentMeshMaterials, a warning will be printed and the MaterialID left unmodified, 
+	 * unless RemapInvalidMaterialID is set to a value >= 0, in which case RemapInvalidMaterialID will be assigned
+	 *
+	 * @param TargetMesh Mesh to update
+	 * @param TargetMeshMaterials Initial materials used by the TargetMesh
+	 * @param RequiredMaterials Materials that must be used, unchanged, in the output
+	 * @param CombinedMaterials Final materials used by the TargetMesh after remapping. Always starts with the RequiredMaterials.
+	 * @param RemapInvalidMaterialID If >= 0, automatically remap invalid input material IDs to this value
+	 * @param bCompactDuplicateMaterials If true, materials from TargetMeshMaterials will only be added if they are not already in RequiredMaterials. If false, all TargetMeshMaterials are appended.
+	 */
+	UFUNCTION(BlueprintCallable, Category = "GeometryScript|Materials", meta = (ScriptMethod))
+	static UPARAM(DisplayName = "Target Mesh") UDynamicMesh*
+	RemapAndCombineMaterials(
+		UDynamicMesh* TargetMesh,
+		const TArray<UMaterialInterface*>& TargetMeshMaterials,
+		const TArray<UMaterialInterface*>& RequiredMaterials,
+		TArray<UMaterialInterface*>& CombinedMaterials,
+		int RemapInvalidMaterialID = -1,
+		bool bCompactDuplicateMaterials = true,
+		UGeometryScriptDebug* Debug = nullptr
+	);
 
 	/**
 	 * Returns the current Material ID for a Triangle.  
