Fix backwards compat issue with material handling for static mesh -> dynamic mesh geometry script:
Revert the default behavior of CopyMeshFromStaticMesh to again use section material indices, rename it to make it clear that it does so, and make a CopyMeshFromStaticMeshV2 w/ the more-expected default behavior of using the regular material indices (and a new parameter to control the behavior)

+ add a method to more easily get the regular materials list off of a static mesh asset.

diff --git a/Engine/Plugins/Runtime/GeometryScripting/Source/GeometryScriptingCore/Private/MeshAssetFunctions.cpp b/Engine/Plugins/Runtime/GeometryScripting/Source/GeometryScriptingCore/Private/MeshAssetFunctions.cpp
--- a/Engine/Plugins/Runtime/GeometryScripting/Source/GeometryScriptingCore/Private/MeshAssetFunctions.cpp
+++ b/Engine/Plugins/Runtime/GeometryScripting/Source/GeometryScriptingCore/Private/MeshAssetFunctions.cpp
@@ -68,12 +68,13 @@ static void ConvertGeometryScriptWriteLOD(const FGeometryScriptMeshWriteLOD& Wri
 }
 
 
-UDynamicMesh*  UGeometryScriptLibrary_StaticMeshFunctions::CopyMeshFromStaticMesh(
+UDynamicMesh*  UGeometryScriptLibrary_StaticMeshFunctions::CopyMeshFromStaticMeshV2(
 	UStaticMesh* FromStaticMeshAsset, 
 	UDynamicMesh* ToDynamicMesh, 
 	FGeometryScriptCopyMeshFromAssetOptions AssetOptions,
 	FGeometryScriptMeshReadLOD RequestedLOD,
 	EGeometryScriptOutcomePins& Outcome,
+	bool bUseSectionMaterials,
 	UGeometryScriptDebug* Debug)
 {
 	Outcome = EGeometryScriptOutcomePins::Failure;
@@ -98,6 +99,7 @@ UDynamicMesh*  UGeometryScriptLibrary_StaticMeshFunctions::CopyMeshFromStaticMes
 	ConversionOptions.bRequestTangents = AssetOptions.bRequestTangents;
 	ConversionOptions.bIgnoreRemoveDegenerates = AssetOptions.bIgnoreRemoveDegenerates;
 	ConversionOptions.bUseBuildScale = AssetOptions.bUseBuildScale;
+	ConversionOptions.bUseSectionMaterialIndices = bUseSectionMaterials;
 
 	FText ErrorMessage;
 
@@ -454,6 +456,24 @@ int UGeometryScriptLibrary_StaticMeshFunctions::GetNumStaticMeshLODsOfType(
 
 
 
+void UGeometryScriptLibrary_StaticMeshFunctions::GetMaterialListFromStaticMesh(
+	UStaticMesh* FromStaticMeshAsset,
+	TArray<UMaterialInterface*>& MaterialList,
+	UGeometryScriptDebug* Debug)
+{
+	if (FromStaticMeshAsset == nullptr)
+	{
+		UE::Geometry::AppendError(Debug, EGeometryScriptErrorType::InvalidInputs, LOCTEXT("GetMaterialListFromStaticMesh_InvalidInput1", "GetMaterialListFromStaticMesh: FromStaticMeshAsset is Null"));
+		return;
+	}
+
+	const TArray<FStaticMaterial>& AssetMaterials = FromStaticMeshAsset->GetStaticMaterials();
+	for (int32 k = 0; k < AssetMaterials.Num(); ++k)
+	{
+		MaterialList.Add(AssetMaterials[k].MaterialInterface);
+	}
+}
+
 void UGeometryScriptLibrary_StaticMeshFunctions::GetSectionMaterialListFromStaticMesh(
 	UStaticMesh* FromStaticMeshAsset, 
 	FGeometryScriptMeshReadLOD RequestedLOD,
diff --git a/Engine/Plugins/Runtime/GeometryScripting/Source/GeometryScriptingCore/Public/GeometryScript/MeshAssetFunctions.h b/Engine/Plugins/Runtime/GeometryScripting/Source/GeometryScriptingCore/Public/GeometryScript/MeshAssetFunctions.h
--- a/Engine/Plugins/Runtime/GeometryScripting/Source/GeometryScriptingCore/Public/GeometryScript/MeshAssetFunctions.h
+++ b/Engine/Plugins/Runtime/GeometryScripting/Source/GeometryScriptingCore/Public/GeometryScript/MeshAssetFunctions.h
@@ -144,16 +144,36 @@ public:
 	* 
 	* Note that the LOD Index in RequestedLOD will be silently clamped to the available number of LODs (SourceModel or RenderData)
 	*/
-	UFUNCTION(BlueprintCallable, Category = "GeometryScript|StaticMesh", meta = (ExpandEnumAsExecs = "Outcome"))
+	//~ Note this is V2 to change the default behavior to use the normal materials list, rather than the section materials list.
+	UFUNCTION(BlueprintCallable, Category = "GeometryScript|StaticMesh", meta = (DisplayName = "Copy Mesh From Static Mesh", ExpandEnumAsExecs = "Outcome"))
 	static UPARAM(DisplayName = "Dynamic Mesh") UDynamicMesh* 
-	CopyMeshFromStaticMesh(
+	CopyMeshFromStaticMeshV2(
 		UStaticMesh* FromStaticMeshAsset, 
 		UDynamicMesh* ToDynamicMesh, 
 		FGeometryScriptCopyMeshFromAssetOptions AssetOptions,
 		FGeometryScriptMeshReadLOD RequestedLOD,
 		EGeometryScriptOutcomePins& Outcome,
+		bool bUseSectionMaterials = false,
 		UGeometryScriptDebug* Debug = nullptr);
 
+	/**
+	* Extracts a Dynamic Mesh from a Static Mesh Asset, using section material indices for the materials.
+	*
+	* Note that the LOD Index in RequestedLOD will be silently clamped to the available number of LODs (SourceModel or RenderData)
+	*/
+	UFUNCTION(BlueprintCallable, Category = "GeometryScript|StaticMesh", meta = (DisplayName = "Copy Mesh From Static Mesh with Section Materials", ExpandEnumAsExecs = "Outcome"))
+	static UPARAM(DisplayName = "Dynamic Mesh") UDynamicMesh*
+	CopyMeshFromStaticMesh(
+		UStaticMesh* FromStaticMeshAsset,
+		UDynamicMesh* ToDynamicMesh,
+		FGeometryScriptCopyMeshFromAssetOptions AssetOptions,
+		FGeometryScriptMeshReadLOD RequestedLOD,
+		EGeometryScriptOutcomePins& Outcome,
+		UGeometryScriptDebug* Debug = nullptr)
+	{
+		return CopyMeshFromStaticMeshV2(FromStaticMeshAsset, ToDynamicMesh, AssetOptions, RequestedLOD, Outcome, true, Debug);
+	}
+
 	/** 
 	* Updates a Static Mesh Asset with new geometry converted from a Dynamic Mesh.
 	*/
@@ -170,8 +190,8 @@ public:
 
     /** 
 	* Extracts the Material List and corresponding Material Indices from the specified LOD of the Static Mesh Asset. 
-	* The MaterialList is sorted by Section, so if CopyMeshToStaticMesh was used to create a DynamicMesh, then the returned
-	* MaterialList here will correspond to the MaterialIDs in that DynamicMesh (as each Static Mesh Section becomes a MaterialID, in-order). 
+	* The MaterialList is sorted by Section, so if CopyMeshToStaticMesh was used to create a DynamicMesh with bUseSectionMaterials=true, then the 
+	* returned MaterialList here will correspond to the MaterialIDs in that DynamicMesh (as each Static Mesh Section becomes a MaterialID, in-order). 
 	* So, the returned MaterialList can be passed directly to (eg) a DynamicMeshComponent.
 	* 
 	* @param MaterialIndex this returned array is the same size as MaterialList, with each value the index of that Material in the StaticMesh Material List
@@ -190,6 +210,16 @@ public:
 		EGeometryScriptOutcomePins& Outcome,
 		UGeometryScriptDebug* Debug = nullptr);
 
+	/**
+	 * Get the asset materials from the static mesh asset. These will match the DynamicMesh material if CopyMeshToStaticMesh
+	 * was used to create a DynamicMesh with bUseSectionMaterials=false
+	 */
+	UFUNCTION(BlueprintCallable, Category = "GeometryScript|StaticMesh")
+	static void 
+	GetMaterialListFromStaticMesh(UStaticMesh* FromStaticMeshAsset,
+		TArray<UMaterialInterface*>& MaterialList,
+		UGeometryScriptDebug* Debug = nullptr);
+
 	/** 
 	* Extracts a Dynamic Mesh from a Skeletal Mesh Asset. 
 	*/
diff --git a/Engine/Plugins/Runtime/MeshModelingToolset/Source/ModelingComponents/Private/ConversionUtils/SceneComponentToDynamicMesh.cpp b/Engine/Plugins/Runtime/MeshModelingToolset/Source/ModelingComponents/Private/ConversionUtils/SceneComponentToDynamicMesh.cpp
--- a/Engine/Plugins/Runtime/MeshModelingToolset/Source/ModelingComponents/Private/ConversionUtils/SceneComponentToDynamicMesh.cpp
+++ b/Engine/Plugins/Runtime/MeshModelingToolset/Source/ModelingComponents/Private/ConversionUtils/SceneComponentToDynamicMesh.cpp
@@ -271,7 +271,10 @@ namespace Private::ConversionHelper
 		}
 
 		FMeshDescriptionToDynamicMesh Converter;
-		Converter.SetPolygonGroupToMaterialIndexMap(PolygonGroupToMaterialMap);
+		if (!AssetOptions.bUseSectionMaterialIndices)
+		{
+			Converter.SetPolygonGroupToMaterialIndexMap(PolygonGroupToMaterialMap);
+		}
 		Converter.Convert(SourceMesh, OutMesh, AssetOptions.bRequestTangents);
 
 		bSuccess = true;
diff --git a/Engine/Plugins/Runtime/MeshModelingToolset/Source/ModelingComponents/Public/ConversionUtils/SceneComponentToDynamicMesh.h b/Engine/Plugins/Runtime/MeshModelingToolset/Source/ModelingComponents/Public/ConversionUtils/SceneComponentToDynamicMesh.h
--- a/Engine/Plugins/Runtime/MeshModelingToolset/Source/ModelingComponents/Public/ConversionUtils/SceneComponentToDynamicMesh.h
+++ b/Engine/Plugins/Runtime/MeshModelingToolset/Source/ModelingComponents/Public/ConversionUtils/SceneComponentToDynamicMesh.h
@@ -100,6 +100,9 @@ namespace Conversion {
 
 		// Whether to scale the copied mesh by the Build Setting's 'Build Scale'. Note: This is considered separately from the 'Apply Build Settings' option.
 		bool bUseBuildScale = true;
+
+		// Whether to directly use the section material indices as output material IDs. Otherwise, material IDs will be indices in to the materials array.
+		bool bUseSectionMaterialIndices = false;
 	};
 
 	/**
