Add null checks for GeometryScript SceneUtilityFunctions.cpp on UDynamicMesh* inputs

diff --git a/Engine/Plugins/Runtime/GeometryScripting/Source/GeometryScriptingCore/Private/SceneUtilityFunctions.cpp b/Engine/Plugins/Runtime/GeometryScripting/Source/GeometryScriptingCore/Private/SceneUtilityFunctions.cpp
--- a/Engine/Plugins/Runtime/GeometryScripting/Source/GeometryScriptingCore/Private/SceneUtilityFunctions.cpp
+++ b/Engine/Plugins/Runtime/GeometryScripting/Source/GeometryScriptingCore/Private/SceneUtilityFunctions.cpp
@@ -50,6 +50,12 @@ UDynamicMesh* UGeometryScriptLibrary_SceneUtilityFunctions::CopyMeshFromComponen
 	UGeometryScriptDebug* Debug)
 {
 	Outcome = EGeometryScriptOutcomePins::Failure;
+	
+	if (ToDynamicMesh == nullptr)
+	{
+		UE::Geometry::AppendError(Debug, EGeometryScriptErrorType::InvalidInputs, LOCTEXT("CopyMeshFromComponent_InvalidInput", "CopyMeshFromComponent: ToDynamicMesh is Null"));
+		return ToDynamicMesh;
+	}
 
 	UE::Conversion::FToMeshOptions ToMeshOptions;
 	ToMeshOptions.bUseClosestLOD = false;
@@ -124,6 +130,12 @@ UDynamicMesh* UGeometryScriptLibrary_SceneUtilityFunctions::CopyCollisionMeshesF
 {
 	Outcome = EGeometryScriptOutcomePins::Failure;
 
+	if (ToDynamicMesh == nullptr)
+	{
+		UE::Geometry::AppendError(Debug, EGeometryScriptErrorType::InvalidInputs, LOCTEXT("CopyCollisionMeshesFromObject_InvalidInput", "CopyCollisionMeshesFromObject: ToDynamicMesh is Null"));
+		return ToDynamicMesh;
+	}
+
 	FDynamicMesh3 AccumulatedMesh;
 	ToDynamicMesh->ProcessMesh([&](const FDynamicMesh3& ReadMesh)
 	{
