Add geometry script method to set a custom complex collision mesh for a static mesh asset

diff --git a/Engine/Plugins/Runtime/GeometryScripting/Source/GeometryScriptingCore/Private/CollisionFunctions.cpp b/Engine/Plugins/Runtime/GeometryScripting/Source/GeometryScriptingCore/Private/CollisionFunctions.cpp
--- a/Engine/Plugins/Runtime/GeometryScripting/Source/GeometryScriptingCore/Private/CollisionFunctions.cpp
+++ b/Engine/Plugins/Runtime/GeometryScripting/Source/GeometryScriptingCore/Private/CollisionFunctions.cpp
@@ -155,14 +155,12 @@ void ComputeCollisionFromMesh(
 	GeneratedCollision = NewCollision.AggGeom;
 }
 
-
-
-static void SetStaticMeshSimpleCollision(UStaticMesh* StaticMeshAsset, const FKAggregateGeom& NewSimpleCollision, bool bEmitTransaction, bool bMarkCollisionAsCustomized = true)
+static void UpdateStaticMeshCollision(UStaticMesh* StaticMeshAsset, TFunctionRef<bool(UStaticMesh*, UBodySetup*)> ApplyUpdate, bool bEmitTransaction, bool bMarkCollisionAsCustomized = true)
 {
 #if WITH_EDITOR
 	if (bEmitTransaction && GEditor)
 	{
-		GEditor->BeginTransaction(LOCTEXT("UpdateStaticMesh", "Set Simple Collision"));
+		GEditor->BeginTransaction(LOCTEXT("UpdateStaticMesh", "Update Static Mesh Collision"));
 
 		StaticMeshAsset->Modify();
 	}
@@ -171,25 +169,22 @@ static void SetStaticMeshSimpleCollision(UStaticMesh* StaticMeshAsset, const FKA
 	UBodySetup* BodySetup = StaticMeshAsset->GetBodySetup();
 	if (BodySetup != nullptr)
 	{
-		// mark the BodySetup for modification. Do we need to modify the UStaticMesh??
+		// mark the BodySetup for modification.
 #if WITH_EDITOR
 		if (bEmitTransaction)
 		{
 			BodySetup->Modify();
 		}
 #endif
+	}
 
-		// clear existing simple collision. This will call BodySetup->InvalidatePhysicsData()
-		BodySetup->RemoveSimpleCollision();
-
-		// set new collision geometry
-		BodySetup->AggGeom = NewSimpleCollision;
-
-		// update collision type
-		//BodySetup->CollisionTraceFlag = (ECollisionTraceFlag)(int32)Settings->SetCollisionType;
-
-		// rebuild physics meshes
-		BodySetup->CreatePhysicsMeshes();
+	if (ApplyUpdate(StaticMeshAsset, BodySetup))
+	{
+		if (BodySetup)
+		{
+			// rebuild physics meshes
+			BodySetup->CreatePhysicsMeshes();
+		}
 
 		StaticMeshAsset->RecreateNavCollision();
 
@@ -229,6 +224,27 @@ static void SetStaticMeshSimpleCollision(UStaticMesh* StaticMeshAsset, const FKA
 
 }
 
+static void SetStaticMeshSimpleCollision(UStaticMesh* StaticMeshAsset, const FKAggregateGeom& NewSimpleCollision, bool bEmitTransaction, bool bMarkCollisionAsCustomized = true)
+{
+	UpdateStaticMeshCollision(StaticMeshAsset, [&NewSimpleCollision](UStaticMesh*, UBodySetup* BodySetup) -> bool
+	{
+		if (!BodySetup)
+		{
+			return false;
+		}
+		// clear existing simple collision. This will call BodySetup->InvalidatePhysicsData()
+		BodySetup->RemoveSimpleCollision();
+
+		// set new collision geometry
+		BodySetup->AggGeom = NewSimpleCollision;
+
+		// update collision type
+		//BodySetup->CollisionTraceFlag = (ECollisionTraceFlag)(int32)Settings->SetCollisionType;
+
+		return true;
+	}, bEmitTransaction, bMarkCollisionAsCustomized);
+}
+
 
 // local helper to convert the blueprint-accessible enum to the geometrycore equivalent
 static UE::Geometry::FNegativeSpaceSampleSettings::ESampleMethod ConvertNegativeSpaceSampleMethodEnum(ENegativeSpaceSampleMethod SampleMethod)
@@ -296,6 +312,43 @@ static double GetConvexElemVolume(const FKConvexElem& Convex)
 }		// end namespace UELocal
 
 
+bool UGeometryScriptLibrary_CollisionFunctions::SetStaticMeshCustomComplexCollision(
+	UStaticMesh* StaticMeshAsset,
+	UStaticMesh* StaticMeshCollisionAsset,
+	bool bEmitTransaction,
+	bool bMarkCollisionAsCustomized,
+	UGeometryScriptDebug* Debug
+)
+{
+	if (StaticMeshAsset == nullptr)
+	{
+		UE::Geometry::AppendError(Debug, EGeometryScriptErrorType::InvalidInputs, LOCTEXT("SetStaticMeshCustomComplexCollision_InvalidInput", "SetStaticMeshCustomComplexCollision: StaticMeshAsset is Null"));
+		return false;
+	}
+	if (StaticMeshCollisionAsset == nullptr)
+	{
+		UE::Geometry::AppendError(Debug, EGeometryScriptErrorType::InvalidInputs, LOCTEXT("SetStaticMeshCustomComplexCollision_InvalidInput2", "SetStaticMeshCustomComplexCollision: StaticMeshCollisionAsset is Null"));
+		return false;
+	}
+#if WITH_EDITORONLY_DATA
+	UELocal::UpdateStaticMeshCollision(StaticMeshAsset, [StaticMeshCollisionAsset](UStaticMesh* UpdateStaticMesh, UBodySetup* BodySetup) -> bool
+	{
+		if (BodySetup)
+		{
+			BodySetup->InvalidatePhysicsData();
+		}
+
+		// set new collision geometry
+		UpdateStaticMesh->ComplexCollisionMesh = StaticMeshCollisionAsset;
+
+		return true;
+	}, bEmitTransaction, bMarkCollisionAsCustomized);
+	return true;
+#else
+	return false;
+#endif
+}
+
 UDynamicMesh* UGeometryScriptLibrary_CollisionFunctions::SetStaticMeshCollisionFromMesh(
 	UDynamicMesh* FromDynamicMesh,
 	UStaticMesh* ToStaticMeshAsset,
diff --git a/Engine/Plugins/Runtime/GeometryScripting/Source/GeometryScriptingCore/Public/GeometryScript/CollisionFunctions.h b/Engine/Plugins/Runtime/GeometryScripting/Source/GeometryScriptingCore/Public/GeometryScript/CollisionFunctions.h
--- a/Engine/Plugins/Runtime/GeometryScripting/Source/GeometryScriptingCore/Public/GeometryScript/CollisionFunctions.h
+++ b/Engine/Plugins/Runtime/GeometryScripting/Source/GeometryScriptingCore/Public/GeometryScript/CollisionFunctions.h
@@ -331,6 +331,21 @@ class GEOMETRYSCRIPTINGCORE_API UGeometryScriptLibrary_CollisionFunctions : publ
 	GENERATED_BODY()
 public:
 
+	/**
+	 * Set a static mesh as the custom collision for another static mesh to use.
+	 * Note: Only works if editor-only data is available.
+	 * 
+	 * @return true on success.
+	 */
+	UFUNCTION(BlueprintCallable, Category = "GeometryScript|Collision")
+	static bool SetStaticMeshCustomComplexCollision(
+		UStaticMesh* StaticMeshAsset,
+		UStaticMesh* StaticMeshCollisionAsset,
+		bool bEmitTransaction,
+		bool bMarkCollisionAsCustomized = true,
+		UGeometryScriptDebug* Debug = nullptr
+	);
+
 	/** 
 	* Generates Simple Collision shapes for a Static Mesh Asset based on the input Dynamic Mesh.
 	*/
