GeometryScript: Exposed GetAllUVSeamEdges and GetAllVertexPositionsAtEdges to MeshQueryFunctions.

diff --git a/Engine/Plugins/Runtime/GeometryScripting/Source/GeometryScriptingCore/Private/MeshQueryFunctions.cpp b/Engine/Plugins/Runtime/GeometryScripting/Source/GeometryScriptingCore/Private/MeshQueryFunctions.cpp
--- a/Engine/Plugins/Runtime/GeometryScripting/Source/GeometryScriptingCore/Private/MeshQueryFunctions.cpp
+++ b/Engine/Plugins/Runtime/GeometryScripting/Source/GeometryScriptingCore/Private/MeshQueryFunctions.cpp
@@ -391,6 +391,34 @@ UDynamicMesh* UGeometryScriptLibrary_MeshQueryFunctions::GetVertexConnectedVerti
 }
 
 
+UDynamicMesh* UGeometryScriptLibrary_MeshQueryFunctions::GetAllVertexPositionsAtEdges(UDynamicMesh* TargetMesh, const FGeometryScriptIndexList& EdgeIDs, FGeometryScriptVectorList& Start, FGeometryScriptVectorList& End)
+{
+	Start.Reset();
+	End.Reset();
+	
+	if (!EdgeIDs.IsCompatibleWith(EGeometryScriptIndexType::Edge) && EdgeIDs.List.IsValid() && EdgeIDs.List->Num() > 0)
+	{
+		return TargetMesh;
+	}
+
+	SimpleMeshQuery<bool>(TargetMesh, false,
+		[&EdgeIDs, &Start, &End](const FDynamicMesh3& Mesh)
+		{
+			for (const int EID : *EdgeIDs.List)
+			{
+				if (Mesh.IsEdge(EID))
+				{
+					FIndex2i VIDs = Mesh.GetEdgeV(EID);
+					Start.List->Add(Mesh.GetVertex(VIDs[0]));
+					End.List->Add(Mesh.GetVertex(VIDs[1]));
+				}
+			}
+			return true;
+		});
+	return TargetMesh;
+}
+
+
 int32 UGeometryScriptLibrary_MeshQueryFunctions::GetNumUVSets( UDynamicMesh* TargetMesh )
 {
 	return SimpleMeshQuery<int32>(TargetMesh, 0, [&](const FDynamicMesh3& Mesh) { 
@@ -480,6 +508,24 @@ UDynamicMesh* UGeometryScriptLibrary_MeshQueryFunctions::GetInterpolatedTriangle
 	return TargetMesh;
 }
 
+UDynamicMesh* UGeometryScriptLibrary_MeshQueryFunctions::GetAllUVSeamEdges(UDynamicMesh* TargetMesh, int32 UVSetIndex, FGeometryScriptIndexList& ElementIDs)
+{
+	ElementIDs.Reset(EGeometryScriptIndexType::Edge);
+	bool bIsValidUVSet = SimpleMeshUVSetQuery<bool>(TargetMesh, UVSetIndex, bIsValidUVSet, false,
+		[&](const FDynamicMesh3& Mesh, const FDynamicMeshUVOverlay& UVSet)
+		{
+			for (int EID : Mesh.EdgeIndicesItr())
+			{
+				if (UVSet.IsSeamEdge(EID))
+				{
+					ElementIDs.List->Add(EID);
+				}
+			}
+			return true;
+		});
+	return TargetMesh;
+}
+
 
 bool UGeometryScriptLibrary_MeshQueryFunctions::GetHasTriangleNormals( UDynamicMesh* TargetMesh )
 {
diff --git a/Engine/Plugins/Runtime/GeometryScripting/Source/GeometryScriptingCore/Public/GeometryScript/GeometryScriptTypes.h b/Engine/Plugins/Runtime/GeometryScripting/Source/GeometryScriptingCore/Public/GeometryScript/GeometryScriptTypes.h
--- a/Engine/Plugins/Runtime/GeometryScripting/Source/GeometryScriptingCore/Public/GeometryScript/GeometryScriptTypes.h
+++ b/Engine/Plugins/Runtime/GeometryScripting/Source/GeometryScriptingCore/Public/GeometryScript/GeometryScriptTypes.h
@@ -304,6 +304,7 @@ enum class EGeometryScriptIndexType : uint8
 	// Index lists of Any type are compatible with any other index list type
 	Any,
 	Triangle,
+	Edge,
 	Vertex,
 	MaterialID,
 	PolygroupID UMETA(DisplayName = "PolyGroup ID")
diff --git a/Engine/Plugins/Runtime/GeometryScripting/Source/GeometryScriptingCore/Public/GeometryScript/MeshQueryFunctions.h b/Engine/Plugins/Runtime/GeometryScripting/Source/GeometryScriptingCore/Public/GeometryScript/MeshQueryFunctions.h
--- a/Engine/Plugins/Runtime/GeometryScripting/Source/GeometryScriptingCore/Public/GeometryScript/MeshQueryFunctions.h
+++ b/Engine/Plugins/Runtime/GeometryScripting/Source/GeometryScriptingCore/Public/GeometryScript/MeshQueryFunctions.h
@@ -217,6 +217,18 @@ public:
 	static UPARAM(DisplayName = "Target Mesh") UDynamicMesh* 
 	GetAllVertexPositions( UDynamicMesh* TargetMesh, FGeometryScriptVectorList& PositionList, bool bSkipGaps, bool& bHasVertexIDGaps );
 
+	/**
+	 * Returns the vertex positions for each edge in the given index list.
+	 * 
+	 * @param TargetMesh The mesh to query
+	 * @param EdgeIDs The edge IDs to query
+	 * @param Start The output list of start vertex positions
+	 * @param End The output list of end vertex positions
+	 * @return The target mesh that was queried
+	 */
+	UFUNCTION(BlueprintCallable, Category = "GeometryScript|MeshQueries", meta=(ScriptMethod))
+	static UPARAM(DisplayName = "Target Mesh") UDynamicMesh*
+	GetAllVertexPositionsAtEdges( UDynamicMesh* TargetMesh, const FGeometryScriptIndexList& EdgeIDs, FGeometryScriptVectorList& Start, FGeometryScriptVectorList& End);
 
 	/**
 	 * Return array of Triangle IDs connected to the given VertexID, ie the triangle one-ring
@@ -282,6 +294,19 @@ public:
 		FVector2D& InterpolatedUV );
 
 
+	/**
+	 * Returns all edge element IDs that are UV seam edges for a given UV channel.
+	 * 
+	 * @param TargetMesh The mesh to query.
+	 * @param UVSetIndex The UV channel to query
+	 * @param ElementIDs The returned edge element IDs 
+	 * @return The target mesh that was queried. 
+	 */
+	UFUNCTION(BlueprintCallable, Category = "GeometryScript|MeshQueries", meta=(ScriptMethod))
+	static UPARAM(DisplayName = "Target Mesh") UDynamicMesh*
+	GetAllUVSeamEdges( UDynamicMesh* TargetMesh, UPARAM(DisplayName = "UV Channel") int32 UVSetIndex, FGeometryScriptIndexList& ElementIDs);
+
+	
 	//
 	// Normal queries
 	//
