GeometryScript: Added bHaveValidUVs output param to GetAllUVSeamEdges.

diff --git a/Engine/Plugins/Runtime/GeometryScripting/Source/GeometryScriptingCore/Private/MeshQueryFunctions.cpp b/Engine/Plugins/Runtime/GeometryScripting/Source/GeometryScriptingCore/Private/MeshQueryFunctions.cpp
--- a/Engine/Plugins/Runtime/GeometryScripting/Source/GeometryScriptingCore/Private/MeshQueryFunctions.cpp
+++ b/Engine/Plugins/Runtime/GeometryScripting/Source/GeometryScriptingCore/Private/MeshQueryFunctions.cpp
@@ -508,10 +508,10 @@ UDynamicMesh* UGeometryScriptLibrary_MeshQueryFunctions::GetInterpolatedTriangle
 	return TargetMesh;
 }
 
-UDynamicMesh* UGeometryScriptLibrary_MeshQueryFunctions::GetAllUVSeamEdges(UDynamicMesh* TargetMesh, int32 UVSetIndex, FGeometryScriptIndexList& ElementIDs)
+UDynamicMesh* UGeometryScriptLibrary_MeshQueryFunctions::GetAllUVSeamEdges(UDynamicMesh* TargetMesh, int32 UVSetIndex, bool& bHaveValidUVs, FGeometryScriptIndexList& ElementIDs)
 {
 	ElementIDs.Reset(EGeometryScriptIndexType::Edge);
-	bool bIsValidUVSet = SimpleMeshUVSetQuery<bool>(TargetMesh, UVSetIndex, bIsValidUVSet, false,
+	bHaveValidUVs = SimpleMeshUVSetQuery<bool>(TargetMesh, UVSetIndex, bHaveValidUVs, false,
 		[&](const FDynamicMesh3& Mesh, const FDynamicMeshUVOverlay& UVSet)
 		{
 			for (int EID : Mesh.EdgeIndicesItr())
diff --git a/Engine/Plugins/Runtime/GeometryScripting/Source/GeometryScriptingCore/Public/GeometryScript/MeshQueryFunctions.h b/Engine/Plugins/Runtime/GeometryScripting/Source/GeometryScriptingCore/Public/GeometryScript/MeshQueryFunctions.h
--- a/Engine/Plugins/Runtime/GeometryScripting/Source/GeometryScriptingCore/Public/GeometryScript/MeshQueryFunctions.h
+++ b/Engine/Plugins/Runtime/GeometryScripting/Source/GeometryScriptingCore/Public/GeometryScript/MeshQueryFunctions.h
@@ -299,12 +299,13 @@ public:
 	 * 
 	 * @param TargetMesh The mesh to query.
 	 * @param UVSetIndex The UV channel to query
+	 * @param bHaveValidUVs True, if the mesh has UVs for the given UVSetIndex.
 	 * @param ElementIDs The returned edge element IDs 
 	 * @return The target mesh that was queried. 
 	 */
 	UFUNCTION(BlueprintCallable, Category = "GeometryScript|MeshQueries", meta=(ScriptMethod))
 	static UPARAM(DisplayName = "Target Mesh") UDynamicMesh*
-	GetAllUVSeamEdges( UDynamicMesh* TargetMesh, UPARAM(DisplayName = "UV Channel") int32 UVSetIndex, FGeometryScriptIndexList& ElementIDs);
+	GetAllUVSeamEdges( UDynamicMesh* TargetMesh, UPARAM(DisplayName = "UV Channel") int32 UVSetIndex, bool& bHaveValidUVs, FGeometryScriptIndexList& ElementIDs);
 
 	
 	//
