Add helper geometry script methods for specifying loops and chains of edges, which can be passed to the constrained delaunay triangulation method

diff --git a/Engine/Plugins/Runtime/GeometryScripting/Source/GeometryScriptingCore/Private/MeshPrimitiveFunctions.cpp b/Engine/Plugins/Runtime/GeometryScripting/Source/GeometryScriptingCore/Private/MeshPrimitiveFunctions.cpp
--- a/Engine/Plugins/Runtime/GeometryScripting/Source/GeometryScriptingCore/Private/MeshPrimitiveFunctions.cpp
+++ b/Engine/Plugins/Runtime/GeometryScripting/Source/GeometryScriptingCore/Private/MeshPrimitiveFunctions.cpp
@@ -1446,6 +1446,38 @@ UDynamicMesh* UGeometryScriptLibrary_MeshPrimitiveFunctions::AppendDelaunayTrian
 	return TargetMesh;
 }
 
+TArray<FIntPoint> UGeometryScriptLibrary_MeshPrimitiveFunctions::CreateConstrainedEdgesLoop(int32 NumVertices, int32 Start)
+{
+	TArray<FIntPoint> Loop;
+	if (NumVertices < 3)
+	{
+		return Loop;
+	}
+
+	Loop.Reserve(NumVertices);
+	for (int32 Idx = 0, Prev = NumVertices - 1; Idx < NumVertices; Prev = Idx++)
+	{
+		Loop.Emplace(Prev + Start, Idx + Start);
+	}
+	return Loop;
+}
+
+TArray<FIntPoint> UGeometryScriptLibrary_MeshPrimitiveFunctions::CreateConstrainedEdgesChain(int32 NumVertices, int32 Start)
+{
+	TArray<FIntPoint> Chain;
+	if (NumVertices < 2)
+	{
+		return Chain;
+	}
+
+	Chain.Reserve(NumVertices - 1);
+	for (int32 Idx = 0; Idx + 1 < NumVertices; ++Idx)
+	{
+		Chain.Emplace(Idx, Idx + 1);
+	}
+	return Chain;
+}
+
 UDynamicMesh* UGeometryScriptLibrary_MeshPrimitiveFunctions::AppendPolygonListTriangulation(
 	UDynamicMesh* TargetMesh,
 	FGeometryScriptPrimitiveOptions PrimitiveOptions,
diff --git a/Engine/Plugins/Runtime/GeometryScripting/Source/GeometryScriptingCore/Public/GeometryScript/MeshPrimitiveFunctions.h b/Engine/Plugins/Runtime/GeometryScripting/Source/GeometryScriptingCore/Public/GeometryScript/MeshPrimitiveFunctions.h
--- a/Engine/Plugins/Runtime/GeometryScripting/Source/GeometryScriptingCore/Public/GeometryScript/MeshPrimitiveFunctions.h
+++ b/Engine/Plugins/Runtime/GeometryScripting/Source/GeometryScriptingCore/Public/GeometryScript/MeshPrimitiveFunctions.h
@@ -587,6 +587,22 @@ public:
 		bool& bHasDuplicateVertices,
 		UGeometryScriptDebug* Debug = nullptr);
 
+	/**
+	 * Intended for use with AppendDelaunayTriangulation2D:
+	 * Create a loop of edges through sequential vertices
+	 * e.g., a Loop(3,0) will construct edges (2,0), (0,1) and (1,2)
+	 */
+	UFUNCTION(BlueprintCallable, Category = "GeometryScript|Primitives")
+	static TArray<FIntPoint> CreateConstrainedEdgesLoop(int32 NumVertices, int32 Start = 0);
+
+	/**
+	 * Intended for use with AppendDelaunayTriangulation2D:
+	 * Create a chain of edges through sequential vertices
+	 * e.g., a Chain(3,0) will construct edges (0,1) and (1,2)
+	 */
+	UFUNCTION(BlueprintCallable, Category = "GeometryScript|Primitives")
+	static TArray<FIntPoint> CreateConstrainedEdgesChain(int32 NumVertices, int32 Start = 0);
+
 
 	/**
 	* Generates a Delaunay Triangulation of the provided Polygon List, and appends it to the Target Mesh.
