Added a geometry script node for writing morph target data into a skeletal mesh asset.

diff --git a/Engine/Plugins/Runtime/GeometryScripting/Source/GeometryScriptingCore/Private/MeshAssetFunctions.cpp b/Engine/Plugins/Runtime/GeometryScripting/Source/GeometryScriptingCore/Private/MeshAssetFunctions.cpp
--- a/Engine/Plugins/Runtime/GeometryScripting/Source/GeometryScriptingCore/Private/MeshAssetFunctions.cpp
+++ b/Engine/Plugins/Runtime/GeometryScripting/Source/GeometryScriptingCore/Private/MeshAssetFunctions.cpp
@@ -792,5 +792,122 @@ UDynamicMesh* UGeometryScriptLibrary_StaticMeshFunctions::CopyMeshToSkeletalMesh
 }
 
 
+UDynamicMesh* UGeometryScriptLibrary_StaticMeshFunctions::CopyMorphTargetToSkeletalMesh(
+		UDynamicMesh* FromDynamicMesh, 
+		USkeletalMesh* ToSkeletalMeshAsset,
+		FName MorphTargetName,
+		FGeometryScriptCopyMeshToAssetOptions Options,
+		FGeometryScriptMeshWriteLOD TargetLOD,
+		EGeometryScriptOutcomePins& Outcome,
+		UGeometryScriptDebug* Debug)
+{
+	Outcome = EGeometryScriptOutcomePins::Failure;
+
+	if (ToSkeletalMeshAsset == nullptr)
+	{
+		UE::Geometry::AppendError(Debug, EGeometryScriptErrorType::InvalidInputs, LOCTEXT("CopyMorphTargetToSkeletalMesh_InvalidInput1", "CopyMorphTargetToSkeletalMesh: ToSkeletalMeshAsset is Null"));
+		return FromDynamicMesh;
+	}
+	if (FromDynamicMesh == nullptr)
+	{
+		UE::Geometry::AppendError(Debug, EGeometryScriptErrorType::InvalidInputs, LOCTEXT("CopyMorphTargetToSkeletalMesh_InvalidInput2", "CopyMorphTargetToSkeletalMesh: FromDynamicMesh is Null"));
+		return FromDynamicMesh;
+	}
+
+#if WITH_EDITOR
+	if (ToSkeletalMeshAsset->GetPathName().StartsWith(TEXT("/Engine/")))
+	{
+		const FText Error = FText::Format(LOCTEXT("CopyMorphTargetToSkeletalMesh_BuiltInAsset", "CopyMorphTargetToSkeletalMesh: Cannot modify built-in engine asset: {0}"), FText::FromString(*ToSkeletalMeshAsset->GetPathName()));
+		UE::Geometry::AppendError(Debug, EGeometryScriptErrorType::InvalidInputs, Error);
+		return FromDynamicMesh;
+	}
+
+	// flush any pending rendering commands, which might touch a component while we are rebuilding it's mesh
+	FlushRenderingCommands();
+
+	if (Options.bEmitTransaction)
+	{
+		GEditor->BeginTransaction(LOCTEXT("UpdateSkeletalMesh", "Update Skeletal Mesh"));
+	}
+	
+	// make sure transactional flag is on for this asset
+	ToSkeletalMeshAsset->SetFlags(RF_Transactional);
+
+	verify(ToSkeletalMeshAsset->Modify());
+	
+	// Ensure we have enough LODInfos to cover up to the requested LOD.
+	for (int32 LODIndex = ToSkeletalMeshAsset->GetLODNum(); LODIndex <= TargetLOD.LODIndex; LODIndex++)
+	{
+		FSkeletalMeshLODInfo& LODInfo = ToSkeletalMeshAsset->AddLODInfo();
+		
+		ToSkeletalMeshAsset->GetImportedModel()->LODModels.Add(new FSkeletalMeshLODModel);
+		LODInfo.ReductionSettings.BaseLOD = 0;
+	}
+
+	FMeshDescription* MeshDescription = ToSkeletalMeshAsset->GetMeshDescription(TargetLOD.LODIndex);;
+
+	if (MeshDescription == nullptr)
+	{
+		UE::Geometry::AppendError(Debug, EGeometryScriptErrorType::InvalidInputs, LOCTEXT("CopyMorphTargetToSkeletalMesh_TargetMeshDescription", "CopyMorphTargetToSkeletalMesh: Failed to generate the mesh data for the Target LOD Index"));
+		return FromDynamicMesh;
+	}
+
+	// Morph targets must be compact and have the same number of vertices as the skeletal asset mesh description
+	if (FromDynamicMesh->GetMeshRef().MaxVertexID() != FromDynamicMesh->GetMeshRef().VertexCount() ||
+		FromDynamicMesh->GetMeshRef().MaxVertexID() != MeshDescription->Vertices().Num())
+	{
+		UE::Geometry::AppendError(Debug, EGeometryScriptErrorType::InvalidInputs, LOCTEXT("CopyMorphTargetToSkeletalMesh_InvalidMorphTargetGeometry", "CopyMorphTargetToSkeletalMesh: Morph target mesh doesnt have the same number of vertices as the skeletal mesh."));
+		return FromDynamicMesh;
+	}
+
+	FSkeletalMeshAttributes MeshAttributes(*MeshDescription);
+	MeshAttributes.Register();
+
+	ToSkeletalMeshAsset->ModifyMeshDescription(TargetLOD.LODIndex);
+
+	if (MeshAttributes.GetMorphTargetNames().Contains(MorphTargetName))
+	{
+		UE::Geometry::AppendError(Debug, EGeometryScriptErrorType::InvalidInputs, LOCTEXT("CopyMorphTargetToSkeletalMesh_InvalidMorphTargetName1", "CopyMorphTargetToSkeletalMesh: Morph target name already exists"));
+		return FromDynamicMesh;
+	}
+
+	if (!MeshAttributes.RegisterMorphTargetAttribute(MorphTargetName, false))
+	{
+		UE::Geometry::AppendError(Debug, EGeometryScriptErrorType::InvalidInputs, LOCTEXT("CopyMorphTargetToSkeletalMesh_InvalidMorphTargetName2", "CopyMorphTargetToSkeletalMesh: Morph target name is invalid."));
+		return FromDynamicMesh;
+	}
+
+	TVertexAttributesRef<FVector3f> PositionDelta = MeshAttributes.GetVertexMorphPositionDelta(MorphTargetName);
+	TVertexAttributesRef<FVector3f> VertexPositions = MeshAttributes.GetVertexPositions();
+
+	for (int32 VID = 0; VID < FromDynamicMesh->GetMeshRef().MaxVertexID(); ++VID)
+	{
+		const FVector3d V0 = FromDynamicMesh->GetMeshRef().GetVertex(VID);
+		const FVector3f V1 = VertexPositions[VID];
+
+		PositionDelta.Set(VID, FVector3f(float(V0[0]) - V1[0], 
+										 float(V0[1]) - V1[1], 
+										 float(V0[2]) - V1[2]));
+	}
+
+	ToSkeletalMeshAsset->CommitMeshDescription(TargetLOD.LODIndex);
+
+	if (Options.bDeferMeshPostEditChange == false)
+	{
+		ToSkeletalMeshAsset->PostEditChange();
+	}
+
+	if (Options.bEmitTransaction)
+	{
+		GEditor->EndTransaction();
+	}
+
+	Outcome = EGeometryScriptOutcomePins::Success;
+#else
+	UE::Geometry::AppendError(Debug, EGeometryScriptErrorType::InvalidInputs, LOCTEXT("CopyMorphTargetToSkeletalMesh_EditorOnly", "CopyMorphTargetToSkeletalMesh: Not currently supported at Runtime"));
+#endif
+	
+	return FromDynamicMesh;
+}
 
 #undef LOCTEXT_NAMESPACE
diff --git a/Engine/Plugins/Runtime/GeometryScripting/Source/GeometryScriptingCore/Public/GeometryScript/MeshAssetFunctions.h b/Engine/Plugins/Runtime/GeometryScripting/Source/GeometryScriptingCore/Public/GeometryScript/MeshAssetFunctions.h
--- a/Engine/Plugins/Runtime/GeometryScripting/Source/GeometryScriptingCore/Public/GeometryScript/MeshAssetFunctions.h
+++ b/Engine/Plugins/Runtime/GeometryScripting/Source/GeometryScriptingCore/Public/GeometryScript/MeshAssetFunctions.h
@@ -245,6 +245,26 @@ public:
 		FGeometryScriptMeshWriteLOD TargetLOD,
 		EGeometryScriptOutcomePins& Outcome,
 		UGeometryScriptDebug* Debug = nullptr);
+
+
+   /** 
+	* Add a Dynamic Mesh morph target to a Skeletal Mesh Asset.
+	* 
+	* @param FromMorphTarget the dynamic mesh representing the geometry of the morph target
+	* @param ToSkeletalMeshAsset the asset we are writing the morph target into
+	* @param MorphTargetName the name of the morph target as it will appear in the UI
+	*/
+	UFUNCTION(BlueprintCallable, Category = "GeometryScript|SkeletalMesh", meta = (ExpandEnumAsExecs = "Outcome"))
+	static UPARAM(DisplayName = "Dynamic Mesh") UDynamicMesh* 
+	CopyMorphTargetToSkeletalMesh(
+		UDynamicMesh* FromMorphTarget, 
+		USkeletalMesh* ToSkeletalMeshAsset,
+		FName MorphTargetName,
+		FGeometryScriptCopyMeshToAssetOptions Options,
+		FGeometryScriptMeshWriteLOD TargetLOD,
+		EGeometryScriptOutcomePins& Outcome,
+		UGeometryScriptDebug* Debug = nullptr);
+
 };
 
 
