Add an option to reindex existing skin weights when we are  copying over a new set of bones.

diff --git a/Engine/Plugins/Runtime/GeometryScripting/Source/GeometryScriptingCore/Private/MeshBoneWeightFunctions.cpp b/Engine/Plugins/Runtime/GeometryScripting/Source/GeometryScriptingCore/Private/MeshBoneWeightFunctions.cpp
--- a/Engine/Plugins/Runtime/GeometryScripting/Source/GeometryScriptingCore/Private/MeshBoneWeightFunctions.cpp
+++ b/Engine/Plugins/Runtime/GeometryScripting/Source/GeometryScriptingCore/Private/MeshBoneWeightFunctions.cpp
@@ -456,6 +456,7 @@ UDynamicMesh* UGeometryScriptLibrary_MeshBoneWeightFunctions::TransferBoneWeight
 UDynamicMesh* UGeometryScriptLibrary_MeshBoneWeightFunctions::CopyBonesFromMesh(
 	UDynamicMesh* SourceMesh, 
 	UDynamicMesh* TargetMesh, 
+	FGeometryScriptCopyBonesFromMeshOptions Options,
 	UGeometryScriptDebug* Debug)
 {
 	if (SourceMesh == nullptr)
@@ -490,8 +491,34 @@ UDynamicMesh* UGeometryScriptLibrary_MeshBoneWeightFunctions::CopyBonesFromMesh(
 			{
 				EditMesh.EnableAttributes();
 			}
-				
+			
+			if (Options.ReindexWeights)
+			{
+				if (EditMesh.Attributes()->HasBones())
+				{
+					for (const TPair<FName, TUniquePtr<FDynamicMeshVertexSkinWeightsAttribute>>& AttribPair : EditMesh.Attributes()->GetSkinWeightsAttributes())
+					{
+						FDynamicMeshVertexSkinWeightsAttribute* ToAttrib = EditMesh.Attributes()->GetSkinWeightsAttribute(AttribPair.Key);
+						if (ToAttrib)
+						{	
+							const bool bResult = ToAttrib->ReindexBoneIndicesToSkeleton(EditMesh.Attributes()->GetBoneNames()->GetAttribValues(),
+																	  				    ReadMesh.Attributes()->GetBoneNames()->GetAttribValues());
+							if (!bResult)
+							{
+								const FText Error = FText::Format(LOCTEXT("CopyBonesFromMesh_FailedToReindexWeights", "Failed to reindex bone weights for {0} weights profile"), FText::FromName(AttribPair.Key));
+								AppendError(Debug, EGeometryScriptErrorType::OperationFailed, Error);
+							}
+						}
+					}
+				}
+				else
+				{
+					AppendWarning(Debug, EGeometryScriptErrorType::InvalidInputs, LOCTEXT("CopyBonesFromMesh_TargetMeshHasNoBones", "Bone weight re-indexing was requested but the target mesh has no skeleton data"));
+				}
+			}
+
 			EditMesh.Attributes()->CopyBoneAttributes(*ReadMesh.Attributes());
+
 		}, EDynamicMeshChangeType::AttributeEdit, EDynamicMeshAttributeChangeFlags::Unknown, false);
 	});
 
diff --git a/Engine/Plugins/Runtime/GeometryScripting/Source/GeometryScriptingCore/Public/GeometryScript/MeshBoneWeightFunctions.h b/Engine/Plugins/Runtime/GeometryScripting/Source/GeometryScriptingCore/Public/GeometryScript/MeshBoneWeightFunctions.h
--- a/Engine/Plugins/Runtime/GeometryScripting/Source/GeometryScriptingCore/Public/GeometryScript/MeshBoneWeightFunctions.h
+++ b/Engine/Plugins/Runtime/GeometryScripting/Source/GeometryScriptingCore/Public/GeometryScript/MeshBoneWeightFunctions.h
@@ -188,6 +188,17 @@ struct GEOMETRYSCRIPTINGCORE_API FGeometryScriptBoneInfo
 	FLinearColor Color = FLinearColor::White;
 };
 
+USTRUCT(BlueprintType)
+struct GEOMETRYSCRIPTINGCORE_API FGeometryScriptCopyBonesFromMeshOptions
+{
+	GENERATED_BODY()
+	
+	/** If the target Dynamic mesh has bone weights and a skeleton, re-index (re-bind) target weight indices from the  
+	 * target skeleton to the source skeleton.*/
+	UPROPERTY(EditAnywhere, BlueprintReadWrite, Category = Options)
+	bool ReindexWeights = false;
+};
+
 UCLASS(meta = (ScriptName = "GeometryScript_BoneWeights"))
 class GEOMETRYSCRIPTINGCORE_API UGeometryScriptLibrary_MeshBoneWeightFunctions : public UBlueprintFunctionLibrary
 {
@@ -338,6 +349,7 @@ public:
 	CopyBonesFromMesh(
 		UDynamicMesh* SourceMesh,
 		UDynamicMesh* TargetMesh,
+		FGeometryScriptCopyBonesFromMeshOptions Options = FGeometryScriptCopyBonesFromMeshOptions(),
 		UGeometryScriptDebug* Debug = nullptr);
 
 	/**
