Fix mismatch error return value in Geometry Script function GetMeshSelectionBoundaryLoops in MeshSelectionQueryFunctions.cpp

#rnx

diff --git a/Engine/Plugins/Runtime/GeometryScripting/Source/GeometryScriptingCore/Private/MeshSelectionQueryFunctions.cpp b/Engine/Plugins/Runtime/GeometryScripting/Source/GeometryScriptingCore/Private/MeshSelectionQueryFunctions.cpp
--- a/Engine/Plugins/Runtime/GeometryScripting/Source/GeometryScriptingCore/Private/MeshSelectionQueryFunctions.cpp
+++ b/Engine/Plugins/Runtime/GeometryScripting/Source/GeometryScriptingCore/Private/MeshSelectionQueryFunctions.cpp
@@ -93,7 +93,8 @@ UDynamicMesh* UGeometryScriptLibrary_MeshSelectionQueryFunctions::GetMeshSelecti
 		// TODO: if #Triangles == Mesh.TriangleCount, use MeshBoundaryLoops
 
 		FMeshRegionBoundaryLoops Loops(&ReadMesh, Triangles, false);
-		bFoundErrors = Loops.Compute();
+		// Compute returns false if errors occurred
+		bFoundErrors = !Loops.Compute();
 		NumLoops = Loops.Num();
 		IndexLoops.Reserve(NumLoops);
 		PathLoops.Reserve(NumLoops);
