Use custom options structure for CopyMorphTargetToSkeletalMesh geometry script node.
Add an option to force overwrite existing morph target if one with the same name already exist.

#rnx

diff --git a/Engine/Plugins/Runtime/GeometryScripting/Source/GeometryScriptingCore/Private/MeshAssetFunctions.cpp b/Engine/Plugins/Runtime/GeometryScripting/Source/GeometryScriptingCore/Private/MeshAssetFunctions.cpp
--- a/Engine/Plugins/Runtime/GeometryScripting/Source/GeometryScriptingCore/Private/MeshAssetFunctions.cpp
+++ b/Engine/Plugins/Runtime/GeometryScripting/Source/GeometryScriptingCore/Private/MeshAssetFunctions.cpp
@@ -540,7 +540,7 @@ void UGeometryScriptLibrary_StaticMeshFunctions::GetSectionMaterialListFromStati
 
 namespace UELocal
 {
-	bool CopyMeshFromSkeletalMesh_RenderData(USkeletalMesh* FromSkeletalMeshAsset, FGeometryScriptCopyMeshFromAssetOptions AssetOptions, EGeometryScriptLODType LODType, int32 LODIndex, UDynamicMesh* ToDynamicMesh, UGeometryScriptDebug* Debug)
+	bool CopyMeshFromSkeletalMesh_RenderData(USkeletalMesh* FromSkeletalMeshAsset, FGeometryScriptCopyMeshFromAssetOptions AssetOptions, int32 LODIndex, UDynamicMesh* ToDynamicMesh, UGeometryScriptDebug* Debug)
 	{
 		
 
@@ -627,7 +627,7 @@ UDynamicMesh* UGeometryScriptLibrary_StaticMeshFunctions::CopyMeshFromSkeletalMe
 	}
 	else if (RequestedLOD.LODType == EGeometryScriptLODType::RenderData)
 	{
-		if (UELocal::CopyMeshFromSkeletalMesh_RenderData(FromSkeletalMeshAsset, AssetOptions, RequestedLOD.LODType, RequestedLOD.LODIndex, ToDynamicMesh, Debug))
+		if (UELocal::CopyMeshFromSkeletalMesh_RenderData(FromSkeletalMeshAsset, AssetOptions, RequestedLOD.LODIndex, ToDynamicMesh, Debug))
 		{
 			Outcome = EGeometryScriptOutcomePins::Success;
 		}
@@ -846,7 +846,7 @@ UDynamicMesh* UGeometryScriptLibrary_StaticMeshFunctions::CopyMorphTargetToSkele
 		UDynamicMesh* FromDynamicMesh, 
 		USkeletalMesh* ToSkeletalMeshAsset,
 		FName MorphTargetName,
-		FGeometryScriptCopyMeshToAssetOptions Options,
+		FGeometryScriptCopyMorphTargetToAssetOptions Options,
 		FGeometryScriptMeshWriteLOD TargetLOD,
 		EGeometryScriptOutcomePins& Outcome,
 		UGeometryScriptDebug* Debug)
@@ -917,14 +917,19 @@ UDynamicMesh* UGeometryScriptLibrary_StaticMeshFunctions::CopyMorphTargetToSkele
 
 	if (MeshAttributes.GetMorphTargetNames().Contains(MorphTargetName))
 	{
-		UE::Geometry::AppendError(Debug, EGeometryScriptErrorType::InvalidInputs, LOCTEXT("CopyMorphTargetToSkeletalMesh_InvalidMorphTargetName1", "CopyMorphTargetToSkeletalMesh: Morph target name already exists"));
-		return FromDynamicMesh;
+		if (!Options.bOverwriteExistingTarget) // only throw error if we dont want to overwrite the existing target 
+		{
+			UE::Geometry::AppendError(Debug, EGeometryScriptErrorType::InvalidInputs, LOCTEXT("CopyMorphTargetToSkeletalMesh_InvalidMorphTargetName1", "CopyMorphTargetToSkeletalMesh: Morph target name already exists"));
+			return FromDynamicMesh;
+		}
 	}
-
-	if (!MeshAttributes.RegisterMorphTargetAttribute(MorphTargetName, false))
+	else
 	{
-		UE::Geometry::AppendError(Debug, EGeometryScriptErrorType::InvalidInputs, LOCTEXT("CopyMorphTargetToSkeletalMesh_InvalidMorphTargetName2", "CopyMorphTargetToSkeletalMesh: Morph target name is invalid."));
-		return FromDynamicMesh;
+		if (!MeshAttributes.RegisterMorphTargetAttribute(MorphTargetName, false))
+		{
+			UE::Geometry::AppendError(Debug, EGeometryScriptErrorType::InvalidInputs, LOCTEXT("CopyMorphTargetToSkeletalMesh_InvalidMorphTargetName2", "CopyMorphTargetToSkeletalMesh: Morph target name is invalid."));
+			return FromDynamicMesh;
+		}
 	}
 
 	TVertexAttributesRef<FVector3f> PositionDelta = MeshAttributes.GetVertexMorphPositionDelta(MorphTargetName);
diff --git a/Engine/Plugins/Runtime/GeometryScripting/Source/GeometryScriptingCore/Public/GeometryScript/MeshAssetFunctions.h b/Engine/Plugins/Runtime/GeometryScripting/Source/GeometryScriptingCore/Public/GeometryScript/MeshAssetFunctions.h
--- a/Engine/Plugins/Runtime/GeometryScripting/Source/GeometryScriptingCore/Public/GeometryScript/MeshAssetFunctions.h
+++ b/Engine/Plugins/Runtime/GeometryScripting/Source/GeometryScriptingCore/Public/GeometryScript/MeshAssetFunctions.h
@@ -107,7 +107,23 @@ public:
 	bool bDeferMeshPostEditChange = false;
 };
 
+USTRUCT(BlueprintType)
+struct GEOMETRYSCRIPTINGCORE_API FGeometryScriptCopyMorphTargetToAssetOptions
+{
+	GENERATED_BODY()
+public:
 
+	/** If true and the morph target with the given name exists, it will be overwritten. If false, will abort and print a 
+	 * console error. */
+	UPROPERTY(EditAnywhere, BlueprintReadWrite, Category = Options)
+	bool bOverwriteExistingTarget = false;
+
+	UPROPERTY(EditAnywhere, BlueprintReadWrite, Category = Options)
+	bool bEmitTransaction = true;
+
+	UPROPERTY(EditAnywhere, BlueprintReadWrite, Category = Options)
+	bool bDeferMeshPostEditChange = false;
+};
 
 // Although the class name indicates StaticMeshFunctions, that was a naming mistake that is difficult
 // to correct. This class is intended to serve as a generic asset utils function library. The naming
@@ -260,7 +276,7 @@ public:
 		UDynamicMesh* FromMorphTarget, 
 		USkeletalMesh* ToSkeletalMeshAsset,
 		FName MorphTargetName,
-		FGeometryScriptCopyMeshToAssetOptions Options,
+		FGeometryScriptCopyMorphTargetToAssetOptions Options,
 		FGeometryScriptMeshWriteLOD TargetLOD,
 		EGeometryScriptOutcomePins& Outcome,
 		UGeometryScriptDebug* Debug = nullptr);
