add max miter scale to geometry script sweep functions and apply mitering if the value is > 1
 (+ make them consistently have a rotation parameter)

diff --git a/Engine/Plugins/Runtime/GeometryScripting/Source/GeometryScriptingCore/Private/MeshPrimitiveFunctions.cpp b/Engine/Plugins/Runtime/GeometryScripting/Source/GeometryScriptingCore/Private/MeshPrimitiveFunctions.cpp
--- a/Engine/Plugins/Runtime/GeometryScripting/Source/GeometryScriptingCore/Private/MeshPrimitiveFunctions.cpp
+++ b/Engine/Plugins/Runtime/GeometryScripting/Source/GeometryScriptingCore/Private/MeshPrimitiveFunctions.cpp
@@ -647,6 +647,7 @@ UDynamicMesh* UGeometryScriptLibrary_MeshPrimitiveFunctions::AppendSweepPolyline
 	float StartScale,
 	float EndScale,
 	float RotationAngleDeg,
+	float MaxMiterScale,
 	UGeometryScriptDebug* Debug)
 {
 	if (TargetMesh == nullptr)
@@ -706,6 +707,10 @@ UDynamicMesh* UGeometryScriptLibrary_MeshPrimitiveFunctions::AppendSweepPolyline
 	SweepGen.InitialFrame = FFrame3d(SweepGen.Path[0]);
 	SweepGen.StartScale = StartScale;
 	SweepGen.EndScale = EndScale;
+	if (MaxMiterScale > 1.0)
+	{
+		SweepGen.EnableMitering((double)MaxMiterScale);
+	}
 
 	SweepGen.Generate();
 
@@ -775,6 +780,8 @@ UDynamicMesh* UGeometryScriptLibrary_MeshPrimitiveFunctions::AppendSimpleSweptPo
 	bool bCapped,
 	float StartScale,
 	float EndScale,
+	float RotationAngleDeg,
+	float MaxMiterScale,
 	UGeometryScriptDebug* Debug)
 {
 	if (TargetMesh == nullptr)
@@ -794,9 +801,10 @@ UDynamicMesh* UGeometryScriptLibrary_MeshPrimitiveFunctions::AppendSimpleSweptPo
 	}
 
 	FGeneralizedCylinderGenerator SweepGen;
+	FMatrix2d Rotation2D = FMatrix2d::RotationDeg(-RotationAngleDeg);
 	for (FVector2D Point : PolygonVertices)
 	{
-		SweepGen.CrossSection.AppendVertex(FVector2d(Point.X, Point.Y));
+		SweepGen.CrossSection.AppendVertex(Rotation2D * FVector2d(Point.X, Point.Y));
 	}
 	for (FVector SweepPathPos : SweepPath)
 	{
@@ -809,6 +817,10 @@ UDynamicMesh* UGeometryScriptLibrary_MeshPrimitiveFunctions::AppendSimpleSweptPo
 	SweepGen.InitialFrame = FFrame3d(SweepGen.Path[0]);
 	SweepGen.StartScale = StartScale;
 	SweepGen.EndScale = EndScale;
+	if (MaxMiterScale > 1.0f)
+	{
+		SweepGen.EnableMitering((double)MaxMiterScale);
+	}
 
 	SweepGen.Generate();
 
@@ -829,6 +841,7 @@ UDynamicMesh* UGeometryScriptLibrary_MeshPrimitiveFunctions::AppendSweepPolygon(
 	float StartScale,
 	float EndScale,
 	float RotationAngleDeg,
+	float MaxMiterScale,
 	UGeometryScriptDebug* Debug)
 {
 	if (TargetMesh == nullptr)
@@ -870,6 +883,10 @@ UDynamicMesh* UGeometryScriptLibrary_MeshPrimitiveFunctions::AppendSweepPolygon(
 	SweepGen.InitialFrame = FFrame3d(SweepGen.Path[0]);
 	SweepGen.StartScale = StartScale;
 	SweepGen.EndScale = EndScale;
+	if (MaxMiterScale > 1.0f)
+	{
+		SweepGen.EnableMitering((double)MaxMiterScale);
+	}
 
 	SweepGen.Generate();
 
diff --git a/Engine/Plugins/Runtime/GeometryScripting/Source/GeometryScriptingCore/Public/GeometryScript/MeshPrimitiveFunctions.h b/Engine/Plugins/Runtime/GeometryScripting/Source/GeometryScriptingCore/Public/GeometryScript/MeshPrimitiveFunctions.h
--- a/Engine/Plugins/Runtime/GeometryScripting/Source/GeometryScriptingCore/Public/GeometryScript/MeshPrimitiveFunctions.h
+++ b/Engine/Plugins/Runtime/GeometryScripting/Source/GeometryScriptingCore/Public/GeometryScript/MeshPrimitiveFunctions.h
@@ -364,6 +364,7 @@ public:
 	 * @param StartScale uniform scaling applied to the 2D polygon at the start of the path. Interpolated via arc length to EndScale at the end of the path.
 	 * @param EndScale uniform scaling applied to the 2D polygon at the end of the path
 	 * @param RotationAngleDeg Rotation applied to the 2D Polygon. Positive rotation rotates clockwise, ie Up/+Z/+V towards Right/+Y/+U. This Rotation is applied before any rotation in the SweepPath Transforms.
+	 * @param MaxMiterScale If > 1, maximum scaling to apply at sharp turns in the curve path to keep the apparent swept profile from shrinking, and sweep path frames will be aligned to the path direction
 	 */
 	UFUNCTION(BlueprintCallable, Category = "GeometryScript|Primitives", meta=(ScriptMethod, AutoCreateRefTerm="PolylineTexParamU, SweepPathTexParamV"))
 	static UPARAM(DisplayName = "Target Mesh") UDynamicMesh* 
@@ -379,6 +380,7 @@ public:
 		float StartScale = 1.0f,
 		float EndScale = 1.0f,
 		float RotationAngleDeg = 0.0f,
+		float MaxMiterScale = 1.0f,
 		UGeometryScriptDebug* Debug = nullptr);
 
 	/**
@@ -399,8 +401,16 @@ public:
 		UGeometryScriptDebug* Debug = nullptr);
 
 	/**
-	* Sweeps a 2D polygon along an arbitrary 3D path, appending the result to the Target Mesh.
-	*/
+	 * Sweeps a 2D polygon along an arbitrary 3D path, appending the result to the Target Mesh.
+	 * @param PolygonVertices vertices of the closed 2D polyon that will be swept along the SweepPath
+	 * @param SweepPath defines the 3D sweep path curve
+	 * @param bLoop if true, SweepPath is considered to be a Loop and a section connecting the end and start of the path is added (bCapped is ignored)
+	 * @param bCapped if true the open ends of the swept generalized cylinder are triangulated
+	 * @param StartScale uniform scaling applied to the 2D polygon at the start of the path. Interpolated via arc length to EndScale at the end of the path.
+	 * @param EndScale uniform scaling applied to the 2D polygon at the end of the path
+	 * @param RotationAngleDeg Rotation applied to the 2D Polygon. Positive rotation rotates clockwise, ie Up/+Z/+V towards Right/+Y/+U
+	 * @param MaxMiterScale If > 1, maximum scaling to apply at sharp turns in the curve path to keep the apparent swept profile from shrinking
+	 */
 	UFUNCTION(BlueprintCallable, Category = "GeometryScript|Primitives", meta=(ScriptMethod))
 	static UPARAM(DisplayName = "Target Mesh") UDynamicMesh* 
 	AppendSimpleSweptPolygon( 
@@ -413,6 +423,8 @@ public:
 		bool bCapped = true,
 		float StartScale = 1.0f,
 		float EndScale = 1.0f,
+		float RotationAngleDeg = 0.0f,
+		float MaxMiterScale = 1.0f,
 		UGeometryScriptDebug* Debug = nullptr);
 
 	/**
@@ -426,6 +438,7 @@ public:
 	 * @param StartScale uniform scaling applied to the 2D polygon at the start of the path. Interpolated via arc length to EndScale at the end of the path.
 	 * @param EndScale uniform scaling applied to the 2D polygon at the end of the path
 	 * @param RotationAngleDeg Rotation applied to the 2D Polygon. Positive rotation rotates clockwise, ie Up/+Z/+V towards Right/+Y/+U
+	 * @param MaxMiterScale If > 1, maximum scaling to apply at sharp turns in the curve path to keep the apparent swept profile from shrinking, and sweep path frames will be aligned to the path direction
 	 */
 	UFUNCTION(BlueprintCallable, Category = "GeometryScript|Primitives", meta=(ScriptMethod))
 	static UPARAM(DisplayName = "Target Mesh") UDynamicMesh* 
@@ -440,6 +453,7 @@ public:
 		float StartScale = 1.0f,
 		float EndScale = 1.0f,
 		float RotationAngleDeg = 0.0f,
+		float MaxMiterScale = 1.0f,
 		UGeometryScriptDebug* Debug = nullptr);
 
 
