Change MaxMiterScale to MiterLimit in C++ and BP sweep methods to reduce confusion and to match the naming for the same concept in CSS/SVG

diff --git a/Engine/Plugins/Runtime/GeometryScripting/Source/GeometryScriptingCore/Private/MeshPrimitiveFunctions.cpp b/Engine/Plugins/Runtime/GeometryScripting/Source/GeometryScriptingCore/Private/MeshPrimitiveFunctions.cpp
--- a/Engine/Plugins/Runtime/GeometryScripting/Source/GeometryScriptingCore/Private/MeshPrimitiveFunctions.cpp
+++ b/Engine/Plugins/Runtime/GeometryScripting/Source/GeometryScriptingCore/Private/MeshPrimitiveFunctions.cpp
@@ -647,7 +647,7 @@ UDynamicMesh* UGeometryScriptLibrary_MeshPrimitiveFunctions::AppendSweepPolyline
 	float StartScale,
 	float EndScale,
 	float RotationAngleDeg,
-	float MaxMiterScale,
+	float MiterLimit,
 	UGeometryScriptDebug* Debug)
 {
 	if (TargetMesh == nullptr)
@@ -707,9 +707,9 @@ UDynamicMesh* UGeometryScriptLibrary_MeshPrimitiveFunctions::AppendSweepPolyline
 	SweepGen.InitialFrame = FFrame3d(SweepGen.Path[0]);
 	SweepGen.StartScale = StartScale;
 	SweepGen.EndScale = EndScale;
-	if (MaxMiterScale > 1.0)
+	if (MiterLimit > 1.0)
 	{
-		SweepGen.EnableMitering((double)MaxMiterScale);
+		SweepGen.EnableMitering((double)MiterLimit);
 	}
 
 	SweepGen.Generate();
@@ -781,7 +781,7 @@ UDynamicMesh* UGeometryScriptLibrary_MeshPrimitiveFunctions::AppendSimpleSweptPo
 	float StartScale,
 	float EndScale,
 	float RotationAngleDeg,
-	float MaxMiterScale,
+	float MiterLimit,
 	UGeometryScriptDebug* Debug)
 {
 	if (TargetMesh == nullptr)
@@ -817,9 +817,9 @@ UDynamicMesh* UGeometryScriptLibrary_MeshPrimitiveFunctions::AppendSimpleSweptPo
 	SweepGen.InitialFrame = FFrame3d(SweepGen.Path[0]);
 	SweepGen.StartScale = StartScale;
 	SweepGen.EndScale = EndScale;
-	if (MaxMiterScale > 1.0f)
+	if (MiterLimit > 1.0f)
 	{
-		SweepGen.EnableMitering((double)MaxMiterScale);
+		SweepGen.EnableMitering((double)MiterLimit);
 	}
 
 	SweepGen.Generate();
@@ -841,7 +841,7 @@ UDynamicMesh* UGeometryScriptLibrary_MeshPrimitiveFunctions::AppendSweepPolygon(
 	float StartScale,
 	float EndScale,
 	float RotationAngleDeg,
-	float MaxMiterScale,
+	float MiterLimit,
 	UGeometryScriptDebug* Debug)
 {
 	if (TargetMesh == nullptr)
@@ -883,9 +883,9 @@ UDynamicMesh* UGeometryScriptLibrary_MeshPrimitiveFunctions::AppendSweepPolygon(
 	SweepGen.InitialFrame = FFrame3d(SweepGen.Path[0]);
 	SweepGen.StartScale = StartScale;
 	SweepGen.EndScale = EndScale;
-	if (MaxMiterScale > 1.0f)
+	if (MiterLimit > 1.0f)
 	{
-		SweepGen.EnableMitering((double)MaxMiterScale);
+		SweepGen.EnableMitering((double)MiterLimit);
 	}
 
 	SweepGen.Generate();
diff --git a/Engine/Plugins/Runtime/GeometryScripting/Source/GeometryScriptingCore/Public/GeometryScript/MeshPrimitiveFunctions.h b/Engine/Plugins/Runtime/GeometryScripting/Source/GeometryScriptingCore/Public/GeometryScript/MeshPrimitiveFunctions.h
--- a/Engine/Plugins/Runtime/GeometryScripting/Source/GeometryScriptingCore/Public/GeometryScript/MeshPrimitiveFunctions.h
+++ b/Engine/Plugins/Runtime/GeometryScripting/Source/GeometryScriptingCore/Public/GeometryScript/MeshPrimitiveFunctions.h
@@ -364,7 +364,7 @@ public:
 	 * @param StartScale uniform scaling applied to the 2D polygon at the start of the path. Interpolated via arc length to EndScale at the end of the path.
 	 * @param EndScale uniform scaling applied to the 2D polygon at the end of the path
 	 * @param RotationAngleDeg Rotation applied to the 2D Polygon. Positive rotation rotates clockwise, ie Up/+Z/+V towards Right/+Y/+U. This Rotation is applied before any rotation in the SweepPath Transforms.
-	 * @param MaxMiterScale If > 1, maximum scaling to apply at sharp turns in the curve path to keep the apparent swept profile from shrinking, and sweep path frames will be aligned to the path direction
+	 * @param MiterLimit If > 1, maximum scaling to apply at sharp turns in the curve path to keep the apparent swept profile from shrinking, and sweep path frames will be aligned to the path direction
 	 */
 	UFUNCTION(BlueprintCallable, Category = "GeometryScript|Primitives", meta=(ScriptMethod, AutoCreateRefTerm="PolylineTexParamU, SweepPathTexParamV"))
 	static UPARAM(DisplayName = "Target Mesh") UDynamicMesh* 
@@ -380,7 +380,7 @@ public:
 		float StartScale = 1.0f,
 		float EndScale = 1.0f,
 		float RotationAngleDeg = 0.0f,
-		float MaxMiterScale = 1.0f,
+		float MiterLimit = 1.0f,
 		UGeometryScriptDebug* Debug = nullptr);
 
 	/**
@@ -409,7 +409,7 @@ public:
 	 * @param StartScale uniform scaling applied to the 2D polygon at the start of the path. Interpolated via arc length to EndScale at the end of the path.
 	 * @param EndScale uniform scaling applied to the 2D polygon at the end of the path
 	 * @param RotationAngleDeg Rotation applied to the 2D Polygon. Positive rotation rotates clockwise, ie Up/+Z/+V towards Right/+Y/+U
-	 * @param MaxMiterScale If > 1, maximum scaling to apply at sharp turns in the curve path to keep the apparent swept profile from shrinking
+	 * @param MiterLimit If > 1, maximum scaling to apply at sharp turns in the curve path to keep the apparent swept profile from shrinking
 	 */
 	UFUNCTION(BlueprintCallable, Category = "GeometryScript|Primitives", meta=(ScriptMethod))
 	static UPARAM(DisplayName = "Target Mesh") UDynamicMesh* 
@@ -424,7 +424,7 @@ public:
 		float StartScale = 1.0f,
 		float EndScale = 1.0f,
 		float RotationAngleDeg = 0.0f,
-		float MaxMiterScale = 1.0f,
+		float MiterLimit = 1.0f,
 		UGeometryScriptDebug* Debug = nullptr);
 
 	/**
@@ -437,8 +437,8 @@ public:
 	 * @param bCapped if true the open ends of the swept generalized cylinder are triangulated
 	 * @param StartScale uniform scaling applied to the 2D polygon at the start of the path. Interpolated via arc length to EndScale at the end of the path.
 	 * @param EndScale uniform scaling applied to the 2D polygon at the end of the path
-	 * @param RotationAngleDeg Rotation applied to the 2D Polygon. Positive rotation rotates clockwise, ie Up/+Z/+V towards Right/+Y/+U
-	 * @param MaxMiterScale If > 1, maximum scaling to apply at sharp turns in the curve path to keep the apparent swept profile from shrinking, and sweep path frames will be aligned to the path direction
+	 * @param RotationAngleDeg Rotation applied to the 2D Polygon. Positive rotation rotates clockwise, ie Up/+Z/+V towards Right/+Y/+U. This Rotation is applied before any rotation in the SweepPath Transforms.
+	 * @param MiterLimit If > 1, maximum scaling to apply at sharp turns in the curve path to keep the apparent swept profile from shrinking, and sweep path frames will be aligned to the path direction
 	 */
 	UFUNCTION(BlueprintCallable, Category = "GeometryScript|Primitives", meta=(ScriptMethod))
 	static UPARAM(DisplayName = "Target Mesh") UDynamicMesh* 
@@ -453,7 +453,7 @@ public:
 		float StartScale = 1.0f,
 		float EndScale = 1.0f,
 		float RotationAngleDeg = 0.0f,
-		float MaxMiterScale = 1.0f,
+		float MiterLimit = 1.0f,
 		UGeometryScriptDebug* Debug = nullptr);
 
 
diff --git a/Engine/Source/Runtime/GeometryCore/Public/Generators/SweepGenerator.h b/Engine/Source/Runtime/GeometryCore/Public/Generators/SweepGenerator.h
--- a/Engine/Source/Runtime/GeometryCore/Public/Generators/SweepGenerator.h
+++ b/Engine/Source/Runtime/GeometryCore/Public/Generators/SweepGenerator.h
@@ -738,20 +738,29 @@ public:
 	double StartScale = 1.0;
 	double EndScale = 1.0;
 
-	// Maximum mitering scale to apply at turns; only applied if greater than 1
+	// Maximum factor by which mitering can expand the cross section at sharp turns, to give the appearance of a consistent cross section width. Only used if > 1.0.
 	// Note: If PathFrames are specified, then bAlignFramesToSampledTangents must be true for mitering to be performed.
-	double MaxMiterScale = 1.0;
+	// (Similar to the MiterLimit concept in SVG / CSS)
+	double MiterLimit = 1.0;
 
 	// Whether to align the frames to the path tangents. Only relevant if PathFrames is specified.
 	bool bAlignFramesToSampledTangents = false;
 
-	// Configure settings to support 'mitering' -- i.e., scaling the cross sections as needed to maintain consistent cross section size through sharp corners, up to the specified scale limit
-	void EnableMitering(double InMaxMiterScale = 10)
+	// Configure settings to support mitering -- i.e., scaling the cross sections as needed to maintain consistent cross section size through sharp corners, up to the specified scale limit
+	void EnableMitering(double InMiterLimit = 10)
 	{
-		MaxMiterScale = InMaxMiterScale;
+		MiterLimit = InMiterLimit;
 		bAlignFramesToSampledTangents = true;
 	}
 
+	// Set the MiterLimit based on the maximum turn angle at which a correct miter should be applied. For turns sharper than this angle, the cross section will appear to shrink at the turn.
+	// @param MiterAngleLimitInDeg Maximum turn angle for correct mitering; should be >= 0 and < 180
+	void SetMiterLimitByAngle(double MiterAngleLimitInDeg)
+	{
+		MiterAngleLimitInDeg = FMath::Clamp(MiterAngleLimitInDeg, 0, 180 - FMathd::ZeroTolerance);
+		MiterLimit = 1.0/FMath::Sin(MiterAngleLimitInDeg * FMathd::DegToRad * .5);
+	}
+
 	// When true, the generator attempts to scale UV's in a way that preserves scaling across different mesh
 	// results, aiming for 1.0 in UV space to be equal to UnitUVInWorldCoordinates in world space. This in
 	// practice means adjusting the U scale relative to the CrossSection curve length and V scale relative
@@ -845,7 +854,7 @@ public:
 
 				FVector3d Tangent = Normalized(-ToPrev + ToNext);
 				CrossSectionFrame.AlignAxis(2, Tangent);
-				if (MaxMiterScale > 1)
+				if (MiterLimit > 1)
 				{
 					// only miter if neither neighbor vector is exactly zero
 					if (!ToPrev.IsZero() && !ToNext.IsZero())
@@ -857,9 +866,9 @@ public:
 						{
 							double CosTwoTheta = ToPrev.Dot(ToNext);
 							double SqrScale = 2.0 / (1.0 - CosTwoTheta);
-							if (!FMath::IsFinite(SqrScale) || SqrScale > MaxMiterScale * MaxMiterScale)
+							if (!FMath::IsFinite(SqrScale) || SqrScale > MiterLimit * MiterLimit)
 							{
-								ScaleAmount = MaxMiterScale;
+								ScaleAmount = MiterLimit;
 							}
 							else
 							{
