Add geometry script methods for selection by material ID and by polygroup, based on github PR
PR #11742: Add mesh selection by material id for dynamic mesh

diff --git a/Engine/Plugins/Runtime/GeometryScripting/Source/GeometryScriptingCore/Private/MeshSelectionFunctions.cpp b/Engine/Plugins/Runtime/GeometryScripting/Source/GeometryScriptingCore/Private/MeshSelectionFunctions.cpp
--- a/Engine/Plugins/Runtime/GeometryScripting/Source/GeometryScriptingCore/Private/MeshSelectionFunctions.cpp
+++ b/Engine/Plugins/Runtime/GeometryScripting/Source/GeometryScriptingCore/Private/MeshSelectionFunctions.cpp
@@ -21,6 +21,130 @@
 
 
 
+namespace UE::MeshSelectionLocals
+{
+
+// Helper to create a selection for all triangles matching a given filter (or all vertices/groups referencing those triangles)
+static void SelectByTriangleAttribute(const FDynamicMesh3& ReadMesh, 
+	FGeometryScriptMeshSelection& Selection, EGeometryScriptMeshSelectionType SelectionType,
+	TFunctionRef<bool(int32)> TriangleFilter
+)
+{
+	FGeometrySelection NewSelection;
+	if (SelectionType == EGeometryScriptMeshSelectionType::Vertices)
+	{
+		NewSelection.InitializeTypes(EGeometryElementType::Vertex, EGeometryTopologyType::Triangle);
+		for (int32 TID : ReadMesh.TriangleIndicesItr())
+		{
+			if (TriangleFilter(TID))
+			{
+				FIndex3i Tri = ReadMesh.GetTriangle(TID);
+
+				for (int32 SubIdx = 0; SubIdx < 3; ++SubIdx)
+				{
+					NewSelection.Selection.Add(FGeoSelectionID::MeshVertex(Tri[SubIdx]).Encoded());
+				}
+			}
+		}
+	}
+	else if (SelectionType == EGeometryScriptMeshSelectionType::Triangles)
+	{
+		NewSelection.InitializeTypes(EGeometryElementType::Face, EGeometryTopologyType::Triangle);
+		for (int32 TID : ReadMesh.TriangleIndicesItr())
+		{
+			if (TriangleFilter(TID))
+			{
+				NewSelection.Selection.Add(FGeoSelectionID::MeshTriangle(TID).Encoded());
+			}
+		}
+	}
+	else
+	{
+		NewSelection.InitializeTypes(EGeometryElementType::Face, EGeometryTopologyType::Polygroup);
+		TSet<int32> UniqueGroupIDs;
+		for (int32 TID : ReadMesh.TriangleIndicesItr())
+		{
+			if (!TriangleFilter(TID))
+			{
+				continue;
+			}
+
+			int32 GroupID = ReadMesh.GetTriangleGroup(TID);
+			if (UniqueGroupIDs.Contains(GroupID) == false)
+			{
+				NewSelection.Selection.Add(FGeoSelectionID::GroupFace(TID, GroupID).Encoded());
+				UniqueGroupIDs.Add(GroupID);
+			}
+		}
+	}
+	Selection.SetSelection(MoveTemp(NewSelection));
+}
+
+// Helper to select elements based on position/normal of triangles
+static void SelectMeshElementsWithContainmentTest(
+	UDynamicMesh* TargetMesh,
+	TFunctionRef<bool(const FVector3d& Position, const FVector3d& Normal)> ContainmentFunc,
+	FGeometryScriptMeshSelection& SelectionOut,
+	EGeometryScriptMeshSelectionType SelectionType,
+	int NumTrianglePoints,
+	bool bNeedsNormals)
+{
+	NumTrianglePoints = FMath::Clamp(NumTrianglePoints, 1, 3);
+
+	FGeometrySelection GeoSelection;
+	if (SelectionType == EGeometryScriptMeshSelectionType::Vertices)
+	{
+		GeoSelection.InitializeTypes(EGeometryElementType::Vertex, EGeometryTopologyType::Triangle);
+
+		TargetMesh->ProcessMesh([&ContainmentFunc, &GeoSelection, bNeedsNormals](const FDynamicMesh3& Mesh)
+		{
+			for (int32 vid : Mesh.VertexIndicesItr())
+			{
+				FVector3d UseNormal = (bNeedsNormals) ? FMeshNormals::ComputeVertexNormal(Mesh,vid) : FVector3d::UnitZ();
+				if ( ContainmentFunc(Mesh.GetVertex(vid), UseNormal) )
+				{
+					GeoSelection.Selection.Add( FGeoSelectionID::MeshVertex(vid).Encoded() );
+				}
+			}
+		});
+	}
+	else
+	{
+		GeoSelection.InitializeTypes(EGeometryElementType::Face, 
+			(SelectionType == EGeometryScriptMeshSelectionType::Triangles) ? EGeometryTopologyType::Triangle : EGeometryTopologyType::Polygroup);
+
+		TargetMesh->ProcessMesh([&ContainmentFunc, &GeoSelection, SelectionType, NumTrianglePoints, bNeedsNormals](const FDynamicMesh3& Mesh)
+		{
+			for (int32 tid : Mesh.TriangleIndicesItr())
+			{
+				FVector3d UseNormal = (bNeedsNormals) ? Mesh.GetTriNormal(tid) : FVector3d::UnitZ();
+				FIndex3i Tri = Mesh.GetTriangle(tid);
+				// may be wasteful to test each vertex multiple times...could accumulate a cache at cost of some memory allocation...
+				int NumContained =
+					(ContainmentFunc(Mesh.GetVertex(Tri.A), UseNormal) ? 1 : 0) +
+					(ContainmentFunc(Mesh.GetVertex(Tri.B), UseNormal) ? 1 : 0) +
+					(ContainmentFunc(Mesh.GetVertex(Tri.C), UseNormal) ? 1 : 0);
+
+				if ( NumContained >= NumTrianglePoints )
+				{
+					if (SelectionType == EGeometryScriptMeshSelectionType::Triangles)
+					{
+						GeoSelection.Selection.Add( FGeoSelectionID::MeshTriangle(tid).Encoded() );
+					}
+					else
+					{
+						int32 gid = Mesh.GetTriangleGroup(tid);
+						GeoSelection.Selection.Add( FGeoSelectionID::GroupFace(tid, gid).Encoded() );
+					}
+				}
+			}
+		});
+
+	}
+
+	SelectionOut.SetSelection(MoveTemp(GeoSelection));
+}
+}
 
 
 void UGeometryScriptLibrary_MeshSelectionFunctions::GetMeshSelectionInfo(
@@ -92,6 +216,75 @@
 	return TargetMesh;
 }
 
+UDynamicMesh* UGeometryScriptLibrary_MeshSelectionFunctions::SelectMeshElementsByMaterialID(UDynamicMesh* TargetMesh,
+	int MaterialID,
+	FGeometryScriptMeshSelection& Selection,
+	EGeometryScriptMeshSelectionType SelectionType)
+{
+	if (TargetMesh == nullptr)
+	{
+		UE_LOG(LogGeometry, Warning, TEXT("CreateMeshSelectionByMaterialID: TargetMesh is Null"));
+		return TargetMesh;
+	}
+
+	TargetMesh->ProcessMesh([MaterialID, &Selection, SelectionType](const FDynamicMesh3& ReadMesh)
+	{
+		if (!ReadMesh.HasAttributes() || ReadMesh.Attributes()->GetMaterialID() == nullptr)
+		{
+			UE_LOG(LogGeometry, Warning, TEXT("CreateMeshSelectionByMaterialID: Mesh does not have material IDs"));
+			return;
+		}
+
+		const FDynamicMeshMaterialAttribute* MaterialIDAttr = ReadMesh.Attributes()->GetMaterialID();
+
+		UE::MeshSelectionLocals::SelectByTriangleAttribute(ReadMesh, Selection, SelectionType, [&MaterialIDAttr, MaterialID](int32 TID)
+		{
+			return MaterialIDAttr->GetValue(TID) == MaterialID;
+		});
+	});
+	return TargetMesh;
+}
+
+UDynamicMesh* UGeometryScriptLibrary_MeshSelectionFunctions::SelectMeshElementsByPolygroup(
+	UDynamicMesh* TargetMesh,
+	FGeometryScriptGroupLayer GroupLayer,
+	UPARAM(DisplayName = "PolyGroup ID") int PolygroupID,
+	FGeometryScriptMeshSelection& Selection,
+	EGeometryScriptMeshSelectionType SelectionType)
+{
+	if (TargetMesh == nullptr)
+	{
+		UE_LOG(LogGeometry, Warning, TEXT("CreateMeshSelectionByPolygroup: TargetMesh is Null"));
+		return TargetMesh;
+	}
+
+	TargetMesh->ProcessMesh([GroupLayer, PolygroupID, &Selection, SelectionType](const FDynamicMesh3& ReadMesh)
+	{
+		if (GroupLayer.bDefaultLayer == true)
+		{
+			UE::MeshSelectionLocals::SelectByTriangleAttribute(ReadMesh, Selection, SelectionType, [&ReadMesh, PolygroupID](int32 TID)
+			{
+				return ReadMesh.GetTriangleGroup(TID) == PolygroupID;
+			});
+		}
+		else
+		{
+			if (!ReadMesh.HasAttributes() || ReadMesh.Attributes()->NumPolygroupLayers() <= GroupLayer.ExtendedLayerIndex)
+			{
+				UE_LOG(LogGeometry, Warning, TEXT("CreateMeshSelectionByPolygroup: Requested Polygroup Layer (%d) not found"), GroupLayer.ExtendedLayerIndex);
+				return;
+			}
+
+			const FDynamicMeshPolygroupAttribute* PolygroupAttr = ReadMesh.Attributes()->GetPolygroupLayer(GroupLayer.ExtendedLayerIndex);
+			UE::MeshSelectionLocals::SelectByTriangleAttribute(ReadMesh, Selection, SelectionType, [&PolygroupAttr, PolygroupID](int32 TID)
+			{
+				return PolygroupAttr->GetValue(TID) == PolygroupID;
+			});
+		}
+	});
+	return TargetMesh;
+}
+
 void UGeometryScriptLibrary_MeshSelectionFunctions::CombineMeshSelections(
 	FGeometryScriptMeshSelection SelectionA,
 	FGeometryScriptMeshSelection SelectionB,
@@ -477,77 +670,6 @@
 
 
 
-namespace UELocal
-{
-
-static void SelectMeshElementsWithContainmentTest(
-	UDynamicMesh* TargetMesh,
-	TFunctionRef<bool(const FVector3d& Position, const FVector3d& Normal)> ContainmentFunc,
-	FGeometryScriptMeshSelection& SelectionOut,
-	EGeometryScriptMeshSelectionType SelectionType,
-	int NumTrianglePoints,
-	bool bNeedsNormals)
-{
-	NumTrianglePoints = FMath::Clamp(NumTrianglePoints, 1, 3);
-	FVector3d UnitNormal = FVector3d::UnitZ();
-
-	FGeometrySelection GeoSelection;
-	if (SelectionType == EGeometryScriptMeshSelectionType::Vertices)
-	{
-		GeoSelection.InitializeTypes(EGeometryElementType::Vertex, EGeometryTopologyType::Triangle);
-
-		TargetMesh->ProcessMesh([&](const FDynamicMesh3& Mesh) 
-		{
-			for (int32 vid : Mesh.VertexIndicesItr())
-			{
-				FVector3d UseNormal = (bNeedsNormals) ? FMeshNormals::ComputeVertexNormal(Mesh,vid) : UnitNormal;
-				if ( ContainmentFunc(Mesh.GetVertex(vid), UseNormal) )
-				{
-					GeoSelection.Selection.Add( FGeoSelectionID::MeshVertex(vid).Encoded() );
-				}
-			}
-		});
-	}
-	else
-	{
-		GeoSelection.InitializeTypes(EGeometryElementType::Face, 
-			(SelectionType == EGeometryScriptMeshSelectionType::Triangles) ? EGeometryTopologyType::Triangle : EGeometryTopologyType::Polygroup);
-
-		TargetMesh->ProcessMesh([&](const FDynamicMesh3& Mesh) 
-		{
-			for (int32 tid : Mesh.TriangleIndicesItr())
-			{
-				FVector3d UseNormal = (bNeedsNormals) ? Mesh.GetTriNormal(tid) : UnitNormal;
-				FIndex3i Tri = Mesh.GetTriangle(tid);
-				// may be wasteful to test each vertex multiple times...could accumulate a cache at cost of some memory allocation...
-				int NumContained =
-					(ContainmentFunc(Mesh.GetVertex(Tri.A), UseNormal) ? 1 : 0) +
-					(ContainmentFunc(Mesh.GetVertex(Tri.B), UseNormal) ? 1 : 0) +
-					(ContainmentFunc(Mesh.GetVertex(Tri.C), UseNormal) ? 1 : 0);
-
-				if ( NumContained >= NumTrianglePoints )
-				{
-					if (SelectionType == EGeometryScriptMeshSelectionType::Triangles)
-					{
-						GeoSelection.Selection.Add( FGeoSelectionID::MeshTriangle(tid).Encoded() );
-					}
-					else
-					{
-						int32 gid = Mesh.GetTriangleGroup(tid);
-						GeoSelection.Selection.Add( FGeoSelectionID::GroupFace(tid, gid).Encoded() );
-					}
-				}
-			}
-		});
-
-	}
-
-	SelectionOut.SetSelection(MoveTemp(GeoSelection));
-}
-
-} // end namespace UELocal
-
-
 UDynamicMesh* UGeometryScriptLibrary_MeshSelectionFunctions::SelectMeshElementsInBox(
 	UDynamicMesh* TargetMesh,
 	FGeometryScriptMeshSelection& SelectionOut,
@@ -567,7 +689,7 @@
 	{ 
 		return Container.Contains(Point) != bInvert; 
 	};
-	UELocal::SelectMeshElementsWithContainmentTest(TargetMesh, ContainsFunc,
+	UE::MeshSelectionLocals::SelectMeshElementsWithContainmentTest(TargetMesh, ContainsFunc,
 		SelectionOut, SelectionType, MinNumTrianglePoints, false);
 	return TargetMesh;
 }
@@ -595,7 +717,7 @@
 	{ 
 		return Container.Contains(Point) != bInvert; 
 	};
-	UELocal::SelectMeshElementsWithContainmentTest(TargetMesh, ContainsFunc,
+	UE::MeshSelectionLocals::SelectMeshElementsWithContainmentTest(TargetMesh, ContainsFunc,
 		SelectionOut, SelectionType, MinNumTrianglePoints, false);
 	return TargetMesh;
 }
@@ -623,7 +745,7 @@
 		bool bContains = (Point - PlaneOrigin).Dot(PlaneNormal) >= 0;
 		return bContains != bInvert;
 	};
-	UELocal::SelectMeshElementsWithContainmentTest(TargetMesh, ContainsFunc,
+	UE::MeshSelectionLocals::SelectMeshElementsWithContainmentTest(TargetMesh, ContainsFunc,
 		SelectionOut, SelectionType, MinNumTrianglePoints, false);
 	return TargetMesh;
 }
@@ -652,7 +774,7 @@
 		bool bContains = PlaneNormal.Dot(Normal) >= CosMaxAngle;
 		return bContains != bInvert;
 	};
-	UELocal::SelectMeshElementsWithContainmentTest(TargetMesh, ContainsFunc,
+	UE::MeshSelectionLocals::SelectMeshElementsWithContainmentTest(TargetMesh, ContainsFunc,
 		SelectionOut, SelectionType, MinNumTrianglePoints, true);
 	return TargetMesh;
 }
@@ -710,7 +832,7 @@
 			}
 			return bContains != bInvert;
 		};
-		UELocal::SelectMeshElementsWithContainmentTest(TargetMesh, ContainsFunc,
+		UE::MeshSelectionLocals::SelectMeshElementsWithContainmentTest(TargetMesh, ContainsFunc,
 			SelectionOut, SelectionType, MinNumTrianglePoints, true);
 	});
 
diff --git a/Engine/Plugins/Runtime/GeometryScripting/Source/GeometryScriptingCore/Public/GeometryScript/MeshSelectionFunctions.h b/Engine/Plugins/Runtime/GeometryScripting/Source/GeometryScriptingCore/Public/GeometryScript/MeshSelectionFunctions.h
--- a/Engine/Plugins/Runtime/GeometryScripting/Source/GeometryScriptingCore/Public/GeometryScript/MeshSelectionFunctions.h
+++ b/Engine/Plugins/Runtime/GeometryScripting/Source/GeometryScriptingCore/Public/GeometryScript/MeshSelectionFunctions.h
@@ -59,6 +59,28 @@ public:
 		FGeometryScriptMeshSelection& Selection,
 		EGeometryScriptMeshSelectionType SelectionType = EGeometryScriptMeshSelectionType::Triangles );
 
+	/**
+	 * Create a Selection of the SelectionType that contains all mesh elements referencing triangles with the given Material ID
+	 */
+	UFUNCTION(BlueprintCallable, Category = "GeometryScript|MeshSelection", meta=(ScriptMethod))
+	static UPARAM(DisplayName = "Target Mesh") UDynamicMesh*
+	SelectMeshElementsByMaterialID(
+		UDynamicMesh* TargetMesh,
+		int MaterialID,
+		FGeometryScriptMeshSelection& Selection,
+		EGeometryScriptMeshSelectionType SelectionType = EGeometryScriptMeshSelectionType::Triangles );
+
+	/**
+	 * Create a Selection of the SelectionType that contains all mesh elements referencing triangles with the given PolyGroup ID in the given GroupLayer
+	 */
+	UFUNCTION(BlueprintCallable, Category = "GeometryScript|MeshSelection", meta=(ScriptMethod))
+	static UPARAM(DisplayName = "Target Mesh") UDynamicMesh*
+	SelectMeshElementsByPolygroup(
+		UDynamicMesh* TargetMesh,
+		FGeometryScriptGroupLayer GroupLayer,
+		UPARAM(DisplayName = "PolyGroup ID") int PolygroupID,
+		FGeometryScriptMeshSelection& Selection,
+		EGeometryScriptMeshSelectionType SelectionType = EGeometryScriptMeshSelectionType::Triangles );
 
 	/**
 	 * Convert a Mesh Selection to a different Type (eg Vertices to Triangles, etc)
