fix cis issue 500446 -- changing volume texture bake functions scriptname + cateogry from bake to volumebake to remove conflict with existing (2d texture) bake functions

diff --git a/Engine/Plugins/Runtime/GeometryScripting/Source/GeometryScriptingCore/Public/GeometryScript/VolumeTextureBakeFunctions.h b/Engine/Plugins/Runtime/GeometryScripting/Source/GeometryScriptingCore/Public/GeometryScript/VolumeTextureBakeFunctions.h
--- a/Engine/Plugins/Runtime/GeometryScripting/Source/GeometryScriptingCore/Public/GeometryScript/VolumeTextureBakeFunctions.h
+++ b/Engine/Plugins/Runtime/GeometryScripting/Source/GeometryScriptingCore/Public/GeometryScript/VolumeTextureBakeFunctions.h
@@ -74,14 +74,14 @@ public:
 };
 
 
-UCLASS(meta = (ScriptName = "GeometryScript_Bake"))
+UCLASS(meta = (ScriptName = "GeometryScript_VolumeBake"))
 class GEOMETRYSCRIPTINGCORE_API UGeometryScriptLibrary_VolumeTextureBakeFunctions : public UBlueprintFunctionLibrary
 {
 	GENERATED_BODY()
 public:
 
 	// Write a distance field to the given existing volume texture
-	UFUNCTION(BlueprintCallable, Category = "GeometryScript|Bake")
+	UFUNCTION(BlueprintCallable, Category = "GeometryScript|VolumeBake")
 	static UPARAM(DisplayName = "Success") bool BakeSignedDistanceToVolumeTexture(
 		const UDynamicMesh* TargetMesh,
 		UVolumeTexture* VolumeTexture,
