add optional blocking planes to the geometry script navigable approximate convex decomposition settings

diff --git a/Engine/Plugins/Runtime/GeometryScripting/Source/GeometryScriptingCore/Private/CollisionFunctions.cpp b/Engine/Plugins/Runtime/GeometryScripting/Source/GeometryScriptingCore/Private/CollisionFunctions.cpp
--- a/Engine/Plugins/Runtime/GeometryScripting/Source/GeometryScriptingCore/Private/CollisionFunctions.cpp
+++ b/Engine/Plugins/Runtime/GeometryScripting/Source/GeometryScriptingCore/Private/CollisionFunctions.cpp
@@ -1406,6 +1406,21 @@ FGeometryScriptSimpleCollision UGeometryScriptLibrary_CollisionFunctions::Comput
 		NegativeSpaceSettings.MinRadius = NegativeSpaceOptions.MinRadius;
 		NegativeSpaceSettings.MinSpacing = 0;
 
+		if (!NegativeSpaceOptions.UnreachablePlanes.IsEmpty())
+		{
+			NegativeSpaceSettings.OptionalObstacleSDF = [&NegativeSpaceOptions](FVector Pos)
+			{
+				double SDF = FMathd::MaxReal;
+				for (const FPlane& Plane : NegativeSpaceOptions.UnreachablePlanes)
+				{
+					double PlaneSDF = Plane.PlaneDot(Pos);
+					SDF = FMath::Min(PlaneSDF, SDF);
+				}
+				return SDF;
+			};
+		}
+		
+
 		ConvexDecomposition.InitializeNegativeSpace(NegativeSpaceSettings, NegativeSpaceOptions.CustomNavigablePositions);
 
 		ConvexDecomposition.MaxConvexEdgePlanes = 4;
diff --git a/Engine/Plugins/Runtime/GeometryScripting/Source/GeometryScriptingCore/Public/GeometryScript/CollisionFunctions.h b/Engine/Plugins/Runtime/GeometryScripting/Source/GeometryScriptingCore/Public/GeometryScript/CollisionFunctions.h
--- a/Engine/Plugins/Runtime/GeometryScripting/Source/GeometryScriptingCore/Public/GeometryScript/CollisionFunctions.h
+++ b/Engine/Plugins/Runtime/GeometryScripting/Source/GeometryScriptingCore/Public/GeometryScript/CollisionFunctions.h
@@ -194,6 +194,11 @@ public:
 	/// Optional list of locations that we expect to be navigable
 	UPROPERTY(EditAnywhere, BlueprintReadWrite, Category = Options)
 	TArray<FVector> CustomNavigablePositions;
+
+	/// Optional list of planes defining unreachable space (on their negative side)
+	/// Use this for example to specify a ground plane, if a mesh will always be placed on ground and need not be navigable from below.
+	UPROPERTY(EditAnywhere, BlueprintReadWrite, Category = Options)
+	TArray<FPlane> UnreachablePlanes;
 };
 
 
