Geometry Script: When creating a skeletal mesh asset from dynamic mesh, ensure that the bones stored on the mesh match the reference skeleton being used.

#rnx

diff --git a/Engine/Plugins/Runtime/GeometryScripting/Source/GeometryScriptingEditor/Private/CreateNewAssetUtilityFunctions.cpp b/Engine/Plugins/Runtime/GeometryScripting/Source/GeometryScriptingEditor/Private/CreateNewAssetUtilityFunctions.cpp
--- a/Engine/Plugins/Runtime/GeometryScripting/Source/GeometryScriptingEditor/Private/CreateNewAssetUtilityFunctions.cpp
+++ b/Engine/Plugins/Runtime/GeometryScripting/Source/GeometryScriptingEditor/Private/CreateNewAssetUtilityFunctions.cpp
@@ -437,6 +437,34 @@ USkeletalMesh* UGeometryScriptLibrary_CreateNewAssetFunctions::CreateNewSkeletal
 							return nullptr;
 						}
 					}
+
+					// Update the bones on the mesh to match the reference skeleton being used. We try to retain any existing bone color assignments
+					// as much as we can.
+					TMap<FName, FVector4f> BoneColors;
+					if (AttribSet->GetBoneNames() && AttribSet->GetBoneColors())
+					{
+						for (int32 BoneIndex = 0; BoneIndex < AttribSet->GetNumBones(); BoneIndex++)
+						{
+							BoneColors.Add(AttribSet->GetBoneNames()->GetValue(BoneIndex), AttribSet->GetBoneColors()->GetValue(BoneIndex));
+						}
+					}
+					const FReferenceSkeleton* ToRefSkeleton = AssetOptions.RefSkeleton ? AssetOptions.RefSkeleton : &InSkeleton->GetReferenceSkeleton();
+
+					AttribSet->EnableBones(ToRefSkeleton->GetRawBoneNum());
+					const TArray<FMeshBoneInfo>& BoneInfos = ToRefSkeleton->GetRawRefBoneInfo();
+					const TArray<FTransform>& BonePoses = ToRefSkeleton->GetRawRefBonePose();
+
+					for (int32 BoneIndex = 0; BoneIndex < ToRefSkeleton->GetRawBoneNum(); BoneIndex++)
+					{
+						AttribSet->GetBoneNames()->SetValue(BoneIndex, BoneInfos[BoneIndex].Name);
+						AttribSet->GetBoneParentIndices()->SetValue(BoneIndex, BoneInfos[BoneIndex].ParentIndex);
+						AttribSet->GetBonePoses()->SetValue(BoneIndex, BonePoses[BoneIndex]);
+
+						if(const FVector4f* BoneColor = BoneColors.Find(BoneInfos[BoneIndex].Name))
+						{
+							AttribSet->GetBoneColors()->SetValue(BoneIndex, *BoneColor);
+						}
+					}
 				}
 			}
 
