Geometry Scripting: Multiple improvements.
- Add an option to CopyBonesFromMesh to only retain some parts of the skeleton.
- Add a function to blend between two skin weight arrays.
- Add a function to copy weights from one profile to another on the same mesh.

diff --git a/Engine/Plugins/Runtime/GeometryScripting/Source/GeometryScriptingCore/Private/MeshBoneWeightFunctions.cpp b/Engine/Plugins/Runtime/GeometryScripting/Source/GeometryScriptingCore/Private/MeshBoneWeightFunctions.cpp
--- a/Engine/Plugins/Runtime/GeometryScripting/Source/GeometryScriptingCore/Private/MeshBoneWeightFunctions.cpp
+++ b/Engine/Plugins/Runtime/GeometryScripting/Source/GeometryScriptingCore/Private/MeshBoneWeightFunctions.cpp
@@ -74,7 +74,97 @@ ReturnType SimpleMeshBoneWeightEdit(
 	return RetVal;
 }
 
+/** Validate that the bone hierarchy stored on the mesh is a proper tree with a single root node and no cycles. */
+static bool ValidateBoneHierarchy(
+	const FDynamicMesh3& InMesh,
+	UGeometryScriptDebug* InDebug
+	)
+{
+	if (!InMesh.HasAttributes() || InMesh.Attributes()->GetNumBones() == 0)
+	{
+		AppendError(InDebug, EGeometryScriptErrorType::InvalidInputs, LOCTEXT("ValidateBoneHierarchy_NoBones", "No bone attributes defined on the mesh."));
+		return false;
+	}
+
+	const int32 NumBones = InMesh.Attributes()->GetNumBones();
+	if (NumBones > 1 && !InMesh.Attributes()->GetBoneParentIndices())
+	{
+		AppendError(InDebug, EGeometryScriptErrorType::InvalidInputs, LOCTEXT("ValidateBoneHierarchy_MultpleBonesNoParents", "Multiple bones defined but no parent indices present."));
+		return false;
+	}
 
+	const FDynamicMeshBoneNameAttribute* NameAttrib = InMesh.Attributes()->GetBoneNames(); 
+	const FDynamicMeshBoneParentIndexAttribute* ParentIndexAttrib = InMesh.Attributes()->GetBoneParentIndices(); 
+
+	TSet<FName> BoneNamesSeen;
+	
+	int32 RootBoneIndex = INDEX_NONE;
+	for (int32 BoneIndex = 0; BoneIndex < NumBones; BoneIndex++)
+	{
+		const FName BoneName = NameAttrib->GetValue(BoneIndex);
+		
+		if (BoneName == NAME_None)
+		{
+			AppendError(InDebug, EGeometryScriptErrorType::InvalidInputs, FText::Format(LOCTEXT("ValidateBoneHierarchy_UnnamedBone", "Bone at index {0} has no name."), FText::AsNumber(BoneIndex)));
+			return false;
+		}
+		if (BoneNamesSeen.Contains(BoneName))
+		{
+			AppendError(InDebug, EGeometryScriptErrorType::InvalidInputs, FText::Format(LOCTEXT("ValidateBoneHierarchy_DuplicateBoneNames", "Bone '{0}' defined more than once."), FText::FromName(BoneName)));
+			return false;
+		}
+		BoneNamesSeen.Add(BoneName);
+
+		const int32 ParentBoneIndex = ParentIndexAttrib->GetValue(BoneIndex); 
+		if (ParentBoneIndex == INDEX_NONE)
+		{
+			if (RootBoneIndex != INDEX_NONE)
+			{
+				AppendError(InDebug, EGeometryScriptErrorType::InvalidInputs, 
+					FText::Format(LOCTEXT("ValidateBoneHierarchy_MultpleRootBones", "Multiple root bones found ('{0}' and '{1}')."),
+						FText::FromName(BoneName), FText::FromName(NameAttrib->GetValue(RootBoneIndex))));
+				return false;
+			}
+
+			RootBoneIndex = BoneIndex;
+		}
+		else if (ParentBoneIndex < 0 || ParentBoneIndex >= NumBones)
+        {
+            AppendError(InDebug, EGeometryScriptErrorType::InvalidInputs, 
+             	FText::Format(LOCTEXT("ValidateBoneHierarchy_InvalidParentBoneIndex", "Parent bone index {0} for bone '{1}' is invalid ({2} bones defined)."),
+             		FText::AsNumber(ParentBoneIndex), FText::FromName(NameAttrib->GetValue(BoneIndex)), FText::AsNumber(NumBones)));
+            return false;
+        }
+	}
+
+	// Once we've verified that all bones are properly named, all the parent indices are valid and there's only one root, let's check for cycles.
+	TSet<int32> BoneIndicesVisited;
+	
+	for (int32 BoneIndex = 0; BoneIndex < NumBones; BoneIndex++)
+	{
+		BoneIndicesVisited.Reset();
+
+		int32 CurrentBoneIndex = BoneIndex;
+		while (CurrentBoneIndex != RootBoneIndex)
+		{
+			// We store a set of all visited indices rather than just check if a bone cycles back onto itself, since the cycle could be up the hierarchy
+			// and not cycle back on the starting bone.
+			BoneIndicesVisited.Add(CurrentBoneIndex);
+			
+			const int32 ParentBoneIndex = ParentIndexAttrib->GetValue(CurrentBoneIndex);
+			if (BoneIndicesVisited.Contains(ParentBoneIndex))
+			{
+				AppendError(InDebug, EGeometryScriptErrorType::InvalidInputs, 
+					FText::Format(LOCTEXT("ValidateBoneHierarchy_FoundCycle", "Bone '{0}' does not connect up to the root bone '{1}' but connects into a cycle instead."),
+						FText::FromName(NameAttrib->GetValue(CurrentBoneIndex)), FText::FromName(NameAttrib->GetValue(RootBoneIndex))));
+				return false;
+			}
+			CurrentBoneIndex = ParentBoneIndex;
+		}
+	}
+
+	return true;
+}
 
 
 UDynamicMesh* UGeometryScriptLibrary_MeshBoneWeightFunctions::MeshHasBoneWeights(
@@ -124,6 +214,36 @@ UDynamicMesh* UGeometryScriptLibrary_MeshBoneWeightFunctions::MeshCreateBoneWeig
 }
 
 
+UDynamicMesh* UGeometryScriptLibrary_MeshBoneWeightFunctions::MeshCopyBoneWeights(
+	UDynamicMesh* TargetMesh,
+	bool& bProfileExisted,
+	FGeometryScriptBoneWeightProfile TargetProfile,
+	FGeometryScriptBoneWeightProfile SourceProfile
+	)
+{
+	bProfileExisted = false;
+	if (TargetMesh)
+	{
+		TargetMesh->EditMesh([&](FDynamicMesh3& EditMesh)
+		{
+			if (!EditMesh.HasAttributes())
+			{
+				return;
+			}
+			
+			const FDynamicMeshVertexSkinWeightsAttribute* SourceAttribute = EditMesh.Attributes()->GetSkinWeightsAttribute(SourceProfile.GetProfileName());
+			FDynamicMeshVertexSkinWeightsAttribute* TargetAttribute = EditMesh.Attributes()->GetSkinWeightsAttribute(TargetProfile.GetProfileName());
+			if (!SourceAttribute || !TargetAttribute)
+			{
+				return;
+			}
+
+			TargetAttribute->Copy(*SourceAttribute);
+		}, EDynamicMeshChangeType::GeneralEdit, EDynamicMeshAttributeChangeFlags::Unknown, false);
+	}
+	return TargetMesh;
+}
+
 
 UDynamicMesh* UGeometryScriptLibrary_MeshBoneWeightFunctions::GetMaxBoneWeightIndex(
 	UDynamicMesh* TargetMesh,
@@ -274,6 +394,45 @@ UDynamicMesh* UGeometryScriptLibrary_MeshBoneWeightFunctions::SetVertexBoneWeigh
 	return TargetMesh;
 }
 
+void UGeometryScriptLibrary_MeshBoneWeightFunctions::BlendBoneWeights(
+	const TArray<FGeometryScriptBoneWeight>& BoneWeightsA,
+	const TArray<FGeometryScriptBoneWeight>& BoneWeightsB,
+	float Alpha, 
+	TArray<FGeometryScriptBoneWeight>& Result,
+	UGeometryScriptDebug* Debug
+	)
+{
+	TArray<FBoneWeight, TInlineAllocator<UE::AnimationCore::MaxInlineBoneWeightCount>> RawWeightsA;
+
+	for (const FGeometryScriptBoneWeight& BoneWeight: BoneWeightsA)
+	{
+		
+		if (BoneWeight.BoneIndex < 0)
+		{
+			AppendWarning(Debug, EGeometryScriptErrorType::InvalidInputs, LOCTEXT("BlendBoneWeights_InvalidInput", "BlendBoneWeights: Invalid bone index provided; falling back to 0 as bone index."));
+		}
+		RawWeightsA.Add(FBoneWeight(FMath::Max(0, BoneWeight.BoneIndex), BoneWeight.Weight));
+	}
+	
+	TArray<FBoneWeight, TInlineAllocator<UE::AnimationCore::MaxInlineBoneWeightCount>> RawWeightsB;
+	for (const FGeometryScriptBoneWeight& BoneWeight: BoneWeightsB)
+	{
+		if (BoneWeight.BoneIndex < 0)
+		{
+			AppendWarning(Debug, EGeometryScriptErrorType::InvalidInputs, LOCTEXT("BlendBoneWeights_InvalidInput", "BlendBoneWeights: Invalid bone index provided; falling back to 0 as bone index."));
+		}
+		RawWeightsB.Add(FBoneWeight(FMath::Max(0, BoneWeight.BoneIndex), BoneWeight.Weight));
+	}
+	
+	FBoneWeights NewWeights = NewWeights.Blend(FBoneWeights::Create(RawWeightsA), FBoneWeights::Create(RawWeightsB), Alpha);
+	
+	Result.SetNum(NewWeights.Num());
+	for (FBoneWeight BoneWeight: NewWeights)
+	{
+		Result.Emplace(BoneWeight.GetBoneIndex(), BoneWeight.GetWeight());
+	}
+}	
+
 
 UDynamicMesh* UGeometryScriptLibrary_MeshBoneWeightFunctions::SetAllVertexBoneWeights(
 	UDynamicMesh* TargetMesh,
@@ -496,19 +655,139 @@ UDynamicMesh* UGeometryScriptLibrary_MeshBoneWeightFunctions::CopyBonesFromMesh(
 			{
 				EditMesh.EnableAttributes();
 			}
+
+			TMap<FName, FName> BoneHierarchy;
+
+			// Check across all skin weight attributes which bones are bound to the mesh.
+			if (Options.BonesToCopyFromSource != EBonesToCopyFromSource::AllBones)
+			{
+				// Ensure that all required attributes on the source are defined, that all names and parent indices are valid, and that the bone
+				// hierarchy is consistent.
+				if (!ValidateBoneHierarchy(ReadMesh, Debug))
+				{
+					return;
+				}
+				
+				// Target bone name attribute's existence has already been verified above with a call to HasBones, the source has been verified through
+				// ValidateBoneHierarchy.
+				const FDynamicMeshBoneNameAttribute* BoneNameSourceAttrib = ReadMesh.Attributes()->GetBoneNames();
+				const FDynamicMeshBoneParentIndexAttribute* BoneParentIndexSourceAttrib = ReadMesh.Attributes()->GetBoneParentIndices();
+				const FDynamicMeshBoneNameAttribute* BoneNameTargetAttrib = EditMesh.Attributes()->GetBoneNames();
 			
+				FName RootBone;
+				
+				for (int32 Index = 0; Index < BoneNameSourceAttrib->Num(); Index++)
+				{
+					const int32 BoneParentIndex = BoneParentIndexSourceAttrib->GetValue(Index);
+					if (BoneParentIndex == INDEX_NONE || (BoneParentIndex >= 0 && BoneParentIndex < BoneNameSourceAttrib->Num()))
+					{
+						BoneHierarchy.Add(BoneNameSourceAttrib->GetValue(Index), BoneParentIndex != INDEX_NONE ? BoneNameSourceAttrib->GetValue(BoneParentIndex) : NAME_None);
+
+						if (BoneParentIndex == INDEX_NONE)
+						{
+							if (!RootBone.IsNone())
+							{
+								AppendError(Debug, EGeometryScriptErrorType::OperationFailed, LOCTEXT("CopyBonesFromMesh_MultipleRootBonesFound", "Found multiple root bones on source mesh"));
+								return;
+							}
+							RootBone = BoneNameSourceAttrib->GetValue(Index);
+						}
+					}
+				}
+
+				TSet<int32> BoundBoneIndices;
+
+				for (const TPair<FName, TUniquePtr<FDynamicMeshVertexSkinWeightsAttribute>>& AttribPair : EditMesh.Attributes()->GetSkinWeightsAttributes())
+				{
+					if (const FDynamicMeshVertexSkinWeightsAttribute* ToAttrib = EditMesh.Attributes()->GetSkinWeightsAttribute(AttribPair.Key))
+					{
+						BoundBoneIndices.Append(ToAttrib->GetBoundBoneIndices());
+					}
+				}
+
+				TSet<FName> BoundBones;
+				for (int32 BoneIndex: BoundBoneIndices)
+				{
+					if (BoneIndex < 0 || BoneIndex >= BoneNameTargetAttrib->Num())
+					{
+						AppendError(Debug, EGeometryScriptErrorType::OperationFailed, LOCTEXT("CopyBonesFromMesh_InvalidBoneWeightIndex", "Invalid bone index found on mesh"));
+						return;
+					}
+
+					const FName BoneName = BoneNameTargetAttrib->GetValue(BoneIndex);
+					if (BoneName.IsNone())
+					{
+						const FText Error = FText::Format(LOCTEXT("CopyBonesFromMesh_NoBoneName", "Target bone at index {0} has no name."), FText::AsNumber(BoneIndex));
+						AppendError(Debug, EGeometryScriptErrorType::OperationFailed, Error);
+						return;
+					}
+					if (!BoneHierarchy.Contains(BoneName))
+					{
+						const FText Error = FText::Format(LOCTEXT("CopyBonesFromMesh_BoneNotFound", "Target bone '{0}' not found on source mesh."), FText::FromName(BoneName));
+						AppendError(Debug, EGeometryScriptErrorType::OperationFailed, Error);
+						return;
+					}
+					BoundBones.Add(BoneName);
+				}
+
+				// Go from each bound bone, and add parent bones to root, optionally skipping over unbound bones, if using OnlyBoundAndRoot.
+				TSet<FName> UsedBones;
+				UsedBones.Add(RootBone);
+				
+				for (FName BoneName: BoundBones)
+				{
+					UsedBones.Add(BoneName);
+
+					// Traverse up to the root bone. If OnlyBoundAndRoot is set, then each bone's parent is set to either a bone further up
+					// in the hierarchy that's actually bound, or the root.
+					FName ParentBone = BoneHierarchy[BoneName];
+					while (ParentBone != RootBone)
+					{
+						if (Options.BonesToCopyFromSource == EBonesToCopyFromSource::OnlyBoundAndParents)
+						{
+							UsedBones.Add(ParentBone);
+						}
+						else if (Options.BonesToCopyFromSource == EBonesToCopyFromSource::OnlyBoundAndRoot && BoundBones.Contains(ParentBone))
+						{
+							// We found another bound bone, make this our new parent.
+							break;
+						}
+						
+						ParentBone = BoneHierarchy[ParentBone];
+					}
+
+					if (Options.BonesToCopyFromSource == EBonesToCopyFromSource::OnlyBoundAndRoot)
+					{
+						BoneHierarchy[BoneName] = ParentBone;
+					}
+				}
+
+				// Leave only used bones in the hierarchy.
+				BoneHierarchy = BoneHierarchy.FilterByPredicate([&UsedBones](const TPair<FName, FName>& Item)
+				{
+					return UsedBones.Contains(Item.Key);
+				});
+			}
+
 			if (Options.ReindexWeights)
 			{
 				if (EditMesh.Attributes()->HasBones())
 				{
+					TArray<FName> TargetBones(ReadMesh.Attributes()->GetBoneNames()->GetAttribValues());
+					if (!BoneHierarchy.IsEmpty())
+					{
+						// Remove all bones that are not in the hierarchy.
+						TargetBones.RemoveAll([&BoneHierarchy](FName BoneName)
+						{
+							return !BoneHierarchy.Contains(BoneName);
+						});
+					}
+					
 					for (const TPair<FName, TUniquePtr<FDynamicMeshVertexSkinWeightsAttribute>>& AttribPair : EditMesh.Attributes()->GetSkinWeightsAttributes())
 					{
-						FDynamicMeshVertexSkinWeightsAttribute* ToAttrib = EditMesh.Attributes()->GetSkinWeightsAttribute(AttribPair.Key);
-						if (ToAttrib)
-						{	
-							const bool bResult = ToAttrib->ReindexBoneIndicesToSkeleton(EditMesh.Attributes()->GetBoneNames()->GetAttribValues(),
-																	  				    ReadMesh.Attributes()->GetBoneNames()->GetAttribValues());
-							if (!bResult)
+						if (FDynamicMeshVertexSkinWeightsAttribute* ToAttrib = EditMesh.Attributes()->GetSkinWeightsAttribute(AttribPair.Key))
+						{
+							if (!ToAttrib->ReindexBoneIndicesToSkeleton(EditMesh.Attributes()->GetBoneNames()->GetAttribValues(), TargetBones))
 							{
 								const FText Error = FText::Format(LOCTEXT("CopyBonesFromMesh_FailedToReindexWeights", "Failed to reindex bone weights for {0} weights profile"), FText::FromName(AttribPair.Key));
 								AppendError(Debug, EGeometryScriptErrorType::OperationFailed, Error);
@@ -522,7 +801,16 @@ UDynamicMesh* UGeometryScriptLibrary_MeshBoneWeightFunctions::CopyBonesFromMesh(
 				}
 			}
 
-			EditMesh.Attributes()->CopyBoneAttributes(*ReadMesh.Attributes());
+			// If the bone hierarchy wasn't set up, then do a full copy.
+			if (BoneHierarchy.IsEmpty())
+			{
+				EditMesh.Attributes()->CopyBoneAttributes(*ReadMesh.Attributes());
+			}
+			else
+			{
+				// Copy the bone attributes but only copy the ones in the remapping map and update the parent index as well.
+				EditMesh.Attributes()->CopyBoneAttributesWithRemapping(*ReadMesh.Attributes(), BoneHierarchy);
+			}
 
 		}, EDynamicMeshChangeType::AttributeEdit, EDynamicMeshAttributeChangeFlags::Unknown, false);
 	});
diff --git a/Engine/Plugins/Runtime/GeometryScripting/Source/GeometryScriptingCore/Public/GeometryScript/MeshBoneWeightFunctions.h b/Engine/Plugins/Runtime/GeometryScripting/Source/GeometryScriptingCore/Public/GeometryScript/MeshBoneWeightFunctions.h
--- a/Engine/Plugins/Runtime/GeometryScripting/Source/GeometryScriptingCore/Public/GeometryScript/MeshBoneWeightFunctions.h
+++ b/Engine/Plugins/Runtime/GeometryScripting/Source/GeometryScriptingCore/Public/GeometryScript/MeshBoneWeightFunctions.h
@@ -189,6 +189,23 @@ struct GEOMETRYSCRIPTINGCORE_API FGeometryScriptBoneInfo
 	FLinearColor Color = FLinearColor::White;
 };
 
+
+UENUM(BlueprintType)
+enum class EBonesToCopyFromSource : uint8
+{
+	/** Copy all bones from the source mesh to the target, regardless of whether they're bound or not. 
+	 */
+	AllBones = 0,
+
+	/** Keep only bones that are actually bound to the target mesh, including all parent bones up to the root. */
+	OnlyBoundAndParents = 1,
+	
+	/** Keep only bones that are actually bound to the target mesh and the root bone. Any existing bones between
+	 *  the two will not be copied. Bound bones will have their parent as either the root bone or another bound bone. */
+	OnlyBoundAndRoot = 2
+};
+
+
 USTRUCT(BlueprintType)
 struct GEOMETRYSCRIPTINGCORE_API FGeometryScriptCopyBonesFromMeshOptions
 {
@@ -198,6 +215,10 @@ struct GEOMETRYSCRIPTINGCORE_API FGeometryScriptCopyBonesFromMeshOptions
 	 * target skeleton to the source skeleton.*/
 	UPROPERTY(EditAnywhere, BlueprintReadWrite, Category = Options)
 	bool ReindexWeights = false;
+
+	/** Specify which bones are copied from the source mesh to the target. */
+	UPROPERTY(EditAnyWhere, BlueprintReadWrite, Category = Options)
+	EBonesToCopyFromSource BonesToCopyFromSource = EBonesToCopyFromSource::AllBones;
 };
 
 UCLASS(meta = (ScriptName = "GeometryScript_BoneWeights"))
@@ -232,6 +253,23 @@ public:
 		bool& bProfileExisted,
 		bool bReplaceExistingProfile = false,
 		FGeometryScriptBoneWeightProfile Profile = FGeometryScriptBoneWeightProfile() );
+
+	/**
+	 * Copies all bone weights from a source profile onto a target profile, on the same mesh, replacing all 
+	 * weights that existed on the target profile. If either the source or the target profile didn't exist,
+	 * then bProfileExisted will be set to false and no weights are copied.
+	 * @param bProfileExisted will be returned true if both of the requested bone weight profiles exist
+	 * @param TargetProfile The skin weight profile to copy to.
+	 * @param SourceProfile The skin weight profile to copy from.
+	 */
+	UFUNCTION(BlueprintCallable, Category = "GeometryScript|MeshQueries|BoneWeights", meta=(ScriptMethod))
+	static UPARAM(DisplayName = "Target Mesh") UDynamicMesh* 
+	MeshCopyBoneWeights( 
+		UDynamicMesh* TargetMesh,
+		bool& bProfileExisted,
+		FGeometryScriptBoneWeightProfile TargetProfile,
+		FGeometryScriptBoneWeightProfile SourceProfile = FGeometryScriptBoneWeightProfile()
+		);
 	
 	/**
 	 * Determine the largest bone weight index that exists on the Mesh
@@ -296,6 +334,27 @@ public:
 		FGeometryScriptBoneWeightProfile Profile = FGeometryScriptBoneWeightProfile(),
 		UGeometryScriptDebug* Debug = nullptr);
 
+	/**
+	 * Blends two bone weights using an Alpha value that ranges from 0 to 1, inclusive. If Alpha is 0, then only weights from BoneWeightsA are used,
+	 * and if Alpha is 1, then only weights from BoneWeightsB are used. For any value in between the weights are linearly interpolated.
+	 * Each bone weight from either array, that has the same bone index, are linearly interpolated. Any bone weights that are missing from either
+	 * BoneWeightsA or BoneWeightsB, are assumed to exist and have a weight of 0. After blending, the result is renormalized and sorted.
+	 * Values that are below the influence threshold, or exceeding the default bone weight limit (currently set to 12) will be thrown away.
+	 * @param BoneWeightsA List of bone weights to blend, such that its influence is greatest when Alpha is 0 and smallest when Alpha is 1.
+	 * @param BoneWeightsB List of bone weights to blend, such that its influence is greatest when Alpha is 1 and smallest when Alpha is 0.
+	 * @param Alpha The blending factor, ranging from 0 to 1, inclusive. Values outside of this range are clamped.
+	 * @param Result The resulting blend of the two bone weight arrays. 
+	 */
+	UFUNCTION(BlueprintPure, Category = "GeometryScript|MeshQueries|BoneWeights")
+	static void
+	BlendBoneWeights(
+		const TArray<FGeometryScriptBoneWeight>& BoneWeightsA,
+		const TArray<FGeometryScriptBoneWeight>& BoneWeightsB,
+		float Alpha,
+		TArray<FGeometryScriptBoneWeight>& Result,
+		UGeometryScriptDebug* Debug = nullptr
+	);
+	
 	/**
 	 * Set all vertices of the TargetMesh to the given Bone/Skin Weights
 	 * @param BoneWeights input array of bone index/weight pairs for the Vertex
@@ -308,7 +367,7 @@ public:
 		const TArray<FGeometryScriptBoneWeight>& BoneWeights,
 		FGeometryScriptBoneWeightProfile Profile = FGeometryScriptBoneWeightProfile(),
 		UGeometryScriptDebug* Debug = nullptr);
-
+	
 	/** 
 	 *  Computes a smooth skin binding for the given mesh to the skeleton provided.
 	 *  @param Skeleton The skeleton to compute binding for the skin weights.
diff --git a/Engine/Source/Runtime/GeometryCore/Private/DynamicMesh/DynamicMeshAttributeSet.cpp b/Engine/Source/Runtime/GeometryCore/Private/DynamicMesh/DynamicMeshAttributeSet.cpp
--- a/Engine/Source/Runtime/GeometryCore/Private/DynamicMesh/DynamicMeshAttributeSet.cpp
+++ b/Engine/Source/Runtime/GeometryCore/Private/DynamicMesh/DynamicMeshAttributeSet.cpp
@@ -1479,6 +1479,32 @@ namespace BoneAttributeHelpers
 			Attribute.Reset();
 		}
 	}
+
+	template <typename ParentType, typename AttribValueType>
+ 	void SparseCopyBoneAttribute(
+ 		FDynamicMesh3* Mesh,
+		TUniquePtr<TDynamicBoneAttributeBase<ParentType, AttribValueType>>& Attribute,
+		const TDynamicBoneAttributeBase<ParentType, AttribValueType>* Copy,
+		TConstArrayView<int32> InIndicesToCopy
+		)
+	{
+		if (Copy && !InIndicesToCopy.IsEmpty())
+		{
+			TDynamicBoneAttributeBase<ParentType, AttribValueType>* Ptr = new TDynamicBoneAttributeBase<ParentType, AttribValueType>(Mesh);
+			Attribute = TUniquePtr<TDynamicBoneAttributeBase<ParentType, AttribValueType>>(Ptr);
+
+			Attribute->Resize(InIndicesToCopy.Num());
+			int32 TargetIndex = 0;
+			for (int32 SourceIndex: InIndicesToCopy)
+			{
+				Attribute->SetValue(TargetIndex++, Copy->GetValue(SourceIndex));
+			}
+		}
+		else
+		{
+			Attribute.Reset();
+		}
+	}
 }
 
 int32 FDynamicMeshAttributeSet::GetNumBones() const
@@ -1494,6 +1520,54 @@ void FDynamicMeshAttributeSet::CopyBoneAttributes(const FDynamicMeshAttributeSet
 	BoneAttributeHelpers::CopyBoneAttribute(ParentMesh, BoneColorAttrib, Copy.GetBoneColors());
 }
 
+void FDynamicMeshAttributeSet::CopyBoneAttributesWithRemapping(const FDynamicMeshAttributeSet& Copy, const TMap<FName, FName>& BoneHierarchy)
+{
+	// Create an index array to mark which attribute values to copy. If there's no bone name attribute, we copy nothing.
+	TArray<int32> IndicesToCopy;
+	TMap<FName, int32> NameToIndexMap;
+	
+	if (Copy.GetBoneNames())
+	{
+		const TArray<FName>& SourceBoneNames = Copy.GetBoneNames()->GetAttribValues();
+		IndicesToCopy.Reserve(BoneHierarchy.Num());
+		int32 SourceIndex = 0;
+		for (FName BoneName: SourceBoneNames)
+		{
+			if (!BoneName.IsNone() && BoneHierarchy.Contains(BoneName))
+			{
+				NameToIndexMap.Add(BoneName, IndicesToCopy.Num());
+				IndicesToCopy.Add(SourceIndex);
+			}
+			SourceIndex++;
+		}
+	}
+
+	BoneAttributeHelpers::SparseCopyBoneAttribute(ParentMesh, BoneNameAttrib, Copy.GetBoneNames(), IndicesToCopy);
+	BoneAttributeHelpers::SparseCopyBoneAttribute(ParentMesh, BoneParentIndexAttrib, Copy.GetBoneParentIndices(), IndicesToCopy);
+	BoneAttributeHelpers::SparseCopyBoneAttribute(ParentMesh, BonePoseAttrib, Copy.GetBonePoses(), IndicesToCopy);
+	BoneAttributeHelpers::SparseCopyBoneAttribute(ParentMesh, BoneColorAttrib, Copy.GetBoneColors(), IndicesToCopy);
+
+	// If we copied parent indices, and we have bone names on the source, remap the parent bone index list to match.
+	// If the parent name is NAME_None or does not exist in our bone list, then we map that index to INDEX_NONE.
+	if (Copy.GetBoneNames() && BoneParentIndexAttrib)
+	{
+		const TArray<FName>& SourceBoneNames = Copy.GetBoneNames()->GetAttribValues();
+		for (int32& ParentIndex: BoneParentIndexAttrib->AttribValues)
+		{
+			const FName ParentName = SourceBoneNames.IsValidIndex(ParentIndex) ? SourceBoneNames[ParentIndex] : NAME_None;
+
+			if (const int32* BoneIndex = NameToIndexMap.Find(ParentName))
+			{
+				ParentIndex = *BoneIndex;
+			}
+			else
+			{
+				ParentIndex = INDEX_NONE;
+			}
+		}
+	}
+}
+
 void FDynamicMeshAttributeSet::EnableMatchingBoneAttributes(const FDynamicMeshAttributeSet& ToMatch, bool bClearExisting, bool bDiscardExtraAttributes)
 {
 	BoneAttributeHelpers::EnableIfMatching(ParentMesh, BoneNameAttrib, ToMatch.GetBoneNames(), (FName)NAME_None, bClearExisting, bDiscardExtraAttributes);
diff --git a/Engine/Source/Runtime/GeometryCore/Public/DynamicMesh/DynamicMeshAttributeSet.h b/Engine/Source/Runtime/GeometryCore/Public/DynamicMesh/DynamicMeshAttributeSet.h
--- a/Engine/Source/Runtime/GeometryCore/Public/DynamicMesh/DynamicMeshAttributeSet.h
+++ b/Engine/Source/Runtime/GeometryCore/Public/DynamicMesh/DynamicMeshAttributeSet.h
@@ -397,7 +397,20 @@ public:
 	// Bone Attributes
 	//
 
+	/** Copy all bone attributes from the given attribute set to this attribute set, removing any values that were there before. 
+	 *  @param Copy The attribute set to copy from. 
+	 */
 	GEOMETRYCORE_API void CopyBoneAttributes(const FDynamicMeshAttributeSet& Copy);
+
+	/** Copy bone attributes from the given attribute set to this attribute set, removing any values that were there before, using
+	 *  the given bone hierarchy map to control which bones are copied, and how parent bone indices should be remapped. 
+	 *  @param Copy The attribute set to copy from.
+	 *  @param BoneHierarchy The hierarchy to use to copy. The key is the bone name to copy and value is the new parent name to
+	 *    assign to it upon copy. The order of the resulting values is the same as in the Copy. A parent name of NAME_None indicates
+	 *    that the bone is a root bone and will have a parent index of INDEX_NONE. If a parent name is not in the list of bones to
+	 *    copy, then that bone's parent index will also be set to INDEX_NONE.
+	 */
+	GEOMETRYCORE_API void CopyBoneAttributesWithRemapping(const FDynamicMeshAttributeSet& Copy, const TMap<FName, FName>& BoneHierarchy); 
 	
 	GEOMETRYCORE_API void EnableMatchingBoneAttributes(const FDynamicMeshAttributeSet& ToMatch, bool bClearExisting, bool bDiscardExtraAttributes);
 
diff --git a/Engine/Source/Runtime/GeometryCore/Public/DynamicMesh/DynamicVertexSkinWeightsAttribute.h b/Engine/Source/Runtime/GeometryCore/Public/DynamicMesh/DynamicVertexSkinWeightsAttribute.h
--- a/Engine/Source/Runtime/GeometryCore/Public/DynamicMesh/DynamicVertexSkinWeightsAttribute.h
+++ b/Engine/Source/Runtime/GeometryCore/Public/DynamicMesh/DynamicVertexSkinWeightsAttribute.h
@@ -256,6 +256,27 @@ public:
 	// Accessors/Queries
 	//
 
+	/** Returns a list of all unique bone indices that are used by this skin weight attribute.
+	 *  \return A unique list of bone indices used by all bone weights across all vertices. 
+	 */
+	TSet<int32> GetBoundBoneIndices() const
+	{
+		if (VertexBoneWeights.Num() == 0)
+		{
+			return {};	// No weights defined, return the empty set.
+		}
+
+		TSet<int32> UniqueBoneIndices;
+		for (int VertexID = 0; VertexID < VertexBoneWeights.Num(); ++VertexID)
+		{
+			for (FBoneWeight BoneWeight: VertexBoneWeights[VertexID])
+			{
+				UniqueBoneIndices.Add(BoneWeight.GetBoneIndex());
+			}
+		}
+		return UniqueBoneIndices;
+	}
+
 	bool CopyThroughMapping(const TDynamicAttributeBase<ParentType>* Source, const FMeshIndexMappings& Mapping) override
 	{
 		// Don't snarf the FBoneWeight as a concrete object, since it _may_ contain a pointer
