Geometry Scripting: Build fix for Mac/Linux.

#rnx

diff --git a/Engine/Plugins/Runtime/GeometryScripting/Source/GeometryScriptingCore/Public/GeometryScript/MeshBoneWeightFunctions.h b/Engine/Plugins/Runtime/GeometryScripting/Source/GeometryScriptingCore/Public/GeometryScript/MeshBoneWeightFunctions.h
--- a/Engine/Plugins/Runtime/GeometryScripting/Source/GeometryScriptingCore/Public/GeometryScript/MeshBoneWeightFunctions.h
+++ b/Engine/Plugins/Runtime/GeometryScripting/Source/GeometryScriptingCore/Public/GeometryScript/MeshBoneWeightFunctions.h
@@ -17,6 +17,12 @@ struct GEOMETRYSCRIPTINGCORE_API FGeometryScriptBoneWeight
 {
 	GENERATED_BODY()
 
+	FGeometryScriptBoneWeight() = default;
+	FGeometryScriptBoneWeight(int32 InBoneIndex, float InWeight) :
+		BoneIndex(InBoneIndex), Weight(InWeight)
+	{
+	}
+
 	UPROPERTY(BlueprintReadWrite, Category = BoneWeights)
 	int32 BoneIndex = 0;
 
