add edge selection support to geometry script

diff --git a/Engine/Plugins/Runtime/GeometryScripting/Source/GeometryScriptingCore/Private/GeometryScriptSelectionTypes.cpp b/Engine/Plugins/Runtime/GeometryScripting/Source/GeometryScriptingCore/Private/GeometryScriptSelectionTypes.cpp
--- a/Engine/Plugins/Runtime/GeometryScripting/Source/GeometryScriptingCore/Private/GeometryScriptSelectionTypes.cpp
+++ b/Engine/Plugins/Runtime/GeometryScripting/Source/GeometryScriptingCore/Private/GeometryScriptSelectionTypes.cpp
@@ -10,6 +10,37 @@
 
 using namespace UE::Geometry;
 
+namespace UE::Private::GeometryScriptSelectionLocal
+{
+	// Assuming the selection is a triangle-edge selection, this enumerates all the *unique* mesh edge IDs
+	//  -- accounting for the fact that the GeoSelection is effectively a half-edge representation that can include the same edge twice (per triangle that includes it)
+	static void EnumerateUniqueEdgesInEdgeSelection(const FDynamicMesh3& Mesh, const FGeometrySelection& GeoSelection, TFunctionRef<void(int32)> PerEdgeFunc)
+	{
+		checkSlow(GeoSelection.ElementType == EGeometryElementType::Edge && GeoSelection.TopologyType == EGeometryTopologyType::Triangle);
+		for (uint64 Item : GeoSelection.Selection)
+		{
+			FMeshTriEdgeID TriEdgeID(FGeoSelectionID(Item).GeometryID);
+			if (Mesh.IsTriangle(TriEdgeID.TriangleID))
+			{
+				int32 EdgeID = Mesh.GetTriEdge(TriEdgeID.TriangleID, TriEdgeID.TriEdgeIndex);
+				// Check if edge is also in the selection via a lower triangle ID, and if so rely on that element to call the PerEdgeFunc
+				bool bAlsoInSelectionViaLowerTriangleID = false;
+				Mesh.EnumerateTriEdgeIDsFromEdgeID(EdgeID, [&GeoSelection, &TriEdgeID, &bAlsoInSelectionViaLowerTriangleID](FMeshTriEdgeID CurTriEdge)
+				{
+					if (CurTriEdge.TriangleID < TriEdgeID.TriangleID && GeoSelection.Selection.Contains((uint64)CurTriEdge.Encoded()))
+					{
+						bAlsoInSelectionViaLowerTriangleID = true;
+					}
+				});
+				if (!bAlsoInSelectionViaLowerTriangleID)
+				{
+					PerEdgeFunc(EdgeID);
+				}
+			}
+		}
+	}
+}
+
 
 FGeometryScriptMeshSelection::FGeometryScriptMeshSelection()
 {
@@ -62,6 +93,10 @@ EGeometryScriptMeshSelectionType FGeometryScriptMeshSelection::GetSelectionType(
 		{
 			return EGeometryScriptMeshSelectionType::Vertices;
 		}
+		else if (GeoSelection->ElementType == EGeometryElementType::Edge)
+		{
+			return EGeometryScriptMeshSelectionType::Edges;
+		}
 	}
 	UE_LOG(LogGeometry, Warning, TEXT("GeometryScriptMeshSelection::GetSelectionType() - GeoSelection has invalid type for FGeometryScriptMeshSelection!"));
 	return EGeometryScriptMeshSelectionType::Triangles;
@@ -73,6 +108,49 @@ int32 FGeometryScriptMeshSelection::GetNumSelected() const
 	return GeoSelection->Num();
 }
 
+int32 FGeometryScriptMeshSelection::GetNumUniqueSelected(const UE::Geometry::FDynamicMesh3& Mesh) const
+{
+	if (GeoSelection->TopologyType == EGeometryTopologyType::Polygroup)
+	{
+		if (GeoSelection->ElementType != EGeometryElementType::Face)
+		{
+			UE_LOG(LogGeometry, Warning, TEXT("GeometryScriptMeshSelection::GetNumUniqueSelected() - GeoSelection has invalid type for FGeometryScriptMeshSelection!"));
+			return 0;
+		}
+		
+		// The same group could be referenced from more than one triangle ID
+		int32 Count = 0;
+		TSet<int32> UniqueGroupIDs;
+		for (uint64 Item : GeoSelection->Selection)
+		{
+			bool bWasInSet;
+			UniqueGroupIDs.Add(FGeoSelectionID(Item).TopologyID, &bWasInSet);
+			Count += (int32)!bWasInSet;
+		}
+		return Count;
+	}
+	else
+	{
+		// Note vertex and face selections should already be unique, so can just return the selection set size
+		int32 Count = 0;
+		switch (GeoSelection->ElementType)
+		{
+		case EGeometryElementType::Vertex:
+			return GeoSelection->Selection.Num();
+		case EGeometryElementType::Edge:
+			// iterate to count unique edges
+			UE::Private::GeometryScriptSelectionLocal::EnumerateUniqueEdgesInEdgeSelection(Mesh, *GeoSelection, [&Count](int32 EdgeID) {++Count;});
+			return Count;
+		case EGeometryElementType::Face:
+			return GeoSelection->Selection.Num();
+		}
+
+		UE_LOG(LogGeometry, Warning, TEXT("GeometryScriptMeshSelection::GetNumUniqueSelected() - GeoSelection has invalid type for FGeometryScriptMeshSelection!"));
+		return 0;
+	}
+	
+}
+
 void FGeometryScriptMeshSelection::DebugPrint() const
 {
 	bool bIsPolygroups = (GeoSelection->TopologyType == EGeometryTopologyType::Polygroup);
@@ -225,6 +303,7 @@ EGeometryScriptIndexType FGeometryScriptMeshSelection::ConvertToMeshIndexArray(
 
 	bool bIsSameType =
 		(ConvertToType == EGeometryScriptIndexType::Vertex && GeoSelection->TopologyType == EGeometryTopologyType::Triangle && GeoSelection->ElementType == EGeometryElementType::Vertex)
+		|| (ConvertToType == EGeometryScriptIndexType::Edge && GeoSelection->TopologyType == EGeometryTopologyType::Triangle && GeoSelection->ElementType == EGeometryElementType::Edge)
 		|| (ConvertToType == EGeometryScriptIndexType::Triangle && GeoSelection->TopologyType == EGeometryTopologyType::Triangle && GeoSelection->ElementType == EGeometryElementType::Face)
 		|| (ConvertToType == EGeometryScriptIndexType::PolygroupID && GeoSelection->TopologyType == EGeometryTopologyType::Polygroup && GeoSelection->ElementType == EGeometryElementType::Face);
 
@@ -245,6 +324,15 @@ EGeometryScriptIndexType FGeometryScriptMeshSelection::ConvertToMeshIndexArray(
 				}
 				return EGeometryScriptIndexType::Vertex;
 			}
+			else if (GeoSelection->ElementType == EGeometryElementType::Edge)
+			{
+				UE::Private::GeometryScriptSelectionLocal::EnumerateUniqueEdgesInEdgeSelection(Mesh, *GeoSelection, [&IndexListOut](int32 EdgeID)
+				{
+					IndexListOut.Add(EdgeID);
+				});
+
+				return EGeometryScriptIndexType::Edge;
+			}
 			else if (GeoSelection->ElementType == EGeometryElementType::Face)
 			{
 				for (uint64 Item : GeoSelection->Selection)
@@ -273,8 +361,7 @@ EGeometryScriptIndexType FGeometryScriptMeshSelection::ConvertToMeshIndexArray(
 	{
 		TSet<int32> UniqueTriangles;
 
-		if ( (GeoSelection->TopologyType == EGeometryTopologyType::Triangle) &&
-			 (GeoSelection->ElementType == EGeometryElementType::Vertex) )
+		if (GeoSelection->TopologyType == EGeometryTopologyType::Triangle)
 		{
 			UE::Geometry::EnumerateSelectionTriangles(*GeoSelection, Mesh,
 				[&](int32 TriangleID) { UniqueTriangles.Add(TriangleID); });
@@ -292,24 +379,58 @@ EGeometryScriptIndexType FGeometryScriptMeshSelection::ConvertToMeshIndexArray(
 		}
 
 	}
+	else if (ConvertToType == EGeometryScriptIndexType::Edge)
+	{
+		TSet<int32> UniqueEdges;
+		if (GeoSelection->TopologyType == EGeometryTopologyType::Triangle)
+		{
+			UE::Geometry::EnumerateSelectionEdges(*GeoSelection, Mesh,
+			[&Mesh, &UniqueEdges, &IndexListOut](int32 EdgeID)
+			{
+				bool bAlreadyInSet;
+				UniqueEdges.Add(EdgeID, &bAlreadyInSet);
+				if (!bAlreadyInSet)
+				{
+					IndexListOut.Add(EdgeID);
+				}
+			});
+			return EGeometryScriptIndexType::Edge;
+		}
+		else if (GeoSelection->TopologyType == EGeometryTopologyType::Polygroup && GeoSelection->ElementType == EGeometryElementType::Face)
+		{
+			// Currently doing this locally because the util functions don't support a GeoSelectionID w/ a Group ID but not a Triangle ID
+			EnumerateGroupTriangles([&Mesh, &UniqueEdges, &IndexListOut](int32 TriangleID)
+			{
+				FIndex3i TriEdges = Mesh.GetTriEdges(TriangleID);
+				for (int32 SubIdx = 0; SubIdx < 3; ++SubIdx)
+				{
+					bool bAlreadyInSet;
+					UniqueEdges.Add(TriEdges[SubIdx], &bAlreadyInSet);
+					if (!bAlreadyInSet)
+					{
+						IndexListOut.Add(TriEdges[SubIdx]);
+					}
+				}
+			});
+			return EGeometryScriptIndexType::Edge;
+		}
+	}
 	else if (ConvertToType == EGeometryScriptIndexType::Vertex)
 	{
 		TSet<int32> UniqueVertices;
 
-		if ( (GeoSelection->TopologyType == EGeometryTopologyType::Triangle) &&
-			(GeoSelection->ElementType == EGeometryElementType::Face) )
+		if (GeoSelection->TopologyType == EGeometryTopologyType::Triangle)
 		{
-			UE::Geometry::EnumerateSelectionTriangles(*GeoSelection, Mesh,
-			[&](int32 TriangleID) { 
-				FIndex3i TriV = Mesh.GetTriangle(TriangleID);
-				UniqueVertices.Add(TriV.A); 
-				UniqueVertices.Add(TriV.B);
-				UniqueVertices.Add(TriV.C);
-			});
-			for (int32 vid : UniqueVertices)
+			UE::Geometry::EnumerateTriangleSelectionVertices(*GeoSelection, Mesh, nullptr, [&UniqueVertices, &IndexListOut](uint64 VertexID, const FVector3d&)
 			{
-				IndexListOut.Add(vid);
-			}
+				int32 CastVID = (int32)VertexID;
+				bool bAlreadyInSet;
+				UniqueVertices.Add(CastVID, &bAlreadyInSet);
+				if (!bAlreadyInSet)
+				{
+					IndexListOut.Add(CastVID);
+				}
+			});
 			return EGeometryScriptIndexType::Vertex;
 		}
 		else if (GeoSelection->TopologyType == EGeometryTopologyType::Polygroup && GeoSelection->ElementType == EGeometryElementType::Face)
@@ -380,7 +501,7 @@ void FGeometryScriptMeshSelection::ProcessByTriangleID(const UE::Geometry::FDyna
 	}
 	else if (GetSelectionType() == EGeometryScriptMeshSelectionType::Vertices)
 	{
-		TSet<uint64> UniqueItems;
+		TSet<int32> UniqueItems;
 		for (uint64 Item : GeoSelection->Selection)
 		{
 			int32 VertexID = FGeoSelectionID(Item).GeometryID;
@@ -398,6 +519,28 @@ void FGeometryScriptMeshSelection::ProcessByTriangleID(const UE::Geometry::FDyna
 			}
 		}
 	}
+	else if (GetSelectionType() == EGeometryScriptMeshSelectionType::Edges)
+	{
+		TSet<int32> UniqueItems;
+		for (uint64 Item : GeoSelection->Selection)
+		{
+			FMeshTriEdgeID TriEdgeID(FGeoSelectionID(Item).GeometryID);
+			if (Mesh.IsTriangle(TriEdgeID.TriangleID))
+			{
+				int32 EdgeID = Mesh.GetTriEdge(TriEdgeID.TriangleID, TriEdgeID.TriEdgeIndex);
+				FIndex2i EdgeT = Mesh.GetEdgeT(EdgeID);
+				for (int32 SubIdx = 0, NumT = EdgeT.B != FDynamicMesh3::InvalidID ? 2 : 1; SubIdx < NumT; ++SubIdx)
+				{
+					bool bAlreadyInSet = false;
+					UniqueItems.Add(EdgeT[SubIdx], &bAlreadyInSet);
+					if (!bAlreadyInSet)
+					{
+						PerTriangleFunc(EdgeT[SubIdx]);
+					}
+				}
+			}
+		}
+	}
 	else
 	{
 		TSet<int32> GroupIDs;
@@ -442,7 +585,7 @@ void FGeometryScriptMeshSelection::ProcessByVertexID(const UE::Geometry::FDynami
 	}
 	else 
 	{
-		TSet<uint64> UniqueVertices;
+		TSet<int32> UniqueVertices;
 		auto ProcessTriangle = [&Mesh, &PerVertexFunc, &UniqueVertices](int32 TriangleID)
 		{
 			FIndex3i Triangle = Mesh.GetTriangle(TriangleID);
@@ -468,6 +611,27 @@ void FGeometryScriptMeshSelection::ProcessByVertexID(const UE::Geometry::FDynami
 				}
 			}
 		}
+		else if (GetSelectionType() == EGeometryScriptMeshSelectionType::Edges)
+		{
+			for (uint64 Item : GeoSelection->Selection)
+			{
+				FMeshTriEdgeID TriEdgeID(FGeoSelectionID(Item).GeometryID);
+				if (Mesh.IsTriangle(TriEdgeID.TriangleID))
+				{
+					int32 EdgeID = Mesh.GetTriEdge(TriEdgeID.TriangleID, TriEdgeID.TriEdgeIndex);
+					FIndex2i EdgeV = Mesh.GetEdgeV(EdgeID);
+					for (int32 SubIdx = 0; SubIdx < 2; ++SubIdx)
+					{
+						bool bAlreadyInSet = false;
+						UniqueVertices.Add(EdgeV[SubIdx], &bAlreadyInSet);
+						if (!bAlreadyInSet)
+						{
+							PerVertexFunc(EdgeV[SubIdx]);
+						}
+					}
+				}
+			}
+		}
 		else
 		{
 			TSet<int32> GroupIDs;
@@ -485,3 +649,90 @@ void FGeometryScriptMeshSelection::ProcessByVertexID(const UE::Geometry::FDynami
 		}
 	}
 }
+
+void FGeometryScriptMeshSelection::ProcessByEdgeID(const UE::Geometry::FDynamicMesh3& Mesh,
+	TFunctionRef<void(int32)> PerEdgeFunc,
+	bool bProcessAllEdgesIfSelectionEmpty) const
+{
+	if (IsEmpty())
+	{
+		if (bProcessAllEdgesIfSelectionEmpty)
+		{
+			for (int32 EdgeID : Mesh.EdgeIndicesItr())
+			{
+				PerEdgeFunc(EdgeID);
+			}
+		}
+		return;
+	}
+
+	if (GetSelectionType() == EGeometryScriptMeshSelectionType::Edges)
+	{
+		UE::Private::GeometryScriptSelectionLocal::EnumerateUniqueEdgesInEdgeSelection(Mesh, *GeoSelection, PerEdgeFunc);
+	}
+	else
+	{
+		TSet<int32> UniqueEdges;
+		auto ProcessTriangle = [&Mesh, &PerEdgeFunc, &UniqueEdges](int32 TriangleID)
+		{
+			FIndex3i TriEdges = Mesh.GetTriEdges(TriangleID);
+			for (int32 j = 0; j < 3; ++j)
+			{
+				bool bAlreadyInSet = false;
+				UniqueEdges.Add(TriEdges[j], &bAlreadyInSet);
+				if (bAlreadyInSet == false)
+				{
+					PerEdgeFunc(TriEdges[j]);
+				}
+			}
+		};
+
+		if (GetSelectionType() == EGeometryScriptMeshSelectionType::Triangles)
+		{
+			for (uint64 Item : GeoSelection->Selection)
+			{
+				int32 TriangleID = FGeoSelectionID(Item).GeometryID;
+				if (Mesh.IsTriangle(TriangleID))
+				{
+					ProcessTriangle(TriangleID);
+				}
+			}
+		}
+		else if (GetSelectionType() == EGeometryScriptMeshSelectionType::Vertices)
+		{
+			TSet<int32> UniqueItems;
+			for (uint64 Item : GeoSelection->Selection)
+			{
+				int32 VertexID = FGeoSelectionID(Item).GeometryID;
+				if (Mesh.IsVertex(VertexID))
+				{
+					Mesh.EnumerateVertexEdges(VertexID, [&UniqueItems, &PerEdgeFunc](int32 EdgeID)
+						{
+							bool bAlreadyInSet = false;
+							UniqueItems.Add(EdgeID, &bAlreadyInSet);
+							if (bAlreadyInSet == false)
+							{
+								PerEdgeFunc(EdgeID);
+							}
+						});
+				}
+			}
+		}
+		else // polygroups
+		{
+			TSet<int32> GroupIDs;
+			for (uint64 Item : GeoSelection->Selection)
+			{
+				GroupIDs.Add(FGeoSelectionID(Item).TopologyID);
+			}
+			for (int32 TriangleID : Mesh.TriangleIndicesItr())
+			{
+				if (GroupIDs.Contains(Mesh.GetTriangleGroup(TriangleID)))
+				{
+					ProcessTriangle(TriangleID);
+				}
+			}
+		}
+	}
+}
+
diff --git a/Engine/Plugins/Runtime/GeometryScripting/Source/GeometryScriptingCore/Private/MeshNormalsFunctions.cpp b/Engine/Plugins/Runtime/GeometryScripting/Source/GeometryScriptingCore/Private/MeshNormalsFunctions.cpp
--- a/Engine/Plugins/Runtime/GeometryScripting/Source/GeometryScriptingCore/Private/MeshNormalsFunctions.cpp
+++ b/Engine/Plugins/Runtime/GeometryScripting/Source/GeometryScriptingCore/Private/MeshNormalsFunctions.cpp
@@ -171,7 +171,8 @@ UDynamicMesh* UGeometryScriptLibrary_MeshNormalsFunctions::RecomputeNormalsForMe
 			return;
 		}
 		
-		if (Selection.GetSelectionType() == EGeometryScriptMeshSelectionType::Vertices)
+		EGeometryScriptMeshSelectionType SelectionType = Selection.GetSelectionType();
+		if (SelectionType == EGeometryScriptMeshSelectionType::Vertices || SelectionType == EGeometryScriptMeshSelectionType::Edges)
 		{
 			FDynamicMeshNormalOverlay* Normals = EditMesh.Attributes()->PrimaryNormals();
 			TSet<int32> Elements;
diff --git a/Engine/Plugins/Runtime/GeometryScripting/Source/GeometryScriptingCore/Private/MeshSelectionFunctions.cpp b/Engine/Plugins/Runtime/GeometryScripting/Source/GeometryScriptingCore/Private/MeshSelectionFunctions.cpp
--- a/Engine/Plugins/Runtime/GeometryScripting/Source/GeometryScriptingCore/Private/MeshSelectionFunctions.cpp
+++ b/Engine/Plugins/Runtime/GeometryScripting/Source/GeometryScriptingCore/Private/MeshSelectionFunctions.cpp
@@ -11,6 +11,7 @@
 #include "Spatial/FastWinding.h"
 #include "Selections/MeshConnectedComponents.h"
 #include "Selections/MeshFaceSelection.h"
+#include "Selections/MeshEdgeSelection.h"
 #include "Selections/MeshVertexSelection.h"
 
 #include UE_INLINE_GENERATED_CPP_BY_NAME(MeshSelectionFunctions)
@@ -58,6 +59,21 @@ static void SelectByTriangleAttribute(const FDynamicMesh3& ReadMesh,
 			}
 		}
 	}
+	else if (SelectionType == EGeometryScriptMeshSelectionType::Edges)
+	{
+		NewSelection.InitializeTypes(EGeometryElementType::Edge, EGeometryTopologyType::Triangle);
+		for (int32 TID : ReadMesh.TriangleIndicesItr())
+		{
+			if (TriangleFilter(TID))
+			{
+				for (int32 SubIdx = 0; SubIdx < 3; ++SubIdx)
+				{
+					// Note edges are added per triangle that contains them, ala half-edges
+					NewSelection.Selection.Add(FGeoSelectionID::MeshEdge(FMeshTriEdgeID(TID, SubIdx)).Encoded());
+				}
+			}
+		}
+	}
 	else
 	{
 		NewSelection.InitializeTypes(EGeometryElementType::Face, EGeometryTopologyType::Polygroup);
@@ -108,6 +124,29 @@ static void SelectMeshElementsWithContainmentTest(
 			}
 		});
 	}
+	else if (SelectionType == EGeometryScriptMeshSelectionType::Edges)
+	{
+		GeoSelection.InitializeTypes(EGeometryElementType::Edge, EGeometryTopologyType::Triangle);
+		TargetMesh->ProcessMesh([&ContainmentFunc, &GeoSelection, bNeedsNormals, NumTrianglePoints](const FDynamicMesh3& Mesh)
+		{
+			int32 UseNumEdgePoints = FMath::Clamp(NumTrianglePoints, 1, 2);
+			for (int32 EID : Mesh.EdgeIndicesItr())
+			{
+				FIndex2i EdgeV = Mesh.GetEdgeV(EID);
+				FVector3d UseNormal = bNeedsNormals ? Mesh.GetEdgeNormal(EID) : FVector3d::UnitZ();
+				int32 NumContained =
+					int32(ContainmentFunc(Mesh.GetVertex(EdgeV.A), UseNormal)) +
+					int32(ContainmentFunc(Mesh.GetVertex(EdgeV.B), UseNormal));
+				if (NumContained >= UseNumEdgePoints)
+				{
+					Mesh.EnumerateTriEdgeIDsFromEdgeID(EID, [&GeoSelection](FMeshTriEdgeID TriEdgeID)
+					{
+						GeoSelection.Selection.Add(TriEdgeID.Encoded());
+					});
+				}
+			}
+		});
+	}
 	else
 	{
 		GeoSelection.InitializeTypes(EGeometryElementType::Face, 
@@ -156,6 +195,30 @@ void UGeometryScriptLibrary_MeshSelectionFunctions::GetMeshSelectionInfo(
 	NumSelected = Selection.GetNumSelected();
 }
 
+void UGeometryScriptLibrary_MeshSelectionFunctions::GetMeshUniqueSelectionInfo(
+	const UDynamicMesh* TargetMesh,
+	FGeometryScriptMeshSelection Selection,
+	EGeometryScriptMeshSelectionType& SelectionType,
+	int& NumSelected
+)
+{
+	SelectionType = Selection.GetSelectionType();
+	if (!TargetMesh)
+	{
+		UE_LOG(LogGeometry, Warning, TEXT("GetMeshUniqueSelectionInfo: TargetMesh is Null"));
+		NumSelected = Selection.GetNumSelected();
+		return;
+	}
+	else
+	{
+		TargetMesh->ProcessMesh([&NumSelected, &Selection](const FDynamicMesh3& Mesh)
+		{
+			NumSelected = Selection.GetNumUniqueSelected(Mesh);
+		});
+	}
+}
+
+
 
 void UGeometryScriptLibrary_MeshSelectionFunctions::DebugPrintMeshSelection(
 	FGeometryScriptMeshSelection Selection,
@@ -197,6 +260,16 @@ UDynamicMesh* UGeometryScriptLibrary_MeshSelectionFunctions::CreateSelectAllMesh
 				NewSelection.Selection.Add(FGeoSelectionID::MeshTriangle(tid).Encoded());
 			}
 		}
+		else if (SelectionType == EGeometryScriptMeshSelectionType::Edges)
+		{
+			NewSelection.InitializeTypes(EGeometryElementType::Edge, EGeometryTopologyType::Triangle);
+			for (int32 TID : ReadMesh.TriangleIndicesItr())
+			{
+				NewSelection.Selection.Add(FMeshTriEdgeID(TID, 0).Encoded());
+				NewSelection.Selection.Add(FMeshTriEdgeID(TID, 1).Encoded());
+				NewSelection.Selection.Add(FMeshTriEdgeID(TID, 2).Encoded());
+			}
+		}
 		else
 		{
 			NewSelection.InitializeTypes(EGeometryElementType::Face, EGeometryTopologyType::Polygroup);
@@ -322,41 +395,138 @@ UDynamicMesh* UGeometryScriptLibrary_MeshSelectionFunctions::ConvertMeshSelectio
 
 	if (NewType == EGeometryScriptMeshSelectionType::Vertices)
 	{
-		TSet<int32> CurTriangles, CurVertices;
-		TargetMesh->ProcessMesh([&](const FDynamicMesh3& ReadMesh)
+		TargetMesh->ProcessMesh([&FromSelection, &ToSelection, bAllowPartialInclusion](const FDynamicMesh3& ReadMesh)
 		{
+			TSet<int32> CurElements, CurVertices;
+			
+			FGeometrySelection NewSelection;
+			NewSelection.InitializeTypes(EGeometryElementType::Vertex, EGeometryTopologyType::Triangle);
+
 			if (bAllowPartialInclusion)
 			{
-				FromSelection.ProcessByTriangleID(ReadMesh, [&](int32 TriangleID) { 
-					FIndex3i Vertices = ReadMesh.GetTriangle(TriangleID);
-					CurVertices.Add(Vertices.A); CurVertices.Add(Vertices.B); CurVertices.Add(Vertices.C);
+				FromSelection.ProcessByVertexID(ReadMesh, [&NewSelection](int32 VertexID) { 
+					NewSelection.Selection.Add(FGeoSelectionID::MeshVertex(VertexID).Encoded());
 				});
-				ConvertIndexSetToMeshSelection(TargetMesh, CurVertices, EGeometryScriptMeshSelectionType::Vertices, ToSelection);
 			}
-			else
+			else if(FromSelection.GetSelectionType() == EGeometryScriptMeshSelectionType::Triangles
+				|| FromSelection.GetSelectionType() == EGeometryScriptMeshSelectionType::Polygroups)
 			{
-				FromSelection.ProcessByTriangleID(ReadMesh, [&](int32 TriangleID) { 
-					CurTriangles.Add(TriangleID);
+				FromSelection.ProcessByTriangleID(ReadMesh, [&CurElements, &CurVertices, &ReadMesh](int32 TriangleID)
+				{
+					CurElements.Add(TriangleID);
 					FIndex3i Vertices = ReadMesh.GetTriangle(TriangleID);
 					CurVertices.Add(Vertices.A); CurVertices.Add(Vertices.B); CurVertices.Add(Vertices.C);
 				});
 
-				FGeometrySelection NewSelection;
-				NewSelection.InitializeTypes(EGeometryElementType::Vertex, EGeometryTopologyType::Triangle);
-				for (int32 vid : CurVertices)
+				for (int32 VID : CurVertices)
+				{
+					bool bAllInSet = true;
+					ReadMesh.EnumerateVertexTriangles(VID, [&bAllInSet, &CurElements](int32 TID) { bAllInSet = bAllInSet && CurElements.Contains(TID); });
+					if (bAllInSet)
+					{
+						NewSelection.Selection.Add(FGeoSelectionID::MeshVertex(VID).Encoded());
+					}
+				}
+			}
+			else if (FromSelection.GetSelectionType() == EGeometryScriptMeshSelectionType::Edges)
+			{
+				FromSelection.ProcessByEdgeID(ReadMesh, [&CurElements, &CurVertices, &ReadMesh](int32 EdgeID)
+				{
+					CurElements.Add(EdgeID);
+					FIndex2i Vertices = ReadMesh.GetEdgeV(EdgeID);
+					CurVertices.Add(Vertices.A); CurVertices.Add(Vertices.B);
+				});
+
+				for (int32 VID : CurVertices)
 				{
 					bool bAllInSet = true;
-					ReadMesh.EnumerateVertexTriangles(vid, [&](int32 tid) { bAllInSet = bAllInSet && CurTriangles.Contains(tid); });
+					ReadMesh.EnumerateVertexEdges(VID, [&bAllInSet, &CurElements](int32 EID) { bAllInSet = bAllInSet && CurElements.Contains(EID); });
 					if (bAllInSet)
 					{
-						NewSelection.Selection.Add(FGeoSelectionID::MeshVertex(vid).Encoded());
+						NewSelection.Selection.Add(FGeoSelectionID::MeshVertex(VID).Encoded());
 					}
 				}
-				ToSelection.SetSelection(MoveTemp(NewSelection));
 			}
+			else
+			{
+				ensureMsgf(false, TEXT("Unhandled mesh selection type"));
+			}
+
+			ToSelection.SetSelection(MoveTemp(NewSelection));
 		});
 
 	}
+	else if (NewType == EGeometryScriptMeshSelectionType::Edges)
+	{
+		FGeometrySelection NewSelection;
+		NewSelection.InitializeTypes(EGeometryElementType::Edge, EGeometryTopologyType::Triangle);
+
+		if (bAllowPartialInclusion)
+		{
+			TargetMesh->ProcessMesh([&FromSelection, &NewSelection](const FDynamicMesh3& ReadMesh)
+			{
+				FromSelection.ProcessByEdgeID(ReadMesh, [&ReadMesh, &NewSelection](int32 EdgeID)
+				{
+					ReadMesh.EnumerateTriEdgeIDsFromEdgeID(EdgeID, [&NewSelection](FMeshTriEdgeID TriEdgeID)
+					{
+						NewSelection.Selection.Add(TriEdgeID.Encoded());
+					});
+				});
+			});
+		}
+		else if (FromSelection.GetSelectionType() == EGeometryScriptMeshSelectionType::Vertices) // select edges w/ both verts selected
+		{
+			TargetMesh->ProcessMesh([&FromSelection, &NewSelection](const FDynamicMesh3& ReadMesh)
+			{
+				TSet<int32> CurVertices;
+				FromSelection.ProcessByVertexID(ReadMesh, [&CurVertices](int32 VertexID)
+				{
+					CurVertices.Add(VertexID);
+				});
+				FromSelection.ProcessByEdgeID(ReadMesh, [&CurVertices, &ReadMesh, &NewSelection](int32 EdgeID)
+				{
+					FIndex2i EdgeV = ReadMesh.GetEdgeV(EdgeID);
+					if (CurVertices.Contains(EdgeV.A) && CurVertices.Contains(EdgeV.B))
+					{
+						ReadMesh.EnumerateTriEdgeIDsFromEdgeID(EdgeID, [&NewSelection](FMeshTriEdgeID TriEdgeID)
+						{
+							NewSelection.Selection.Add(TriEdgeID.Encoded());
+						});
+					}
+				});
+			});
+		}
+		// select edges w/ all tris selected
+		else if (FromSelection.GetSelectionType() == EGeometryScriptMeshSelectionType::Triangles
+			|| FromSelection.GetSelectionType() == EGeometryScriptMeshSelectionType::Polygroups)
+		{
+			TargetMesh->ProcessMesh([&FromSelection, &NewSelection](const FDynamicMesh3& ReadMesh)
+			{
+				TSet<int32> CurTriangles;
+				FromSelection.ProcessByTriangleID(ReadMesh, [&CurTriangles](int32 TriangleID)
+				{
+					CurTriangles.Add(TriangleID);
+				});
+				FromSelection.ProcessByEdgeID(ReadMesh, [&CurTriangles, &ReadMesh, &NewSelection](int32 EdgeID)
+				{
+					FIndex2i EdgeT = ReadMesh.GetEdgeT(EdgeID);
+					if (CurTriangles.Contains(EdgeT.A) && (EdgeT.B == FDynamicMesh3::InvalidID || CurTriangles.Contains(EdgeT.B)))
+					{
+						ReadMesh.EnumerateTriEdgeIDsFromEdgeID(EdgeID, [&NewSelection](FMeshTriEdgeID TriEdgeID)
+						{
+							NewSelection.Selection.Add(TriEdgeID.Encoded());
+						});
+					}
+				});
+			});
+		}
+		else
+		{
+			ensureMsgf(false, TEXT("Unhandled mesh selection type"));
+		}
+		
+		ToSelection.SetSelection(MoveTemp(NewSelection));
+	}
 	else if (NewType == EGeometryScriptMeshSelectionType::Triangles)
 	{
 		FGeometrySelection NewSelection;
@@ -364,23 +534,23 @@ UDynamicMesh* UGeometryScriptLibrary_MeshSelectionFunctions::ConvertMeshSelectio
 
 		if (FromSelection.GetSelectionType() == EGeometryScriptMeshSelectionType::Polygroups || bAllowPartialInclusion)
 		{
-			TargetMesh->ProcessMesh([&](const FDynamicMesh3& ReadMesh)
+			TargetMesh->ProcessMesh([&FromSelection, &NewSelection](const FDynamicMesh3& ReadMesh)
 			{
-				FromSelection.ProcessByTriangleID(ReadMesh, [&](int32 TriangleID) {
+				FromSelection.ProcessByTriangleID(ReadMesh, [&NewSelection](int32 TriangleID) {
 					NewSelection.Selection.Add(FGeoSelectionID::MeshTriangle(TriangleID).Encoded());
 				});
 			});
 		}
-		else   // vertex selection w/ no partial inclusion, ie only "full" triangles
+		else if (FromSelection.GetSelectionType() == EGeometryScriptMeshSelectionType::Vertices) // vertex selection w/ no partial inclusion, ie only "full" triangles
 		{
-			TargetMesh->ProcessMesh([&](const FDynamicMesh3& ReadMesh)
+			TargetMesh->ProcessMesh([&FromSelection, &NewSelection](const FDynamicMesh3& ReadMesh)
 			{
 				// in this case FromSelection already has this set! but we do not have access to it...
 				TSet<int32> CurVertices;
-				FromSelection.ProcessByVertexID(ReadMesh, [&](int32 VertexID) {
+				FromSelection.ProcessByVertexID(ReadMesh, [&CurVertices](int32 VertexID) {
 					CurVertices.Add(VertexID);
 				});
-				FromSelection.ProcessByTriangleID(ReadMesh, [&](int32 TriangleID) {
+				FromSelection.ProcessByTriangleID(ReadMesh, [&CurVertices, &NewSelection, &ReadMesh](int32 TriangleID) {
 					FIndex3i Triangle = ReadMesh.GetTriangle(TriangleID);
 					if (CurVertices.Contains(Triangle.A) && CurVertices.Contains(Triangle.B) && CurVertices.Contains(Triangle.C))
 					{
@@ -389,6 +559,27 @@ UDynamicMesh* UGeometryScriptLibrary_MeshSelectionFunctions::ConvertMeshSelectio
 				});
 			});
 		}
+		else if (FromSelection.GetSelectionType() == EGeometryScriptMeshSelectionType::Edges)
+		{
+			TargetMesh->ProcessMesh([&FromSelection, &NewSelection](const FDynamicMesh3& ReadMesh)
+			{
+				TSet<int32> CurEdges;
+				FromSelection.ProcessByEdgeID(ReadMesh, [&CurEdges](int32 EdgeID) {
+					CurEdges.Add(EdgeID);
+				});
+				FromSelection.ProcessByTriangleID(ReadMesh, [&ReadMesh, &CurEdges, &NewSelection](int32 TriangleID) {
+					FIndex3i TriangleEdges = ReadMesh.GetTriEdges(TriangleID);
+					if (CurEdges.Contains(TriangleEdges.A) && CurEdges.Contains(TriangleEdges.B) && CurEdges.Contains(TriangleEdges.C))
+					{
+						NewSelection.Selection.Add(FGeoSelectionID::MeshTriangle(TriangleID).Encoded());
+					}
+				});
+			});
+		}
+		else
+		{
+			ensureMsgf(false, TEXT("Unhandled mesh selection type"));
+		}
 
 		ToSelection.SetSelection(MoveTemp(NewSelection));
 	}
@@ -410,6 +601,9 @@ UDynamicMesh* UGeometryScriptLibrary_MeshSelectionFunctions::ConvertMeshSelectio
 			TSet<int32> UniqueGroupIDs;
 			TargetMesh->ProcessMesh([&](const FDynamicMesh3& ReadMesh)
 			{
+				// Note: for vertex and edge selections, will include all 'touched' triangles
+				// This is less strict than one might expect.
+				// If the stricter selection conversion is desired, please consider how to do so without changing existing BP behavior.
 				TSet<int32> AllTriangles;
 				FromSelection.ProcessByTriangleID(ReadMesh, [&](int32 TriangleID) {
 					AllTriangles.Add(TriangleID);
@@ -472,6 +666,20 @@ UDynamicMesh* UGeometryScriptLibrary_MeshSelectionFunctions::ConvertIndexArrayTo
 		GeoSelection.InitializeTypes(EGeometryElementType::Vertex, EGeometryTopologyType::Triangle);
 		break;
 
+	case EGeometryScriptMeshSelectionType::Edges:
+		GeoSelection.InitializeTypes(EGeometryElementType::Edge, EGeometryTopologyType::Triangle);
+		TargetMesh->ProcessMesh([&GeoSelection, &IndexArray](const FDynamicMesh3& Mesh)
+		{
+			for (int32 EID : IndexArray)
+			{
+				Mesh.EnumerateTriEdgeIDsFromEdgeID(EID, [&GeoSelection](FMeshTriEdgeID TriEdgeID)
+				{
+					GeoSelection.Selection.Add(TriEdgeID.Encoded());
+				});
+			}
+		});
+		break;
+
 	case EGeometryScriptMeshSelectionType::Polygroups:
 		for (int32 gid : IndexArray)
 		{
@@ -483,6 +691,9 @@ UDynamicMesh* UGeometryScriptLibrary_MeshSelectionFunctions::ConvertIndexArrayTo
 		GeoSelection.InitializeTypes(EGeometryElementType::Face, EGeometryTopologyType::Polygroup);
 		break;
 
+	default:
+		ensureMsgf(false, TEXT("Unhandled mesh selection type"));
+
 	}
 
 	SelectionOut.SetSelection(MoveTemp(GeoSelection));
@@ -523,6 +734,20 @@ UDynamicMesh* UGeometryScriptLibrary_MeshSelectionFunctions::ConvertIndexSetToMe
 		GeoSelection.InitializeTypes(EGeometryElementType::Vertex, EGeometryTopologyType::Triangle);
 		break;
 
+	case EGeometryScriptMeshSelectionType::Edges:
+		GeoSelection.InitializeTypes(EGeometryElementType::Edge, EGeometryTopologyType::Triangle);
+		TargetMesh->ProcessMesh([&GeoSelection, &IndexSet](const FDynamicMesh3& Mesh)
+		{
+			for (int32 EID : IndexSet)
+			{
+				Mesh.EnumerateTriEdgeIDsFromEdgeID(EID, [&GeoSelection](FMeshTriEdgeID TriEdgeID)
+				{
+					GeoSelection.Selection.Add(TriEdgeID.Encoded());
+				});
+			}
+		});
+		break;
+
 	case EGeometryScriptMeshSelectionType::Polygroups:
 		for (int32 gid : IndexSet)
 		{
@@ -534,6 +759,9 @@ UDynamicMesh* UGeometryScriptLibrary_MeshSelectionFunctions::ConvertIndexSetToMe
 		GeoSelection.InitializeTypes(EGeometryElementType::Face, EGeometryTopologyType::Polygroup);
 
 		break;
+
+	default:
+		ensureMsgf(false, TEXT("Unhandled mesh selection type"));
 	}
 
 	SelectionOut.SetSelection(MoveTemp(GeoSelection));
@@ -581,13 +809,9 @@ UDynamicMesh* UGeometryScriptLibrary_MeshSelectionFunctions::ConvertIndexListToM
 		return TargetMesh;
 	}
 
-	if (IndexList.IndexType == EGeometryScriptIndexType::Vertex
-		|| IndexList.IndexType == EGeometryScriptIndexType::Triangle
-		|| IndexList.IndexType == EGeometryScriptIndexType::PolygroupID )
+	EGeometryScriptMeshSelectionType InitialType;
+	if (FGeometryScriptMeshSelection::ConvertIndexTypeToSelectionType(IndexList.IndexType, InitialType))
 	{
-		EGeometryScriptMeshSelectionType InitialType =
-			(IndexList.IndexType == EGeometryScriptIndexType::Vertex) ? EGeometryScriptMeshSelectionType::Vertices :
-				((IndexList.IndexType == EGeometryScriptIndexType::Triangle) ? EGeometryScriptMeshSelectionType::Triangles : EGeometryScriptMeshSelectionType::Polygroups);
 		if (SelectionType == InitialType)
 		{
 			ConvertIndexArrayToMeshSelection(TargetMesh, *IndexList.List, InitialType, SelectionOut);
@@ -893,6 +1117,13 @@ UDynamicMesh* UGeometryScriptLibrary_MeshSelectionFunctions::ExpandMeshSelection
 		NewSelection.SetSelection(Selection);
 		return TargetMesh;
 	}
+	if (Selection.GetSelectionType() == EGeometryScriptMeshSelectionType::Edges)
+	{
+		// todo: implement this
+		UE_LOG(LogGeometry, Warning, TEXT("ExpandMeshSelectionToConnected: Edge Selection currently not supported"));
+		NewSelection.SetSelection(Selection);
+		return TargetMesh;
+	}
 
 	if (Selection.GetSelectionType() == EGeometryScriptMeshSelectionType::Polygroups &&
 		ConnectionType == EGeometryScriptTopologyConnectionType::Polygroup)
@@ -1029,6 +1260,31 @@ UDynamicMesh* UGeometryScriptLibrary_MeshSelectionFunctions::ExpandContractMeshS
 			}
 		});
 	}
+	else if (Selection.GetSelectionType() == EGeometryScriptMeshSelectionType::Edges)
+	{
+		NewGeoSelection.InitializeTypes(EGeometryElementType::Edge, EGeometryTopologyType::Triangle);
+		
+		TargetMesh->ProcessMesh([&Selection, &NewGeoSelection, bContract, Iterations](const FDynamicMesh3& ReadMesh)
+		{
+			FMeshEdgeSelection EdgeSelection(&ReadMesh);
+			Selection.ProcessByEdgeID(ReadMesh, [&EdgeSelection](int32 EdgeID) { EdgeSelection.Select(EdgeID); });
+			if (bContract)
+			{
+				EdgeSelection.ContractByBorderEdges(Iterations);
+			}
+			else
+			{
+				EdgeSelection.ExpandToOneRingNeighbors(Iterations);
+			}
+			for (int32 EdgeID : EdgeSelection)
+			{
+				ReadMesh.EnumerateTriEdgeIDsFromEdgeID(EdgeID, [&NewGeoSelection](FMeshTriEdgeID TriEdgeID)
+				{
+					NewGeoSelection.Selection.Add(TriEdgeID.Encoded());
+				});
+			}
+		});
+	}
 	else 
 	{
 		TargetMesh->ProcessMesh([&](const FDynamicMesh3& ReadMesh)
diff --git a/Engine/Plugins/Runtime/GeometryScripting/Source/GeometryScriptingCore/Private/MeshSelectionQueryFunctions.cpp b/Engine/Plugins/Runtime/GeometryScripting/Source/GeometryScriptingCore/Private/MeshSelectionQueryFunctions.cpp
--- a/Engine/Plugins/Runtime/GeometryScripting/Source/GeometryScriptingCore/Private/MeshSelectionQueryFunctions.cpp
+++ b/Engine/Plugins/Runtime/GeometryScripting/Source/GeometryScriptingCore/Private/MeshSelectionQueryFunctions.cpp
@@ -42,7 +42,16 @@ UDynamicMesh* UGeometryScriptLibrary_MeshSelectionQueryFunctions::GetMeshSelecti
 				Bounds.Contain(ReadMesh.GetVertex(VertexID));
 			});
 		}
-		else
+		else if (Selection.GetSelectionType() == EGeometryScriptMeshSelectionType::Edges)
+		{
+			Selection.ProcessByEdgeID(ReadMesh, [&Bounds, &ReadMesh](int32 EdgeID)
+			{
+				FIndex2i EdgeV = ReadMesh.GetEdgeV(EdgeID);
+				Bounds.Contain(ReadMesh.GetVertex(EdgeV.A));
+				Bounds.Contain(ReadMesh.GetVertex(EdgeV.B));
+			});
+		}
+		else // EGeometryScriptMeshSelectionType::Triangles or EGeometryScriptMeshSelectionType::Polygroups
 		{
 			Selection.ProcessByTriangleID(ReadMesh, [&](int32 TriangleID) {
 				Bounds.Contain(ReadMesh.GetTriBounds(TriangleID));
diff --git a/Engine/Plugins/Runtime/GeometryScripting/Source/GeometryScriptingCore/Private/MeshSpatialFunctions.cpp b/Engine/Plugins/Runtime/GeometryScripting/Source/GeometryScriptingCore/Private/MeshSpatialFunctions.cpp
--- a/Engine/Plugins/Runtime/GeometryScripting/Source/GeometryScriptingCore/Private/MeshSpatialFunctions.cpp
+++ b/Engine/Plugins/Runtime/GeometryScripting/Source/GeometryScriptingCore/Private/MeshSpatialFunctions.cpp
@@ -350,6 +350,31 @@ UDynamicMesh* UGeometryScriptLibrary_MeshSpatial::SelectMeshElementsInBoxWithBVH
 						};
 						QueryBVH.Spatial->DoTraversal(Traversal, QueryOptions);
 					}
+					else if (SelectionType == EGeometryScriptMeshSelectionType::Edges)
+					{
+						int32 UseMinNumEdgePoints = FMath::Clamp(MinNumTrianglePoints, 1, 2);
+						GeoSelection.InitializeTypes(EGeometryElementType::Edge, EGeometryTopologyType::Triangle);
+						FDynamicMeshAABBTree3::FTreeTraversal Traversal;
+						Traversal.NextBoxF = [QueryBounds](const FAxisAlignedBox3d& Box, int Depth) { return Box.Intersects(QueryBounds); };
+						Traversal.NextTriangleF = [&ReadMesh, &GeoSelection, QueryBounds, UseMinNumEdgePoints](int TriangleID)
+						{
+							FIndex3i Tri = ReadMesh.GetTriangle(TriangleID);
+							bool ContainsV[3]{};
+							for (int32 k = 0; k < 3; ++k)
+							{
+								ContainsV[k] = QueryBounds.Contains(ReadMesh.GetVertex(Tri[k]));
+							}
+							for (int32 Idx = 0, PrevIdx = 2; Idx < 3; PrevIdx = Idx++)
+							{
+								int32 ContainCount = (int32)ContainsV[Idx] + (int32)ContainsV[PrevIdx];
+								if (ContainCount >= UseMinNumEdgePoints)
+								{
+									GeoSelection.Selection.Add(FMeshTriEdgeID(TriangleID, PrevIdx).Encoded());
+								}
+							}
+						};
+						QueryBVH.Spatial->DoTraversal(Traversal, QueryOptions);
+					}
 					else if (SelectionType == EGeometryScriptMeshSelectionType::Polygroups)
 					{
 						TArray<int32> UniqueGroups;
diff --git a/Engine/Plugins/Runtime/GeometryScripting/Source/GeometryScriptingCore/Private/MeshVertexColorFunctions.cpp b/Engine/Plugins/Runtime/GeometryScripting/Source/GeometryScriptingCore/Private/MeshVertexColorFunctions.cpp
--- a/Engine/Plugins/Runtime/GeometryScripting/Source/GeometryScriptingCore/Private/MeshVertexColorFunctions.cpp
+++ b/Engine/Plugins/Runtime/GeometryScripting/Source/GeometryScriptingCore/Private/MeshVertexColorFunctions.cpp
@@ -162,7 +162,7 @@ UDynamicMesh* UGeometryScriptLibrary_MeshVertexColorFunctions::SetMeshSelectionV
 				Colors->SetTriangle(TriangleID, TriVerts);
 			}
 		}
-		else if (Selection.GetSelectionType() == EGeometryScriptMeshSelectionType::Vertices)
+		else if (Selection.GetSelectionType() == EGeometryScriptMeshSelectionType::Vertices || Selection.GetSelectionType() == EGeometryScriptMeshSelectionType::Edges)
 		{
 			Selection.ProcessByVertexID(EditMesh, [&](int32 VertexID)
 			{
diff --git a/Engine/Plugins/Runtime/GeometryScripting/Source/GeometryScriptingCore/Public/GeometryScript/GeometryScriptSelectionTypes.h b/Engine/Plugins/Runtime/GeometryScripting/Source/GeometryScriptingCore/Public/GeometryScript/GeometryScriptSelectionTypes.h
--- a/Engine/Plugins/Runtime/GeometryScripting/Source/GeometryScriptingCore/Public/GeometryScript/GeometryScriptSelectionTypes.h
+++ b/Engine/Plugins/Runtime/GeometryScripting/Source/GeometryScriptingCore/Public/GeometryScript/GeometryScriptSelectionTypes.h
@@ -8,6 +8,9 @@
 
 namespace UE::Geometry { struct FGeometrySelection; }
 
+//~ Note that Edge selections are represented by triangle, index-in-triangle pairs,
+//~ so each non-boundary edge can be represented in the selection twice (once per 'half edge').
+//~ Currently, our convention is to store both representations in the selection.
 /**
  * Type of index stored in a FGeometryScriptMeshSelection
  */
@@ -15,6 +18,7 @@ UENUM(BlueprintType)
 enum class EGeometryScriptMeshSelectionType : uint8
 {
 	Vertices = 0,
+	Edges = 3,
 	Triangles = 1,
 	Polygroups = 2 UMETA(DisplayName = "PolyGroups")
 };
@@ -22,8 +26,10 @@ enum class EGeometryScriptMeshSelectionType : uint8
 /**
  * Type of Conversion to apply to a FGeometryScriptMeshSelection
  */
+
 UENUM(BlueprintType)
-enum class EGeometryScriptMeshSelectionConversionType : uint8
+enum class UE_DEPRECATED(5.5, "This enum is unused, and may be removed in the future. To convert selection types, we use EGeometryScriptMeshSelectionType to specify the desired result type.")
+	EGeometryScriptMeshSelectionConversionType : uint8
 {
 	NoConversion = 0,
 	ToVertices = 1,
@@ -80,9 +86,35 @@ struct GEOMETRYSCRIPTINGCORE_API FGeometryScriptMeshSelection
 	bool IsEmpty() const;
 
 	EGeometryScriptMeshSelectionType GetSelectionType() const;
+
+	// Note that for edge selections, this can return more elements than expected because both can be redundantly represented
+	// i.e. an edge can be in the selection once per 'half edge', since it is represented by triangle/index-in-triangle pair.
+	// Call GetNumUniqueValidSelected(Mesh) to get the number of unique elements in the selection
 	int32 GetNumSelected() const;
+	// Return the number of valid, unique elements in the selection; e.g., with no double-counting of selected edges
+	int32 GetNumUniqueSelected(const UE::Geometry::FDynamicMesh3& Mesh) const;
 	void DebugPrint() const;
 
+	static bool ConvertIndexTypeToSelectionType(EGeometryScriptIndexType IndexType, EGeometryScriptMeshSelectionType& OutSelectionType)
+	{
+		switch (IndexType)
+		{
+		case EGeometryScriptIndexType::Triangle:
+			OutSelectionType = EGeometryScriptMeshSelectionType::Triangles;
+			return true;
+		case EGeometryScriptIndexType::Edge:
+			OutSelectionType = EGeometryScriptMeshSelectionType::Edges;
+			return true;
+		case EGeometryScriptIndexType::Vertex:
+			OutSelectionType = EGeometryScriptMeshSelectionType::Vertices;
+			return true;
+		case EGeometryScriptIndexType::PolygroupID:
+			OutSelectionType = EGeometryScriptMeshSelectionType::Polygroups;
+			return true;
+		}
+		return false;
+	}
+
 	/** 
 	 * Combine SelectionB with current selection, updating current selection, using CombineMode to control how combining happens
 	 */
@@ -114,6 +146,16 @@ struct GEOMETRYSCRIPTINGCORE_API FGeometryScriptMeshSelection
 		TFunctionRef<void(int32)> PerVertexFunc,
 		bool bProcessAllVertsIfSelectionEmpty = false) const;
 
+	/**
+	 * Call PerEdgeFunc for each EdgeID in the Selection.
+	 * For Vertex Selections, Vertex Edge one-rings are enumerated and accumulated in a TSet.
+	 * For Triangle Selections, Triangle Edges are enumerated and accumulated in a TSet.
+	 * For PolyGroup Selections, a full mesh iteration is used to find all Triangle Edges in the groups (accumulated in a TSet)
+	 */
+	void ProcessByEdgeID(const UE::Geometry::FDynamicMesh3& Mesh,
+		TFunctionRef<void(int32)> PerEdgeFunc,
+		bool bProcessAllVertsIfSelectionEmpty = false) const;
+
 
 	// Required by TStructOpsTypeTraits interface
 	bool operator==(const FGeometryScriptMeshSelection& Other) const
diff --git a/Engine/Plugins/Runtime/GeometryScripting/Source/GeometryScriptingCore/Public/GeometryScript/MeshSelectionFunctions.h b/Engine/Plugins/Runtime/GeometryScripting/Source/GeometryScriptingCore/Public/GeometryScript/MeshSelectionFunctions.h
--- a/Engine/Plugins/Runtime/GeometryScripting/Source/GeometryScriptingCore/Public/GeometryScript/MeshSelectionFunctions.h
+++ b/Engine/Plugins/Runtime/GeometryScripting/Source/GeometryScriptingCore/Public/GeometryScript/MeshSelectionFunctions.h
@@ -31,6 +31,8 @@ public:
 
 	/**
 	 * Query information about a Mesh Selection
+	 * Note that NumSelected may double count some polygroups and (non-border) edges due to their internal representation
+	 * Use GetMeshUniqueSelectionInfo for an accurate count of unique mesh elements
 	 */
 	UFUNCTION(BlueprintCallable, Category = "GeometryScript|MeshSelection", meta=(ScriptMethod))
 	static void
@@ -39,6 +41,17 @@ public:
 		EGeometryScriptMeshSelectionType& SelectionType,
 		int& NumSelected);
 
+	/**
+	 * Query information about a Mesh Selection, and get a count of unique selected elements
+	 */
+	UFUNCTION(BlueprintCallable, Category = "GeometryScript|MeshSelection", meta = (ScriptMethod))
+	static void
+	GetMeshUniqueSelectionInfo(
+		const UDynamicMesh* TargetMesh,
+		FGeometryScriptMeshSelection Selection,
+		EGeometryScriptMeshSelectionType& SelectionType,
+		int& NumSelected);
+
 	/**
 	 * Print information about the Mesh Selection to the Output Log
 	 * @param bDisable if true, printing will be disabled (useful for debugging)
@@ -86,9 +99,11 @@ public:
 	 * Convert a Mesh Selection to a different Type (eg Vertices to Triangles, etc)
 	 * By default, Vertices map to Triangle one-rings, and Triangles to all contained vertices.
 	 * If bAllowPartialInclusion is disabled, then more restrictive conversions are performed, as follows:
-	 *   For To-Vertices, only include vertices where all one-ring triangles are included in FromSelection.
-	 *   For To-Triangles, only include triangles where all tri vertices are included in FromSelection.
-	 *   For To-PolyGroups, only include groups where all group triangles are included in FromSelection
+	 *   For To-Vertices, only include vertices where all one-ring triangles or edges are included in FromSelection.
+	 *   For To-Edges, only include edges where all adjacent triangles or vertices are included in FromSelection
+	 *   For To-Triangles, only include triangles where all tri vertices or edges are included in FromSelection.
+	 *   For To-PolyGroups, only include groups where all group triangles are touched by FromSelection
+	 * (Note: The To-PolyGroups rule allows vertex and edge selections that miss some vertices or edges, as long as they touch all the polygroup triangles.)
 	 * @param bAllowPartialInclusion if false, perform more limited selection conversion as described above
 	 */
 	UFUNCTION(BlueprintCallable, Category = "GeometryScript|MeshSelection", meta=(ScriptMethod))
@@ -180,7 +195,7 @@ public:
 	/**
 	 * Create a new Mesh Selection of the SelectionType for the TargetMesh by finding all elements contained in the Box.
 	 * @param bInvert return a selection of all elements not in the Box
-	 * @param MinNumTrianglePoints number of vertices of a triangle that must be in the box for it to be selected (1,2, or 3)
+	 * @param MinElementVertices number of vertices of a triangle or edge that must be in the box for it to be selected (clamped to 1-3 for tri; 1-2 for edge)
 	 */
 	UFUNCTION(BlueprintCallable, Category = "GeometryScript|MeshSelection", meta=(ScriptMethod, AdvancedDisplay="MinNumTrianglePoints"))
 	static UPARAM(DisplayName = "Target Mesh") UDynamicMesh* 
@@ -190,14 +205,14 @@ public:
 		FBox Box,
 		EGeometryScriptMeshSelectionType SelectionType = EGeometryScriptMeshSelectionType::Triangles,
 		bool bInvert = false,
-		int MinNumTrianglePoints = 3);
+		UPARAM(DisplayName = "Min Element Vertices") int MinNumTrianglePoints = 3);
 
 	/**
 	* Create a new Mesh Selection of the SelectionType for the TargetMesh by finding all elements contained in the Sphere.
 	* @param SphereOrigin center point of the Sphere
 	* @param SphereRadius radius of the Sphere
 	* @param bInvert return a selection of all elements not in the Sphere
-	* @param MinNumTrianglePoints number of vertices of a triangle that must be in the Sphere for it to be selected (1,2, or 3)
+	* @param MinElementVertices number of vertices of a triangle or edge that must be in the Sphere for it to be selected (clamped to 1-3 for tri; 1-2 for edge)
 	*/
 	UFUNCTION(BlueprintCallable, Category = "GeometryScript|MeshSelection", meta=(ScriptMethod, AdvancedDisplay="MinNumTrianglePoints"))
 	static UPARAM(DisplayName = "Target Mesh") UDynamicMesh* 
@@ -208,14 +223,14 @@ public:
 		double SphereRadius = 100.0,
 		EGeometryScriptMeshSelectionType SelectionType = EGeometryScriptMeshSelectionType::Triangles,
 		bool bInvert = false,
-		int MinNumTrianglePoints = 3);
+		UPARAM(DisplayName = "Min Element Vertices") int MinNumTrianglePoints = 3);
 
 	/**
 	* Create a new Mesh Selection of the SelectionType for the TargetMesh by finding all elements on the "positive" side of a Plane
 	* @param PlaneOrigin center point of the Plane
 	* @param PlaneNormal normal vector for the Plane
 	* @param bInvert return a selection of all elements on the other (negative) side of the Plane
-	* @param MinNumTrianglePoints number of vertices of a triangle that must be on the positive Plane side to be selected (1,2, or 3)
+	* @param MinElementVertices number of vertices of a triangle or edge that must be on the positive Plane side to be selected (clamped to 1-3 for tri; 1-2 for edge)
 	*/
 	UFUNCTION(BlueprintCallable, Category = "GeometryScript|MeshSelection", meta=(ScriptMethod, AdvancedDisplay="MinNumTrianglePoints"))
 	static UPARAM(DisplayName = "Target Mesh") UDynamicMesh* 
@@ -226,7 +241,7 @@ public:
 		FVector PlaneNormal = FVector::UpVector,
 		EGeometryScriptMeshSelectionType SelectionType = EGeometryScriptMeshSelectionType::Triangles,
 		bool bInvert = false,
-		int MinNumTrianglePoints = 3);
+		UPARAM(DisplayName = "Min Element Vertices") int MinNumTrianglePoints = 3);
 
 	/**
 	* Create a new Mesh Selection of the SelectionType for the TargetMesh by finding all elements that have a normal
@@ -235,7 +250,7 @@ public:
 	* @param Normal normal/direction vector to measure against
 	* @param MaxAngleDeg maximum angular deviation from Normal, in degrees
 	* @param bInvert return a selection of all elements not within the given deviation
-	* @param MinNumTrianglePoints number of vertices of a triangle that must be within the angular deviation for it to be selected (1,2, or 3)
+	* @param MinElementVertices number of vertices of a triangle or edge that must be within the angular deviation for it to be selected  (clamped to 1-3 for tri; 1-2 for edge)
 	*/
 	UFUNCTION(BlueprintCallable, Category = "GeometryScript|MeshSelection", meta=(ScriptMethod, AdvancedDisplay="MinNumTrianglePoints"))
 	static UPARAM(DisplayName = "Target Mesh") UDynamicMesh* 
@@ -246,7 +261,7 @@ public:
 		double MaxAngleDeg = 1.0,
 		EGeometryScriptMeshSelectionType SelectionType = EGeometryScriptMeshSelectionType::Triangles,
 		bool bInvert = false,
-		int MinNumTrianglePoints = 3);
+		UPARAM(DisplayName = "Min Element Vertices") int MinNumTrianglePoints = 3);
 
 
 	/**
@@ -256,7 +271,7 @@ public:
 	* @param bInvert return a selection of all elements not within the given deviation
 	* @param ShellDistance If > 0, points within this distance from SelectionMesh will also be considered "inside"
 	* @param WindingThreshold Threshold used for Fast Mesh Winding Number inside/outside test (range is [0,1], with 1 being "inside")
-	* @param MinNumTrianglePoints number of vertices of a triangle that must be within the angular deviation for it to be selected (1,2, or 3)
+	* @param MinElementVertices number of vertices of a triangle or edge that must be within the selection mesh for it to be selected  (clamped to 1-3 for tri; 1-2 for edge)
 	*/
 	UFUNCTION(BlueprintCallable, Category = "GeometryScript|MeshSelection", meta=(ScriptMethod, AdvancedDisplay="MinNumTrianglePoints"))
 	static UPARAM(DisplayName = "Target Mesh") UDynamicMesh* 
@@ -269,7 +284,7 @@ public:
 		bool bInvert = false,
 		double ShellDistance = 0.0,
 		double WindingThreshold = 0.5,
-		int MinNumTrianglePoints = 3);
+		UPARAM(DisplayName = "Min Element Vertices") int MinNumTrianglePoints = 3);
 
 
 	/**
diff --git a/Engine/Plugins/Runtime/GeometryScripting/Source/GeometryScriptingCore/Public/GeometryScript/MeshSpatialFunctions.h b/Engine/Plugins/Runtime/GeometryScripting/Source/GeometryScriptingCore/Public/GeometryScript/MeshSpatialFunctions.h
--- a/Engine/Plugins/Runtime/GeometryScripting/Source/GeometryScriptingCore/Public/GeometryScript/MeshSpatialFunctions.h
+++ b/Engine/Plugins/Runtime/GeometryScripting/Source/GeometryScriptingCore/Public/GeometryScript/MeshSpatialFunctions.h
@@ -154,9 +154,11 @@ public:
 
 	/**
 	* Create Mesh Selection of mesh elements in TargetMesh contained by QueryBox, using QueryBVH to accellerate the computation.
-	* Triangles are selected if MinNumTrianglePoints or more vertices are inside the box. 
+	* Triangles and Edges are selected if Min Element Vertices (clamped to a 1-3 or 1-2 range, for triangles or edges respectively) or more are inside the box. 
 	* PolyGroups are selected if any of their triangles are inside the box
 	* 
+	* Note that this method cannot select mesh elements that cut through the query box without having any vertices in the query box.
+	* 
 	* @param QueryBVH is an acceleration structure previously built with TargetMesh.
 	* @param QueryPoint the point in the mesh's 3D local space.
 	* @param Options control the fast winding number threshold 
@@ -170,9 +172,7 @@ public:
 		FGeometryScriptSpatialQueryOptions Options,
 		FGeometryScriptMeshSelection& Selection,
 		EGeometryScriptMeshSelectionType SelectionType = EGeometryScriptMeshSelectionType::Vertices,
-		int MinNumTrianglePoints = 3,
+		UPARAM(DisplayName = "Min Element Vertices") int MinNumTrianglePoints = 3,
 		UGeometryScriptDebug* Debug = nullptr );
 
-
-
 };
\ No newline at end of file
diff --git a/Engine/Source/Runtime/GeometryCore/Public/Selections/MeshEdgeSelection.h b/Engine/Source/Runtime/GeometryCore/Public/Selections/MeshEdgeSelection.h
--- a/Engine/Source/Runtime/GeometryCore/Public/Selections/MeshEdgeSelection.h
+++ b/Engine/Source/Runtime/GeometryCore/Public/Selections/MeshEdgeSelection.h
@@ -154,6 +154,154 @@ public:
 		}
 	}
 
+private:
+	template<bool bHasFilter, typename ExpandContainerType>
+	void ExpandToOneRingNeighbors_FindNeighborsHelper(const ExpandContainerType& ToExpand, TArray<int32>& ToAdd, TFunctionRef<bool(int)> FilterF)
+	{
+		for (int EID : ToExpand)
+		{
+			FIndex2i EdgeV = Mesh->GetEdgeV(EID);
+			for (int32 SubIdx = 0; SubIdx < 2; ++SubIdx)
+			{
+				for (int32 NbrEID : Mesh->VtxEdgesItr(EdgeV[SubIdx]))
+				{
+					if constexpr (bHasFilter)
+					{
+						if (!FilterF(NbrEID))
+						{
+							continue;
+						}
+					}
+					if (!IsSelected(NbrEID))
+					{
+						ToAdd.Add(NbrEID);
+					}
+				}
+			}
+		}
+	}
+
+	template<bool bHasFilter>
+	void ExpandToOneRingNeighbors_Helper(int32 NumRings, TFunctionRef<bool(int)> FilterF)
+	{
+		if (NumRings <= 0)
+		{
+			return;
+		}
+
+		// ToAdd can have the same edge multiple times, so to avoid accumulating duplicates
+		// this will clean up the extra copies as it add to the selection
+		auto AddToSelectionAndRemoveRedundant = [this](TArray<int32>& ToAdd)
+		{
+			for (int32 Idx = 0; Idx < ToAdd.Num(); ++Idx)
+			{
+				bool bAlreadyInSet;
+				Selected.Add(ToAdd[Idx], &bAlreadyInSet);
+				if (bAlreadyInSet)
+				{
+					ToAdd.RemoveAtSwap(Idx, EAllowShrinking::No);
+					--Idx;
+				}
+			}
+		};
+
+		TArray<int32> ToAdd;
+		ExpandToOneRingNeighbors_FindNeighborsHelper<bHasFilter>(Selected, ToAdd, FilterF);
+		if (NumRings == 1)
+		{
+			// In the one ring case, no need to clean duplicates from ToAdd since we'll never use it again
+			// (note: could also do this on the final iteration for the NumRings > 1 case ...)
+			for (int32 ID : ToAdd)
+			{
+				add(ID);
+			}
+			return;
+		}
+		else
+		{
+			AddToSelectionAndRemoveRedundant(ToAdd);
+			TArray<int32> ToExpand;
+			for (int32 Iter = 1; Iter < NumRings; ++Iter)
+			{
+				Swap(ToAdd, ToExpand);
+				ToAdd.Reset();
+				ExpandToOneRingNeighbors_FindNeighborsHelper<bHasFilter>(ToExpand, ToAdd, FilterF);
+				AddToSelectionAndRemoveRedundant(ToAdd);
+			}
+		}
+	}
+
+public:
+
+	/**
+	 *  Add all one-ring neighbors of current selection to set.
+	 *  On a large mesh this is quite expensive as we don't know the boundary,
+	 *  so we have to iterate over all selected triangles.
+	 *
+	 *  Return false from FilterF to prevent vertices from being included.
+	 */
+	void ExpandToOneRingNeighbors(TFunctionRef<bool(int32)> FilterF)
+	{
+		ExpandToOneRingNeighbors_Helper<true>(1, FilterF);
+	}
+
+	void ExpandToOneRingNeighbors()
+	{
+		ExpandToOneRingNeighbors_Helper<false>(1, [](int32) {return true;});
+	}
+
+	void ExpandToOneRingNeighbors(int NumRings, TFunctionRef<bool(int32)> FilterF)
+	{
+		ExpandToOneRingNeighbors_Helper<true>(NumRings, FilterF);
+	}
+	
+	void ExpandToOneRingNeighbors(int NumRings)
+	{
+		ExpandToOneRingNeighbors_Helper<false>(NumRings, [](int32) {return true;});
+	}
+
+
+	/**
+	 * For each contraction, remove edges in current selection set that have 
+	 * any unselected edge neighboring either of the edge's vertices
+	 */
+	void ContractByBorderEdges(int32 nRings = 1)
+	{
+		// find set of boundary edges
+		TArray<int> BorderEdges;
+		for (int32 k = 0; k < nRings; ++k)
+		{
+			BorderEdges.Reset();
+
+			for (int EID : Selected)
+			{
+				FIndex2i EdgeV = Mesh->GetEdgeV(EID);
+				bool bEitherSideBoundary = false;
+				for (int32 SubIdx = 0; SubIdx < 2; ++SubIdx)
+				{
+					bool bIsBoundary = false;
+					for (int32 NbrEID : Mesh->VtxEdgesItr(EdgeV[SubIdx]))
+					{
+						if (NbrEID != EID && !IsSelected(NbrEID))
+						{
+							bIsBoundary = true;
+						}
+					}
+					if (bIsBoundary)
+					{
+						bEitherSideBoundary = true;
+						break;
+					}
+				}
+				if (bEitherSideBoundary)
+				{
+					BorderEdges.Add(EID);
+				}
+			}
+			Deselect(BorderEdges);
+		}
+	}
+
 
 	GEOMETRYCORE_API void SelectBoundaryTriEdges(const FMeshFaceSelection& Triangles);
 
