fix some geometry script method comments (ray box intersection and make transform from z axis)

diff --git a/Engine/Plugins/Runtime/GeometryScripting/Source/GeometryScriptingCore/Public/GeometryScript/ShapeFunctions.h b/Engine/Plugins/Runtime/GeometryScripting/Source/GeometryScriptingCore/Public/GeometryScript/ShapeFunctions.h
--- a/Engine/Plugins/Runtime/GeometryScripting/Source/GeometryScriptingCore/Public/GeometryScript/ShapeFunctions.h
+++ b/Engine/Plugins/Runtime/GeometryScripting/Source/GeometryScriptingCore/Public/GeometryScript/ShapeFunctions.h
@@ -17,9 +17,7 @@ class GEOMETRYSCRIPTINGCORE_API UGeometryScriptLibrary_TransformFunctions : publ
 public:
 
 	/**
-	 * Create a Transform at the given Location, with the ZAxis vector as the Z axis
-	 * of the Transform, and the X or Y axis oriented to the Tangent vector, based on
-	 * the bTangentIsX parameter.
+	 * Create a Transform at the given Location, with the ZAxis vector as the Z axis of the Transform.
 	 */
 	UFUNCTION(BlueprintPure, Category = "GeometryScript|Shapes|Transform")
 	static UPARAM(DisplayName="Transform") FTransform
@@ -140,7 +138,7 @@ public:
 	GetRaySphereIntersection(FRay Ray, FVector SphereCenter, double SphereRadius, double& Distance1, double& Distance2);
 
 	/**
-	 * Check if the Ray intersects a Sphere defined by the SphereCenter and SphereRadius.
+	 * Find the intersection of a Ray and a Box
 	 * @param HitDistance Distance along the ray (Ray Parameter) to first intersection point with the Box
 	 * @return true if the ray hits the box, and false otherwise
 	 */
