Geometry Script: Add a new node to copy bones from a skeleton asset to a mesh. Refactored CopyBonesFromMesh so that the guts of the algorithm can be used by both.

diff --git a/Engine/Plugins/Runtime/GeometryScripting/Source/GeometryScriptingCore/Private/MeshBoneWeightFunctions.cpp b/Engine/Plugins/Runtime/GeometryScripting/Source/GeometryScriptingCore/Private/MeshBoneWeightFunctions.cpp
--- a/Engine/Plugins/Runtime/GeometryScripting/Source/GeometryScriptingCore/Private/MeshBoneWeightFunctions.cpp
+++ b/Engine/Plugins/Runtime/GeometryScripting/Source/GeometryScriptingCore/Private/MeshBoneWeightFunctions.cpp
@@ -617,6 +617,166 @@ UDynamicMesh* UGeometryScriptLibrary_MeshBoneWeightFunctions::TransferBoneWeight
 }
 
 
+static bool
+GetBoneCopyHierarchyAndReindexMeshIfNeeded(
+	TConstArrayView<FName> SourceMeshBoneNames,
+	TConstArrayView<int32> SourceMeshBoneParentIndices,
+	FDynamicMesh3& TargetMesh,
+	TMap<FName, FName>& OutBoneHierarchy,
+	FGeometryScriptCopyBonesFromMeshOptions Options,
+	UGeometryScriptDebug* Debug)
+{
+	// We can only perform this operation if there are source bones present. 
+	check(!SourceMeshBoneNames.IsEmpty());
+	
+	OutBoneHierarchy.Reset();
+
+	// Check across all skin weight attributes which bones are bound to the mesh.
+	if (Options.BonesToCopyFromSource != EBonesToCopyFromSource::AllBones)
+	{
+		// This should have been verified by the caller.
+		check(SourceMeshBoneNames.Num() == SourceMeshBoneParentIndices.Num());
+		
+		if (!TargetMesh.Attributes()->HasBones())
+		{
+			AppendError(Debug, EGeometryScriptErrorType::OperationFailed, LOCTEXT("GetBoneCopyHierarchy_NoTargetBonesDefined", "Target mesh has no bone names defined which is needed when not copying all bones"));
+			return false;
+		}
+
+		// Construct the actual bone hierarchy as a map of bone -> parent from the bone name and bone parent index lists.
+		FName RootBone;
+		
+		for (int32 Index = 0; Index < SourceMeshBoneNames.Num(); Index++)
+		{
+			const int32 BoneParentIndex = SourceMeshBoneParentIndices[Index];
+			if (BoneParentIndex == INDEX_NONE || (BoneParentIndex >= 0 && BoneParentIndex < SourceMeshBoneNames.Num()))
+			{
+				OutBoneHierarchy.Add(SourceMeshBoneNames[Index], BoneParentIndex != INDEX_NONE ? SourceMeshBoneNames[BoneParentIndex] : NAME_None);
+
+				if (BoneParentIndex == INDEX_NONE)
+				{
+					if (!RootBone.IsNone())
+					{
+						AppendError(Debug, EGeometryScriptErrorType::OperationFailed, LOCTEXT("GetBoneCopyHierarchy_MultipleRootBonesFound", "Found multiple root bones on source mesh"));
+						return false;
+					}
+					RootBone = SourceMeshBoneNames[Index];
+				}
+			}
+		}
+
+		// Get the name of all bones that contribute to skin binding on all skin weight profiles.
+		TSet<int32> BoundBoneIndices;
+		for (const TPair<FName, TUniquePtr<FDynamicMeshVertexSkinWeightsAttribute>>& AttribPair : TargetMesh.Attributes()->GetSkinWeightsAttributes())
+		{
+			if (const FDynamicMeshVertexSkinWeightsAttribute* ToAttrib = TargetMesh.Attributes()->GetSkinWeightsAttribute(AttribPair.Key))
+			{
+				BoundBoneIndices.Append(ToAttrib->GetBoundBoneIndices());
+			}
+		}
+
+		const FDynamicMeshBoneNameAttribute* BoneNameTargetAttrib = TargetMesh.Attributes()->GetBoneNames();
+		TSet<FName> BoundBones;
+		for (int32 BoneIndex: BoundBoneIndices)
+		{
+			if (BoneIndex < 0 || BoneIndex >= BoneNameTargetAttrib->Num())
+			{
+				AppendError(Debug, EGeometryScriptErrorType::OperationFailed, LOCTEXT("GetBoneCopyHierarchy_InvalidBoneWeightIndex", "Invalid bone index found on mesh"));
+				return false;
+			}
+
+			const FName BoneName = BoneNameTargetAttrib->GetValue(BoneIndex);
+			if (BoneName.IsNone())
+			{
+				const FText Error = FText::Format(LOCTEXT("GetBoneCopyHierarchy_NoBoneName", "Target bone at index {0} has no name."), FText::AsNumber(BoneIndex));
+				AppendError(Debug, EGeometryScriptErrorType::OperationFailed, Error);
+				return false;
+			}
+			if (!OutBoneHierarchy.Contains(BoneName))
+			{
+				const FText Error = FText::Format(LOCTEXT("GetBoneCopyHierarchy_BoneNotFound", "Target bone '{0}' not found on source mesh."), FText::FromName(BoneName));
+				AppendError(Debug, EGeometryScriptErrorType::OperationFailed, Error);
+				return false;
+			}
+			BoundBones.Add(BoneName);
+		}
+
+		// Go from each bound bone, and add parent bones to root, optionally skipping over unbound bones, if using OnlyBoundAndRoot.
+		TSet<FName> UsedBones;
+		UsedBones.Add(RootBone);
+		
+		for (FName BoneName: BoundBones)
+		{
+			UsedBones.Add(BoneName);
+
+			// Traverse up to the root bone. If OnlyBoundAndRoot is set, then each bone's parent is set to either a bone further up
+			// in the hierarchy that's actually bound, or the root.
+			FName ParentBone = OutBoneHierarchy[BoneName];
+			while (ParentBone != RootBone)
+			{
+				if (Options.BonesToCopyFromSource == EBonesToCopyFromSource::OnlyBoundAndParents)
+				{
+					UsedBones.Add(ParentBone);
+				}
+				else if (Options.BonesToCopyFromSource == EBonesToCopyFromSource::OnlyBoundAndRoot && BoundBones.Contains(ParentBone))
+				{
+					// We found another bound bone, make this our new parent.
+					break;
+				}
+				
+				ParentBone = OutBoneHierarchy[ParentBone];
+			}
+
+			if (Options.BonesToCopyFromSource == EBonesToCopyFromSource::OnlyBoundAndRoot)
+			{
+				OutBoneHierarchy[BoneName] = ParentBone;
+			}
+		}
+
+		// Leave only used bones in the hierarchy.
+		OutBoneHierarchy = OutBoneHierarchy.FilterByPredicate([&UsedBones](const TPair<FName, FName>& Item)
+		{
+			return UsedBones.Contains(Item.Key);
+		});
+	}
+
+	if (Options.ReindexWeights)
+	{
+		if (TargetMesh.Attributes()->HasBones())
+		{
+			TArray<FName> ToBones(SourceMeshBoneNames);
+			if (!OutBoneHierarchy.IsEmpty())
+			{
+				// Remove all bones that are not in the hierarchy.
+				ToBones.RemoveAll([&OutBoneHierarchy](FName BoneName)
+				{
+					return !OutBoneHierarchy.Contains(BoneName);
+				});
+			}
+			
+			for (const TPair<FName, TUniquePtr<FDynamicMeshVertexSkinWeightsAttribute>>& AttribPair : TargetMesh.Attributes()->GetSkinWeightsAttributes())
+			{
+				if (FDynamicMeshVertexSkinWeightsAttribute* ToAttrib = TargetMesh.Attributes()->GetSkinWeightsAttribute(AttribPair.Key))
+				{
+					if (!ToAttrib->ReindexBoneIndicesToSkeleton(TargetMesh.Attributes()->GetBoneNames()->GetAttribValues(), ToBones))
+					{
+						const FText Error = FText::Format(LOCTEXT("GetBoneCopyHierarchy_FailedToReindexWeights", "Failed to reindex bone weights for {0} weights profile"), FText::FromName(AttribPair.Key));
+						AppendError(Debug, EGeometryScriptErrorType::OperationFailed, Error);
+						return false;
+					}
+				}
+			}
+		}
+		else
+		{
+			AppendWarning(Debug, EGeometryScriptErrorType::InvalidInputs, LOCTEXT("GetBoneCopyHierarchy_TargetMeshHasNoBones", "Bone weight re-indexing was requested but the target mesh has no skeleton data"));
+		}
+	}
+
+	return true;
+}
+	
+
 UDynamicMesh* UGeometryScriptLibrary_MeshBoneWeightFunctions::CopyBonesFromMesh(
 	UDynamicMesh* SourceMesh, 
 	UDynamicMesh* TargetMesh, 
@@ -656,8 +816,9 @@ UDynamicMesh* UGeometryScriptLibrary_MeshBoneWeightFunctions::CopyBonesFromMesh(
 				EditMesh.EnableAttributes();
 			}
 
-			TMap<FName, FName> BoneHierarchy;
-
+			TConstArrayView<FName> SourceMeshBoneNames = ReadMesh.Attributes()->GetBoneNames()->GetAttribValues();
+			TConstArrayView<int32> SourceMeshParentBoneIndices;
+			
 			// Check across all skin weight attributes which bones are bound to the mesh.
 			if (Options.BonesToCopyFromSource != EBonesToCopyFromSource::AllBones)
 			{
@@ -667,157 +828,130 @@ UDynamicMesh* UGeometryScriptLibrary_MeshBoneWeightFunctions::CopyBonesFromMesh(
 				{
 					return;
 				}
-				
+		
 				// Target bone name attribute's existence has already been verified above with a call to HasBones, the source has been verified through
 				// ValidateBoneHierarchy.
-				const FDynamicMeshBoneNameAttribute* BoneNameSourceAttrib = ReadMesh.Attributes()->GetBoneNames();
-				const FDynamicMeshBoneParentIndexAttribute* BoneParentIndexSourceAttrib = ReadMesh.Attributes()->GetBoneParentIndices();
-				const FDynamicMeshBoneNameAttribute* BoneNameTargetAttrib = EditMesh.Attributes()->GetBoneNames();
-			
-				FName RootBone;
-				
-				for (int32 Index = 0; Index < BoneNameSourceAttrib->Num(); Index++)
+				SourceMeshParentBoneIndices = ReadMesh.Attributes()->GetBoneParentIndices()->GetAttribValues();
+			}
+
+			TMap<FName, FName> BoneHierarchy;
+			if (GetBoneCopyHierarchyAndReindexMeshIfNeeded(SourceMeshBoneNames, SourceMeshParentBoneIndices, EditMesh, BoneHierarchy, Options, Debug))
+			{
+				// If the bone hierarchy wasn't set up, then do a full copy.
+				if (BoneHierarchy.IsEmpty())
 				{
-					const int32 BoneParentIndex = BoneParentIndexSourceAttrib->GetValue(Index);
-					if (BoneParentIndex == INDEX_NONE || (BoneParentIndex >= 0 && BoneParentIndex < BoneNameSourceAttrib->Num()))
-					{
-						BoneHierarchy.Add(BoneNameSourceAttrib->GetValue(Index), BoneParentIndex != INDEX_NONE ? BoneNameSourceAttrib->GetValue(BoneParentIndex) : NAME_None);
-
-						if (BoneParentIndex == INDEX_NONE)
-						{
-							if (!RootBone.IsNone())
-							{
-								AppendError(Debug, EGeometryScriptErrorType::OperationFailed, LOCTEXT("CopyBonesFromMesh_MultipleRootBonesFound", "Found multiple root bones on source mesh"));
-								return;
-							}
-							RootBone = BoneNameSourceAttrib->GetValue(Index);
-						}
-					}
+					EditMesh.Attributes()->CopyBoneAttributes(*ReadMesh.Attributes());
 				}
-
-				TSet<int32> BoundBoneIndices;
-
-				for (const TPair<FName, TUniquePtr<FDynamicMeshVertexSkinWeightsAttribute>>& AttribPair : EditMesh.Attributes()->GetSkinWeightsAttributes())
+				else
 				{
-					if (const FDynamicMeshVertexSkinWeightsAttribute* ToAttrib = EditMesh.Attributes()->GetSkinWeightsAttribute(AttribPair.Key))
-					{
-						BoundBoneIndices.Append(ToAttrib->GetBoundBoneIndices());
-					}
+					// Copy the bone attributes but only copy the ones in the remapping map and update the parent index as well.
+					EditMesh.Attributes()->CopyBoneAttributesWithRemapping(*ReadMesh.Attributes(), BoneHierarchy);
 				}
+			}
 
-				TSet<FName> BoundBones;
-				for (int32 BoneIndex: BoundBoneIndices)
-				{
-					if (BoneIndex < 0 || BoneIndex >= BoneNameTargetAttrib->Num())
-					{
-						AppendError(Debug, EGeometryScriptErrorType::OperationFailed, LOCTEXT("CopyBonesFromMesh_InvalidBoneWeightIndex", "Invalid bone index found on mesh"));
-						return;
-					}
+		}, EDynamicMeshChangeType::AttributeEdit, EDynamicMeshAttributeChangeFlags::Unknown, false);
+	});
 
-					const FName BoneName = BoneNameTargetAttrib->GetValue(BoneIndex);
-					if (BoneName.IsNone())
-					{
-						const FText Error = FText::Format(LOCTEXT("CopyBonesFromMesh_NoBoneName", "Target bone at index {0} has no name."), FText::AsNumber(BoneIndex));
-						AppendError(Debug, EGeometryScriptErrorType::OperationFailed, Error);
-						return;
-					}
-					if (!BoneHierarchy.Contains(BoneName))
-					{
-						const FText Error = FText::Format(LOCTEXT("CopyBonesFromMesh_BoneNotFound", "Target bone '{0}' not found on source mesh."), FText::FromName(BoneName));
-						AppendError(Debug, EGeometryScriptErrorType::OperationFailed, Error);
-						return;
-					}
-					BoundBones.Add(BoneName);
-				}
+	return TargetMesh;
+}
 
-				// Go from each bound bone, and add parent bones to root, optionally skipping over unbound bones, if using OnlyBoundAndRoot.
-				TSet<FName> UsedBones;
-				UsedBones.Add(RootBone);
-				
-				for (FName BoneName: BoundBones)
-				{
-					UsedBones.Add(BoneName);
 
-					// Traverse up to the root bone. If OnlyBoundAndRoot is set, then each bone's parent is set to either a bone further up
-					// in the hierarchy that's actually bound, or the root.
-					FName ParentBone = BoneHierarchy[BoneName];
-					while (ParentBone != RootBone)
-					{
-						if (Options.BonesToCopyFromSource == EBonesToCopyFromSource::OnlyBoundAndParents)
-						{
-							UsedBones.Add(ParentBone);
-						}
-						else if (Options.BonesToCopyFromSource == EBonesToCopyFromSource::OnlyBoundAndRoot && BoundBones.Contains(ParentBone))
-						{
-							// We found another bound bone, make this our new parent.
-							break;
-						}
-						
-						ParentBone = BoneHierarchy[ParentBone];
-					}
+UDynamicMesh* UGeometryScriptLibrary_MeshBoneWeightFunctions::CopyBonesFromSkeleton(
+	USkeleton* SourceSkeleton,
+	UDynamicMesh* TargetMesh,
+	FGeometryScriptCopyBonesFromMeshOptions Options,
+	UGeometryScriptDebug* Debug
+	)
+{
+	if (SourceSkeleton == nullptr)
+	{
+		AppendError(Debug, EGeometryScriptErrorType::InvalidInputs, LOCTEXT("CopyBonesFromSkeleton_InvalidSourceSkeleton", "CopyBonesFromMesh: SourceSkeleton is Null"));
+		return TargetMesh;
+	}
 
-					if (Options.BonesToCopyFromSource == EBonesToCopyFromSource::OnlyBoundAndRoot)
-					{
-						BoneHierarchy[BoneName] = ParentBone;
-					}
-				}
+	if (TargetMesh == nullptr)
+	{
+		AppendError(Debug, EGeometryScriptErrorType::InvalidInputs, LOCTEXT("CopyBonesFromSkeleton_InvalidTargetMesh", "CopyBonesFromMesh: TargetMesh is Null"));
+		return TargetMesh;
+	}
 
-				// Leave only used bones in the hierarchy.
-				BoneHierarchy = BoneHierarchy.FilterByPredicate([&UsedBones](const TPair<FName, FName>& Item)
-				{
-					return UsedBones.Contains(Item.Key);
-				});
-			}
+	TargetMesh->EditMesh([&](FDynamicMesh3& EditMesh)
+	{
+		if (!EditMesh.HasAttributes())
+		{
+			EditMesh.EnableAttributes();
+		}
+
+		TMap<FName, FName> BoneHierarchy;
+
+		const FReferenceSkeleton& RefSkeleton = SourceSkeleton->GetReferenceSkeleton();
+		const int32 NumSourceBones = RefSkeleton.GetRawBoneNum();
+		TArray<FName> SourceSkeletonBoneNames;
+		TArray<int32> SourceSkeletonParentBoneIndices;
+		SourceSkeletonBoneNames.Reserve(NumSourceBones);
+		SourceSkeletonParentBoneIndices.Reserve(NumSourceBones);
+		
+		for (const FMeshBoneInfo& BoneInfo: SourceSkeleton->GetReferenceSkeleton().GetRawRefBoneInfo())
+		{
+			SourceSkeletonBoneNames.Add(BoneInfo.Name);
+			SourceSkeletonParentBoneIndices.Add(BoneInfo.ParentIndex);
+		}
+		
+		if (GetBoneCopyHierarchyAndReindexMeshIfNeeded(SourceSkeletonBoneNames, SourceSkeletonParentBoneIndices, EditMesh, BoneHierarchy, Options, Debug))
+		{
+			EditMesh.Attributes()->EnableBones(BoneHierarchy.IsEmpty() ? NumSourceBones : BoneHierarchy.Num());
 
-			if (Options.ReindexWeights)
+			FDynamicMeshBoneNameAttribute* BoneNameAttrib = EditMesh.Attributes()->GetBoneNames();
+			FDynamicMeshBoneParentIndexAttribute* BoneParentIndexAttrib = EditMesh.Attributes()->GetBoneParentIndices();
+			FDynamicMeshBonePoseAttribute* BonePoses = EditMesh.Attributes()->GetBonePoses();
+			
+			const TArray<FTransform>& SourceSkeletonBonePoses = RefSkeleton.GetRawRefBonePose();
+
+			TMap<FName, int32> NameToIndexMap;
+
+			int32 NumTargetBones = 0;
+			for (int32 BoneIndex = 0; BoneIndex < NumSourceBones; BoneIndex++)
 			{
-				if (EditMesh.Attributes()->HasBones())
+				const FName BoneName = SourceSkeletonBoneNames[BoneIndex];
+				if (BoneHierarchy.IsEmpty() || BoneHierarchy.Contains(BoneName))
 				{
-					TArray<FName> TargetBones(ReadMesh.Attributes()->GetBoneNames()->GetAttribValues());
-					if (!BoneHierarchy.IsEmpty())
-					{
-						// Remove all bones that are not in the hierarchy.
-						TargetBones.RemoveAll([&BoneHierarchy](FName BoneName)
-						{
-							return !BoneHierarchy.Contains(BoneName);
-						});
-					}
-					
-					for (const TPair<FName, TUniquePtr<FDynamicMeshVertexSkinWeightsAttribute>>& AttribPair : EditMesh.Attributes()->GetSkinWeightsAttributes())
-					{
-						if (FDynamicMeshVertexSkinWeightsAttribute* ToAttrib = EditMesh.Attributes()->GetSkinWeightsAttribute(AttribPair.Key))
-						{
-							if (!ToAttrib->ReindexBoneIndicesToSkeleton(EditMesh.Attributes()->GetBoneNames()->GetAttribValues(), TargetBones))
-							{
-								const FText Error = FText::Format(LOCTEXT("CopyBonesFromMesh_FailedToReindexWeights", "Failed to reindex bone weights for {0} weights profile"), FText::FromName(AttribPair.Key));
-								AppendError(Debug, EGeometryScriptErrorType::OperationFailed, Error);
-							}
-						}
-					}
-				}
-				else
-				{
-					AppendWarning(Debug, EGeometryScriptErrorType::InvalidInputs, LOCTEXT("CopyBonesFromMesh_TargetMeshHasNoBones", "Bone weight re-indexing was requested but the target mesh has no skeleton data"));
+					NameToIndexMap.Add(BoneName, NumTargetBones);
+					BoneNameAttrib->SetValue(NumTargetBones, BoneName);
+					BonePoses->SetValue(NumTargetBones, SourceSkeletonBonePoses[BoneIndex]);
+					NumTargetBones++;
 				}
 			}
 
-			// If the bone hierarchy wasn't set up, then do a full copy.
+			// Now that we have all bone names and their new indices, remap the parent indices to them.
+			// The target list should already contain all parents listed in the BoneHierarchy map.
 			if (BoneHierarchy.IsEmpty())
 			{
-				EditMesh.Attributes()->CopyBoneAttributes(*ReadMesh.Attributes());
+				for (int32 BoneIndex = 0; BoneIndex < NumSourceBones; BoneIndex++)
+				{
+					BoneParentIndexAttrib->SetValue(BoneIndex, SourceSkeletonParentBoneIndices[BoneIndex]);
+				}
 			}
 			else
 			{
-				// Copy the bone attributes but only copy the ones in the remapping map and update the parent index as well.
-				EditMesh.Attributes()->CopyBoneAttributesWithRemapping(*ReadMesh.Attributes(), BoneHierarchy);
+				// Add a marker for the root bone, so we don't need to special case it below.
+				NameToIndexMap.Add(NAME_None, INDEX_NONE);
+				
+				for (const TPair<FName, FName>& BoneItem: BoneHierarchy)
+				{
+					const int32 BoneIndex = NameToIndexMap[BoneItem.Key];
+					const int32 BoneParentIndex = NameToIndexMap[BoneItem.Value];
+
+					BoneParentIndexAttrib->SetValue(BoneIndex, BoneParentIndex);
+				}
 			}
+		}
 
-		}, EDynamicMeshChangeType::AttributeEdit, EDynamicMeshAttributeChangeFlags::Unknown, false);
-	});
+	}, EDynamicMeshChangeType::AttributeEdit, EDynamicMeshAttributeChangeFlags::Unknown, false);
 
 	return TargetMesh;
 }
 
+
 UDynamicMesh* UGeometryScriptLibrary_MeshBoneWeightFunctions::DiscardBonesFromMesh(
 	UDynamicMesh* TargetMesh, 
 	UGeometryScriptDebug* Debug)
diff --git a/Engine/Plugins/Runtime/GeometryScripting/Source/GeometryScriptingCore/Public/GeometryScript/MeshBoneWeightFunctions.h b/Engine/Plugins/Runtime/GeometryScripting/Source/GeometryScriptingCore/Public/GeometryScript/MeshBoneWeightFunctions.h
--- a/Engine/Plugins/Runtime/GeometryScripting/Source/GeometryScriptingCore/Public/GeometryScript/MeshBoneWeightFunctions.h
+++ b/Engine/Plugins/Runtime/GeometryScripting/Source/GeometryScriptingCore/Public/GeometryScript/MeshBoneWeightFunctions.h
@@ -10,6 +10,7 @@
 #include "MeshBoneWeightFunctions.generated.h"
 
 class UDynamicMesh;
+class USkeleton;
 
 
 USTRUCT(BlueprintType, meta = (DisplayName = "Bone Weights"))
@@ -412,6 +413,7 @@ public:
 	 * Copy the bone attributes (skeleton) from the SourceMesh to the TargetMesh.
 	 * @param SourceMesh Mesh we are copying the bone attributes from.
 	 * @param TargetMesh Mesh we are copying the bone attributes to.
+	 * @param Options An option object to control how the copying is performed.
 	 */
 	UFUNCTION(BlueprintCallable, Category = "GeometryScript|MeshQueries|BoneWeights", meta=(ScriptMethod))
 	static UPARAM(DisplayName = "Target Mesh") UDynamicMesh* 
@@ -421,6 +423,20 @@ public:
 		FGeometryScriptCopyBonesFromMeshOptions Options = FGeometryScriptCopyBonesFromMeshOptions(),
 		UGeometryScriptDebug* Debug = nullptr);
 
+	/**
+	 * Copy the bone attributes (skeleton) from the SourceSkeleton to the TargetMesh.
+	 * @param SourceSkeleton The skeleton asset we are copying the bone attributes from.
+	 * @param TargetMesh Mesh we are copying the bone attributes to.
+	 * @param Options An option object to control how the copying is performed.
+	 */
+	UFUNCTION(BlueprintCallable, Category = "GeometryScript|MeshQueries|BoneWeights", meta=(ScriptMethod))
+	static UPARAM(DisplayName = "Target Mesh") UDynamicMesh* 
+	CopyBonesFromSkeleton(
+		USkeleton* SourceSkeleton,
+		UDynamicMesh* TargetMesh,
+		FGeometryScriptCopyBonesFromMeshOptions Options = FGeometryScriptCopyBonesFromMeshOptions(),
+		UGeometryScriptDebug* Debug = nullptr);
+
 	/**
 	 * Discard the bone attributes (skeleton) from the TargetMesh.
 	 * @param TargetMesh Mesh we are discarding bone attributes from.
