ModelingTools: Exposed NormalSpace for BentNormal bakes in BakeTexture, BakeVertex and GeometryScript Bake.

diff --git a/Engine/Plugins/Experimental/MeshModelingToolsetExp/Source/MeshModelingToolsExp/Private/BakeMeshAttributeMapsTool.cpp b/Engine/Plugins/Experimental/MeshModelingToolsetExp/Source/MeshModelingToolsExp/Private/BakeMeshAttributeMapsTool.cpp
--- a/Engine/Plugins/Experimental/MeshModelingToolsetExp/Source/MeshModelingToolsExp/Private/BakeMeshAttributeMapsTool.cpp
+++ b/Engine/Plugins/Experimental/MeshModelingToolsetExp/Source/MeshModelingToolsExp/Private/BakeMeshAttributeMapsTool.cpp
@@ -170,6 +170,15 @@ public:
 				Eval->MaxDistance = OcclusionSettings.MaxDistance;
 				Eval->SpreadAngle = OcclusionSettings.SpreadAngle;
 				Eval->BiasAngleDeg = OcclusionSettings.BiasAngle;
+				switch (OcclusionSettings.NormalSpace)
+				{
+				case EBakeNormalSpace::Tangent:
+					Eval->NormalSpace = FMeshOcclusionMapEvaluator::ESpace::Tangent;
+					break;
+				case EBakeNormalSpace::Object:
+					Eval->NormalSpace = FMeshOcclusionMapEvaluator::ESpace::Object;
+					break;
+				}
 				Baker->AddEvaluator(Eval);
 			}
 		};
@@ -369,6 +378,7 @@ void UBakeMeshAttributeMapsTool::Setup()
 	OcclusionSettings->WatchProperty(OcclusionSettings->MaxDistance, [this](float) { OpState |= EBakeOpState::Evaluate; });
 	OcclusionSettings->WatchProperty(OcclusionSettings->SpreadAngle, [this](float) { OpState |= EBakeOpState::Evaluate; });
 	OcclusionSettings->WatchProperty(OcclusionSettings->BiasAngle, [this](float) { OpState |= EBakeOpState::Evaluate; });
+	OcclusionSettings->WatchProperty(OcclusionSettings->NormalSpace, [this](EBakeNormalSpace) { OpState |= EBakeOpState::Evaluate; });
 
 
 	CurvatureSettings = NewObject<UBakeCurvatureMapToolProperties>(this);
diff --git a/Engine/Plugins/Experimental/MeshModelingToolsetExp/Source/MeshModelingToolsExp/Private/BakeMeshAttributeTool.cpp b/Engine/Plugins/Experimental/MeshModelingToolsetExp/Source/MeshModelingToolsExp/Private/BakeMeshAttributeTool.cpp
--- a/Engine/Plugins/Experimental/MeshModelingToolsetExp/Source/MeshModelingToolsExp/Private/BakeMeshAttributeTool.cpp
+++ b/Engine/Plugins/Experimental/MeshModelingToolsetExp/Source/MeshModelingToolsExp/Private/BakeMeshAttributeTool.cpp
@@ -79,6 +79,7 @@ EBakeOpState UBakeMeshAttributeTool::UpdateResult_Occlusion(const FImageDimensio
 	OcclusionMapSettings.OcclusionRays = OcclusionSettings->OcclusionRays;
 	OcclusionMapSettings.SpreadAngle = OcclusionSettings->SpreadAngle;
 	OcclusionMapSettings.BiasAngle = OcclusionSettings->BiasAngle;
+	OcclusionMapSettings.NormalSpace = OcclusionSettings->NormalSpace;
 
 	if ( !(CachedOcclusionMapSettings == OcclusionMapSettings) )
 	{
diff --git a/Engine/Plugins/Experimental/MeshModelingToolsetExp/Source/MeshModelingToolsExp/Private/BakeMeshAttributeVertexTool.cpp b/Engine/Plugins/Experimental/MeshModelingToolsetExp/Source/MeshModelingToolsExp/Private/BakeMeshAttributeVertexTool.cpp
--- a/Engine/Plugins/Experimental/MeshModelingToolsetExp/Source/MeshModelingToolsExp/Private/BakeMeshAttributeVertexTool.cpp
+++ b/Engine/Plugins/Experimental/MeshModelingToolsetExp/Source/MeshModelingToolsExp/Private/BakeMeshAttributeVertexTool.cpp
@@ -103,6 +103,15 @@ public:
 			OcclusionEval->MaxDistance = OcclusionSettings.MaxDistance;
 			OcclusionEval->SpreadAngle = OcclusionSettings.SpreadAngle;
 			OcclusionEval->BiasAngleDeg = OcclusionSettings.BiasAngle;
+			switch (OcclusionSettings.NormalSpace)
+			{
+			case EBakeNormalSpace::Tangent:
+				OcclusionEval->NormalSpace = FMeshOcclusionMapEvaluator::ESpace::Tangent;
+				break;
+			case EBakeNormalSpace::Object:
+				OcclusionEval->NormalSpace = FMeshOcclusionMapEvaluator::ESpace::Object;
+				break;
+			}
 		};
 
 		auto InitCurvatureEvaluator = [this] (FMeshCurvatureMapEvaluator* CurvatureEval)
@@ -337,6 +346,7 @@ void UBakeMeshAttributeVertexTool::Setup()
 	OcclusionSettings->WatchProperty(OcclusionSettings->MaxDistance, [this](float) { OpState |= EBakeOpState::Evaluate; });
 	OcclusionSettings->WatchProperty(OcclusionSettings->SpreadAngle, [this](float) { OpState |= EBakeOpState::Evaluate; });
 	OcclusionSettings->WatchProperty(OcclusionSettings->BiasAngle, [this](float) { OpState |= EBakeOpState::Evaluate; });
+	OcclusionSettings->WatchProperty(OcclusionSettings->NormalSpace, [this](EBakeNormalSpace) { OpState |= EBakeOpState::Evaluate; });
 
 	CurvatureSettings = NewObject<UBakeCurvatureMapToolProperties>(this);
 	CurvatureSettings->RestoreProperties(this);
diff --git a/Engine/Plugins/Experimental/MeshModelingToolsetExp/Source/MeshModelingToolsExp/Public/BakeMeshAttributeToolCommon.h b/Engine/Plugins/Experimental/MeshModelingToolsetExp/Source/MeshModelingToolsExp/Public/BakeMeshAttributeToolCommon.h
--- a/Engine/Plugins/Experimental/MeshModelingToolsetExp/Source/MeshModelingToolsExp/Public/BakeMeshAttributeToolCommon.h
+++ b/Engine/Plugins/Experimental/MeshModelingToolsetExp/Source/MeshModelingToolsExp/Public/BakeMeshAttributeToolCommon.h
@@ -162,6 +162,10 @@ public:
 	/** Angle in degrees from the horizon for occlusion rays for which the contribution is attenuated to reduce faceting artifacts. */
 	UPROPERTY(EditAnywhere, Category = OcclusionOutput, meta = (UIMin = "0", UIMax = "45.0", ClampMin = "0", ClampMax = "89.9"))
 	float BiasAngle = 15.0f;
+
+	/** Normal space for Bent Normal bakes. */
+	UPROPERTY(EditAnywhere, Category = OcclusionOutput)
+	EBakeNormalSpace NormalSpace = EBakeNormalSpace::Tangent;
 };
 
 
@@ -340,6 +344,7 @@ struct FOcclusionMapSettings
 	float MaxDistance;
 	float SpreadAngle;
 	float BiasAngle;
+	EBakeNormalSpace NormalSpace = EBakeNormalSpace::Tangent;
 
 	bool operator==(const FOcclusionMapSettings& Other) const
 	{
@@ -347,7 +352,8 @@ struct FOcclusionMapSettings
 			OcclusionRays == Other.OcclusionRays &&
 			MaxDistance == Other.MaxDistance &&
 			SpreadAngle == Other.SpreadAngle &&
-			BiasAngle == Other.BiasAngle;
+			BiasAngle == Other.BiasAngle &&
+			NormalSpace == Other.NormalSpace;
 	}
 };
 
diff --git a/Engine/Plugins/Runtime/GeometryScripting/Source/GeometryScriptingCore/Private/MeshBakeFunctions.cpp b/Engine/Plugins/Runtime/GeometryScripting/Source/GeometryScriptingCore/Private/MeshBakeFunctions.cpp
--- a/Engine/Plugins/Runtime/GeometryScripting/Source/GeometryScriptingCore/Private/MeshBakeFunctions.cpp
+++ b/Engine/Plugins/Runtime/GeometryScripting/Source/GeometryScriptingCore/Private/MeshBakeFunctions.cpp
@@ -378,6 +378,15 @@ namespace GeometryScriptBakeLocals
 			OcclusionEval->NumOcclusionRays = OcclusionOptions->OcclusionRays;
 			OcclusionEval->MaxDistance = (OcclusionOptions->MaxDistance == 0) ? TNumericLimits<float>::Max() : OcclusionOptions->MaxDistance;
 			OcclusionEval->SpreadAngle = OcclusionOptions->SpreadAngle;
+			switch (OcclusionOptions->NormalSpace)
+			{
+			case EGeometryScriptBakeNormalSpace::Tangent:
+				OcclusionEval->NormalSpace = FMeshOcclusionMapEvaluator::ESpace::Tangent;
+				break;
+			case EGeometryScriptBakeNormalSpace::Object:
+				OcclusionEval->NormalSpace = FMeshOcclusionMapEvaluator::ESpace::Object;
+				break;
+			}			
 			if (!GetTargetMeshTangents(EvalState))
 			{
 				return nullptr;
@@ -975,7 +984,8 @@ FGeometryScriptBakeTypeOptions UGeometryScriptLibrary_MeshBakeFunctions::MakeBak
 FGeometryScriptBakeTypeOptions UGeometryScriptLibrary_MeshBakeFunctions::MakeBakeTypeBentNormal(
 	int OcclusionRays,
 	float MaxDistance,
-	float SpreadAngle)
+	float SpreadAngle,
+	EGeometryScriptBakeNormalSpace NormalSpace)
 {
 	FGeometryScriptBakeTypeOptions Output;
 	Output.BakeType = EGeometryScriptBakeTypes::BentNormal;
@@ -983,6 +993,7 @@ FGeometryScriptBakeTypeOptions UGeometryScriptLibrary_MeshBakeFunctions::MakeBak
 	OcclusionOptions->OcclusionRays = OcclusionRays;
 	OcclusionOptions->MaxDistance = MaxDistance;
 	OcclusionOptions->SpreadAngle = SpreadAngle;
+	OcclusionOptions->NormalSpace = NormalSpace;
 	Output.Options = OcclusionOptions;
 	return Output;
 }
diff --git a/Engine/Plugins/Runtime/GeometryScripting/Source/GeometryScriptingCore/Public/GeometryScript/MeshBakeFunctions.h b/Engine/Plugins/Runtime/GeometryScripting/Source/GeometryScriptingCore/Public/GeometryScript/MeshBakeFunctions.h
--- a/Engine/Plugins/Runtime/GeometryScripting/Source/GeometryScriptingCore/Public/GeometryScript/MeshBakeFunctions.h
+++ b/Engine/Plugins/Runtime/GeometryScripting/Source/GeometryScriptingCore/Public/GeometryScript/MeshBakeFunctions.h
@@ -111,6 +111,9 @@ struct GEOMETRYSCRIPTINGCORE_API FGeometryScriptBakeType_Occlusion : public FGeo
 
 	/** Angle in degrees from the horizon for occlusion rays for which the contribution is attenuated to reduce faceting artifacts. */
 	float BiasAngle = 15.0f;
+
+	/** Normal space for Bent Normals. */
+	EGeometryScriptBakeNormalSpace NormalSpace = EGeometryScriptBakeNormalSpace::Tangent;
 };
 
 UENUM(BlueprintType)
@@ -499,7 +502,8 @@ public:
 	static UPARAM(DisplayName="Bake Type Out") FGeometryScriptBakeTypeOptions MakeBakeTypeBentNormal(
 		int OcclusionRays = 16,
 		float MaxDistance = 0.0f,
-		float SpreadAngle = 180.0f);
+		float SpreadAngle = 180.0f,
+		EGeometryScriptBakeNormalSpace NormalSpace = EGeometryScriptBakeNormalSpace::Tangent);
 
 	UFUNCTION(BlueprintPure, Category = "GeometryScript|Bake")
 	static UPARAM(DisplayName="Bake Type Out") FGeometryScriptBakeTypeOptions MakeBakeTypePosition();
