Geo Scripting: Don't crash on UObject::Modify returning false. Just ignore the result.

#rnx

diff --git a/Engine/Plugins/Runtime/GeometryScripting/Source/GeometryScriptingCore/Private/MeshAssetFunctions.cpp b/Engine/Plugins/Runtime/GeometryScripting/Source/GeometryScriptingCore/Private/MeshAssetFunctions.cpp
--- a/Engine/Plugins/Runtime/GeometryScripting/Source/GeometryScriptingCore/Private/MeshAssetFunctions.cpp
+++ b/Engine/Plugins/Runtime/GeometryScripting/Source/GeometryScriptingCore/Private/MeshAssetFunctions.cpp
@@ -883,7 +883,7 @@ UDynamicMesh* UGeometryScriptLibrary_StaticMeshFunctions::CopyMorphTargetToSkele
 	// make sure transactional flag is on for this asset
 	ToSkeletalMeshAsset->SetFlags(RF_Transactional);
 
-	verify(ToSkeletalMeshAsset->Modify());
+	(void)ToSkeletalMeshAsset->Modify();
 	
 	// Ensure we have enough LODInfos to cover up to the requested LOD.
 	for (int32 LODIndex = ToSkeletalMeshAsset->GetLODNum(); LODIndex <= TargetLOD.LODIndex; LODIndex++)
@@ -894,7 +894,7 @@ UDynamicMesh* UGeometryScriptLibrary_StaticMeshFunctions::CopyMorphTargetToSkele
 		LODInfo.ReductionSettings.BaseLOD = 0;
 	}
 
-	FMeshDescription* MeshDescription = ToSkeletalMeshAsset->GetMeshDescription(TargetLOD.LODIndex);;
+	FMeshDescription* MeshDescription = ToSkeletalMeshAsset->GetMeshDescription(TargetLOD.LODIndex);
 
 	if (MeshDescription == nullptr)
 	{
