GeoScripting: Fall back on render data if skelmesh model doesn't exist. This mirrors what CopyMeshFromStaticMesh does with missing source geometry.

#rnx

diff --git a/Engine/Plugins/Runtime/GeometryScripting/Source/GeometryScriptingCore/Private/MeshAssetFunctions.cpp b/Engine/Plugins/Runtime/GeometryScripting/Source/GeometryScriptingCore/Private/MeshAssetFunctions.cpp
--- a/Engine/Plugins/Runtime/GeometryScripting/Source/GeometryScriptingCore/Private/MeshAssetFunctions.cpp
+++ b/Engine/Plugins/Runtime/GeometryScripting/Source/GeometryScriptingCore/Private/MeshAssetFunctions.cpp
@@ -595,9 +595,34 @@ UDynamicMesh* UGeometryScriptLibrary_StaticMeshFunctions::CopyMeshFromSkeletalMe
 	}
 
 	// TODO: Consolidate this code with SkeletalMeshToolTarget::GetMeshDescription(..) 
-	const int32 UseLODIndex = FMath::Clamp(RequestedLOD.LODIndex, 0, FromSkeletalMeshAsset->GetLODNum() - 1);;
-	
-	if (RequestedLOD.LODType == EGeometryScriptLODType::MaxAvailable || RequestedLOD.LODType == EGeometryScriptLODType::SourceModel)
+	int32 UseLODIndex = RequestedLOD.LODIndex;
+	EGeometryScriptLODType UseLODType = RequestedLOD.LODType;
+
+#if WITH_EDITOR
+	if (UseLODType == EGeometryScriptLODType::MaxAvailable || UseLODType == EGeometryScriptLODType::HiResSourceModel)
+	{
+		UseLODType = EGeometryScriptLODType::SourceModel;
+	}
+
+	if (UseLODType == EGeometryScriptLODType::SourceModel)
+	{
+		UseLODIndex = FMath::Clamp(RequestedLOD.LODIndex, 0, FromSkeletalMeshAsset->GetNumSourceModels() - 1);
+		if (!FromSkeletalMeshAsset->GetSourceModel(UseLODIndex).HasMeshDescription())
+		{
+			UseLODType = EGeometryScriptLODType::RenderData;
+		}
+	}
+#endif
+
+	if (UseLODType == EGeometryScriptLODType::RenderData)
+	{
+		// TBD: Do we honor GetMinLodIdx?
+		if (UELocal::CopyMeshFromSkeletalMesh_RenderData(FromSkeletalMeshAsset, AssetOptions, RequestedLOD.LODIndex, ToDynamicMesh, Debug))
+		{
+			Outcome = EGeometryScriptOutcomePins::Success;
+		}
+	}
+	else
 	{
 #if WITH_EDITOR
 		const FMeshDescription* SourceMesh = nullptr;
@@ -620,25 +645,10 @@ UDynamicMesh* UGeometryScriptLibrary_StaticMeshFunctions::CopyMeshFromSkeletalMe
 		ToDynamicMesh->SetMesh(MoveTemp(NewMesh));
 	
 		Outcome = EGeometryScriptOutcomePins::Success;
-#else 
+#else
 		UE::Geometry::AppendError(Debug, EGeometryScriptErrorType::InvalidInputs, LOCTEXT("CopyMeshFromSkeletalMesh_SourceMesh_EditorOnly", "CopyMeshFromSkeletalMesh: Source Meshes are not available at Runtime"));
-	
-#endif	
-	}
-	else if (RequestedLOD.LODType == EGeometryScriptLODType::RenderData)
-	{
-		if (UELocal::CopyMeshFromSkeletalMesh_RenderData(FromSkeletalMeshAsset, AssetOptions, RequestedLOD.LODIndex, ToDynamicMesh, Debug))
-		{
-			Outcome = EGeometryScriptOutcomePins::Success;
-		}
-		
-	}
-
-// todo remove this
-#if !WITH_EDITOR
-
-	UE::Geometry::AppendError(Debug, EGeometryScriptErrorType::InvalidInputs, LOCTEXT("CopyMeshFromSkeletalMesh_EditorOnly", "CopyMeshFromSkeletalMesh: Not currently supported at Runtime"));
 #endif
+	}		
 	
 	return ToDynamicMesh;
 }
