GeoScripting: Utility function to prune specific bones from weight bindings

diff --git a/Engine/Plugins/Runtime/GeometryScripting/Source/GeometryScriptingCore/Private/MeshBoneWeightFunctions.cpp b/Engine/Plugins/Runtime/GeometryScripting/Source/GeometryScriptingCore/Private/MeshBoneWeightFunctions.cpp
--- a/Engine/Plugins/Runtime/GeometryScripting/Source/GeometryScriptingCore/Private/MeshBoneWeightFunctions.cpp
+++ b/Engine/Plugins/Runtime/GeometryScripting/Source/GeometryScriptingCore/Private/MeshBoneWeightFunctions.cpp
@@ -471,6 +471,148 @@ UDynamicMesh* UGeometryScriptLibrary_MeshBoneWeightFunctions::SetAllVertexBoneWe
 	return TargetMesh;
 }
 
+UDynamicMesh* UGeometryScriptLibrary_MeshBoneWeightFunctions::PruneBoneWeights(
+	UDynamicMesh* TargetMesh,
+	const TArray<FName>& BonesToPrune, 
+	FGeometryScriptPruneBoneWeightsOptions Options,
+	FGeometryScriptBoneWeightProfile Profile,
+	UGeometryScriptDebug* Debug
+	)
+{
+	if (TargetMesh == nullptr)
+	{
+		AppendError(Debug, EGeometryScriptErrorType::InvalidInputs, LOCTEXT("PruneBoneWeights_InvalidInput", "PruneBoneWeights: TargetMesh is Null"));
+		return TargetMesh;
+	}
+	
+	// Nothing to do?
+	if (BonesToPrune.IsEmpty())
+	{
+		return TargetMesh;
+	}
+	
+	// Validate that we're not trying to prune the root bone.
+	TargetMesh->EditMesh([&](FDynamicMesh3& EditMesh)
+	{
+		if (!EditMesh.HasAttributes() || !EditMesh.Attributes()->HasBones() || EditMesh.Attributes()->GetNumBones() == 0)
+		{
+			UE::Geometry::AppendError(Debug, EGeometryScriptErrorType::InvalidInputs, LOCTEXT("PruneBoneWeights_NoBones", "Target mesh has no bone attribute"));
+			return;
+		}
+		
+		if (!ValidateBoneHierarchy(EditMesh, Debug))
+		{
+			return;
+		}
+
+		const TArray<FName>& BoneNames = EditMesh.Attributes()->GetBoneNames()->GetAttribValues();
+		const TArray<int32>& BoneParents = EditMesh.Attributes()->GetBoneParentIndices()->GetAttribValues();
+
+		TArray<int32> BoneIndicesToPrune;
+		for (FName BoneName: BonesToPrune)
+		{
+			int32 BoneIndex = BoneNames.IndexOfByKey(BoneName);
+			if (BoneIndex > 0)
+			{
+				BoneIndicesToPrune.Add(BoneIndex);
+			}
+			else if (!Options.bIgnoredInvalidBones)
+			{
+				if (BoneIndex == 0)
+				{
+					UE::Geometry::AppendError(Debug, EGeometryScriptErrorType::InvalidInputs, LOCTEXT("PruneBoneWeights_RootBoneInvalid", "Pruning the root bone is not allowed"));
+					return;
+				}
+				else
+				{
+					UE::Geometry::AppendError(Debug, EGeometryScriptErrorType::InvalidInputs, FText::Format(LOCTEXT("PruneBoneWeights_InvalidBone", "Invalid bone '{0}'"), FText::FromName(BoneName)));
+					return;
+				}
+			}
+		}
+
+		FDynamicMeshVertexSkinWeightsAttribute* SkinWeights = EditMesh.Attributes()->GetSkinWeightsAttribute(Profile.GetProfileName());
+		if (!SkinWeights)
+		{
+			UE::Geometry::AppendWarning(Debug, EGeometryScriptErrorType::InvalidInputs, FText::Format(LOCTEXT("PruneBoneWeights_UnknownProfile", "Unknown skin weight profile '{0}'"), FText::FromName(Profile.GetProfileName())));
+			return;
+		}
+		
+		// Order the bones by descending tree depth.
+		auto BoneDepth = [&BoneParents](int32 BoneIndex)
+		{
+			int32 Depth = 0;
+			while (BoneParents[BoneIndex] != INDEX_NONE)
+			{
+				Depth++;
+				BoneIndex = BoneParents[BoneIndex];
+			}
+			return Depth;
+		};
+		BoneIndicesToPrune.StableSort([BoneDepth](int32 BoneIndexA, int32 BoneIndexB)
+		{
+			return BoneDepth(BoneIndexA) > BoneDepth(BoneIndexB);
+		});
+
+		// Iteratively prune bones such that we properly propagate weights up the skeleton if multiple bones along the same
+		// path are being removed.
+		for (const int32 BoneIndex: BoneIndicesToPrune)
+		{
+			for (const int32 VertexID : EditMesh.VertexIndicesItr())
+			{
+				FBoneWeights BoneWeights;
+				SkinWeights->GetValue(VertexID, BoneWeights);
+
+				const int32 WeightIndex = BoneWeights.FindWeightIndexByBone(BoneIndex);
+				if (WeightIndex == INDEX_NONE)
+				{
+					continue;
+				}
+
+				const int32 ParentBoneIndex = BoneParents[BoneIndex];
+				checkSlow(ParentBoneIndex != INDEX_NONE);
+				
+				if (BoneWeights.Num() == 1)
+				{
+					// Is it the last remaining bone weight? Re-assign this vertex to the parent.
+					BoneWeights = FBoneWeights::Create({FBoneWeight(ParentBoneIndex, MaxRawBoneWeight)});
+				}
+				else if (Options.ReassignmentType == EGeometryScriptPruneBoneWeightsAssignmentType::RenormalizeRemaining)
+				{
+					// Just remove the weight and renormalize what's remaining.
+					BoneWeights.RemoveBoneWeight(BoneIndex);
+				}
+				else if (Options.ReassignmentType == EGeometryScriptPruneBoneWeightsAssignmentType::ReassignToParent)
+				{
+					FBoneWeightsSettings SettingsNoNormalize;
+					SettingsNoNormalize.SetNormalizeType(EBoneWeightNormalizeType::None);
+					
+					FBoneWeight BoneWeight = BoneWeights[WeightIndex];
+
+					// Remove the weight but don't normalize yet.
+					BoneWeights.RemoveBoneWeight(BoneIndex, SettingsNoNormalize);
+
+					// If the parent weight already exists, add the child weight to it.
+					if (const int32 ParentWeightIndex = BoneWeights.FindWeightIndexByBone(ParentBoneIndex);
+						ParentWeightIndex != INDEX_NONE)
+					{
+						FBoneWeight ParentBoneWeight = BoneWeights[ParentWeightIndex];
+						BoneWeight.SetRawWeight(BoneWeight.GetRawWeight() + ParentBoneWeight.GetRawWeight());
+					}
+
+					// Set the weight to be the combination of the removed weight and the parent and renormalize now.
+					BoneWeight.SetBoneIndex(ParentBoneIndex);
+					BoneWeights.SetBoneWeight(BoneWeight);
+				}
+
+				SkinWeights->SetValue(VertexID, BoneWeights);
+			}
+		}
+	}, EDynamicMeshChangeType::AttributeEdit, EDynamicMeshAttributeChangeFlags::Unknown, false);
+
+	return TargetMesh;
+}
+
 
 UDynamicMesh* UGeometryScriptLibrary_MeshBoneWeightFunctions::ComputeSmoothBoneWeights(
 	UDynamicMesh* TargetMesh,
diff --git a/Engine/Plugins/Runtime/GeometryScripting/Source/GeometryScriptingCore/Public/GeometryScript/MeshBoneWeightFunctions.h b/Engine/Plugins/Runtime/GeometryScripting/Source/GeometryScriptingCore/Public/GeometryScript/MeshBoneWeightFunctions.h
--- a/Engine/Plugins/Runtime/GeometryScripting/Source/GeometryScriptingCore/Public/GeometryScript/MeshBoneWeightFunctions.h
+++ b/Engine/Plugins/Runtime/GeometryScripting/Source/GeometryScriptingCore/Public/GeometryScript/MeshBoneWeightFunctions.h
@@ -43,6 +43,26 @@ struct GEOMETRYSCRIPTINGCORE_API FGeometryScriptBoneWeightProfile
 	FName GetProfileName() const { return ProfileName; }
 };
 
+UENUM(BlueprintType)
+enum class EGeometryScriptPruneBoneWeightsAssignmentType : uint8
+{
+	RenormalizeRemaining = 0,	/** Remove the bone from the bone weights and renormalize the remaining weights. */
+	ReassignToParent = 1,		/** Re-assign the removed bone's weight to the parent bone. */
+};
+
+USTRUCT(BlueprintType)
+struct GEOMETRYSCRIPTINGCORE_API FGeometryScriptPruneBoneWeightsOptions
+{
+	GENERATED_BODY()
+
+	/** Specifies how the weight of the removed bone from a vertex's bone weights list gets reassigned. */
+	UPROPERTY(EditAnywhere, BlueprintReadWrite, Category = Options)
+	EGeometryScriptPruneBoneWeightsAssignmentType ReassignmentType = EGeometryScriptPruneBoneWeightsAssignmentType::RenormalizeRemaining;
+
+	/** Ignore invalid bones. Otherwise, if invalid bones are given, the operation terminates with an error */
+	UPROPERTY(EditAnywhere, BlueprintReadWrite, Category = Options)
+	bool bIgnoredInvalidBones = true;
+};
 
 UENUM(BlueprintType)
 enum class EGeometryScriptSmoothBoneWeightsType : uint8
@@ -374,6 +394,28 @@ public:
 		const TArray<FGeometryScriptBoneWeight>& BoneWeights,
 		FGeometryScriptBoneWeightProfile Profile = FGeometryScriptBoneWeightProfile(),
 		UGeometryScriptDebug* Debug = nullptr);
+
+
+	/**
+	 *  Prunes the given bones from any bone weight assignment on the given profile. 
+	 *  The bone weights are re-assigned based on the type of re-assignment specified in the options,
+	 *  although in the case where the bone(s) being pruned are the sole bone weight on a vertex, then
+	 *  the parent bone will be assigned as the sole bone weight for that vertex.
+	 *  Bones are pruned iteratively from leaf to root, to ensure that weighs are progressively re-assigned
+	 *  in case multiple bones along the same branch are being pruned.
+	 *  @param BonesToPrune The list of bones to remove.
+	 *  @param Options The options to set for the pruning algorithm.
+	 *  @param Profile The skin weight profile to prune the bones from.
+	 */
+	UFUNCTION(BlueprintCallable, Category = "GeometryScript|MeshQueries|BoneWeights", meta=(ScriptMethod))
+	static UPARAM(DisplayName = "Target Mesh") UDynamicMesh* 
+	PruneBoneWeights( 
+		UDynamicMesh* TargetMesh,
+		const TArray<FName>& BonesToPrune,
+		FGeometryScriptPruneBoneWeightsOptions Options,
+		FGeometryScriptBoneWeightProfile Profile = FGeometryScriptBoneWeightProfile(),
+		UGeometryScriptDebug* Debug = nullptr);
+	
 	
 	/** 
 	 *  Computes a smooth skin binding for the given mesh to the skeleton provided.
