add method to generate collision shapes from a dynamic mesh and return the collision shapes, without requiring a dynamic mesh component

diff --git a/Engine/Plugins/Runtime/GeometryScripting/Source/GeometryScriptingCore/Private/CollisionFunctions.cpp b/Engine/Plugins/Runtime/GeometryScripting/Source/GeometryScriptingCore/Private/CollisionFunctions.cpp
--- a/Engine/Plugins/Runtime/GeometryScripting/Source/GeometryScriptingCore/Private/CollisionFunctions.cpp
+++ b/Engine/Plugins/Runtime/GeometryScripting/Source/GeometryScriptingCore/Private/CollisionFunctions.cpp
@@ -526,6 +526,29 @@ void UGeometryScriptLibrary_CollisionFunctions::ResetDynamicMeshCollision(
 }
 
 
+FGeometryScriptSimpleCollision
+UGeometryScriptLibrary_CollisionFunctions::GenerateCollisionFromMesh(
+	UDynamicMesh* FromDynamicMesh,
+	FGeometryScriptCollisionFromMeshOptions Options,
+	UGeometryScriptDebug* Debug)
+{
+	FGeometryScriptSimpleCollision ToRet;
+	
+	if (FromDynamicMesh == nullptr)
+	{
+		UE::Geometry::AppendError(Debug, EGeometryScriptErrorType::InvalidInputs, LOCTEXT("GenerateCollisionFromMesh_InvalidInput", "GenerateCollisionFromMesh: FromDynamicMesh is Null"));
+		return ToRet;
+	}
+
+	FromDynamicMesh->ProcessMesh([&ToRet, &Options](const FDynamicMesh3& ReadMesh)
+	{
+		UELocal::ComputeCollisionFromMesh(ReadMesh, ToRet.AggGeom, Options);
+	});
+
+	return ToRet;
+}
+
+
 FGeometryScriptSimpleCollision
 UGeometryScriptLibrary_CollisionFunctions::GetSimpleCollisionFromComponent(
 	UPrimitiveComponent* Component,
diff --git a/Engine/Plugins/Runtime/GeometryScripting/Source/GeometryScriptingCore/Public/GeometryScript/CollisionFunctions.h b/Engine/Plugins/Runtime/GeometryScripting/Source/GeometryScriptingCore/Public/GeometryScript/CollisionFunctions.h
--- a/Engine/Plugins/Runtime/GeometryScripting/Source/GeometryScriptingCore/Public/GeometryScript/CollisionFunctions.h
+++ b/Engine/Plugins/Runtime/GeometryScripting/Source/GeometryScriptingCore/Public/GeometryScript/CollisionFunctions.h
@@ -441,6 +441,16 @@ public:
 		FGeometryScriptSetStaticMeshCollisionOptions StaticMeshCollisionOptions = FGeometryScriptSetStaticMeshCollisionOptions(),
 		UGeometryScriptDebug* Debug = nullptr);
 
+	/**
+	* Generate Simple Collision shapes for an input Dynamic Mesh shape
+	*/
+	UFUNCTION(BlueprintCallable, Category = "GeometryScript|Collision")
+	static UPARAM(DisplayName = "Simple Collision") FGeometryScriptSimpleCollision
+	GenerateCollisionFromMesh(
+		UDynamicMesh* FromDynamicMesh,
+		FGeometryScriptCollisionFromMeshOptions Options,
+		UGeometryScriptDebug* Debug = nullptr);
+
 	/*
 	 * Count of number of simple collision shapes
 	 */
