GeometryScript: Added UVShell bake type to BakeTexture node.

diff --git a/Engine/Plugins/Runtime/GeometryScripting/Source/GeometryScriptingCore/Private/MeshBakeFunctions.cpp b/Engine/Plugins/Runtime/GeometryScripting/Source/GeometryScriptingCore/Private/MeshBakeFunctions.cpp
--- a/Engine/Plugins/Runtime/GeometryScripting/Source/GeometryScriptingCore/Private/MeshBakeFunctions.cpp
+++ b/Engine/Plugins/Runtime/GeometryScripting/Source/GeometryScriptingCore/Private/MeshBakeFunctions.cpp
@@ -15,6 +15,7 @@
 #include "Sampling/MeshOcclusionMapEvaluator.h"
 #include "Sampling/MeshNormalMapEvaluator.h"
 #include "Sampling/MeshPropertyMapEvaluator.h"
+#include "Sampling/MeshUVShellMapEvaluator.h"
 #include "Sampling/MeshResampleImageEvaluator.h"
 
 #include "DynamicMesh/MeshTransforms.h"
@@ -147,11 +148,17 @@ namespace GeometryScriptBakeLocals
 				break;
 			case EMeshPropertyMapType::VertexColor:
 			case EMeshPropertyMapType::MaterialID:
+			case EMeshPropertyMapType::PolyGroupID:
 				TexType = FTexture2DBuilder::ETextureType::Color;
 				break;
 			}
 			break;
 		}
+		case EMeshMapEvaluatorType::UVShell:
+		{
+			TexType = FTexture2DBuilder::ETextureType::Color;
+			break;
+		}
 		case EMeshMapEvaluatorType::Curvature:
 		{
 			TexType = FTexture2DBuilder::ETextureType::ColorLinear;
@@ -483,6 +490,18 @@ namespace GeometryScriptBakeLocals
 			Result = PropertyEval;
 			break;
 		}
+		case EGeometryScriptBakeTypes::UVShell:
+		{
+			FGeometryScriptBakeType_UVShell* UVShellOptions = static_cast<FGeometryScriptBakeType_UVShell*>(Options.Options.Get());
+			TSharedPtr<FMeshUVShellMapEvaluator> UVShellEval = MakeShared<FMeshUVShellMapEvaluator>();
+			UVShellEval->UVLayer = UVShellOptions->SourceUVLayer;
+			UVShellEval->WireframeThickness = UVShellOptions->WireframeThickness;
+			UVShellEval->WireframeColor = UVShellOptions->WireframeColor;
+			UVShellEval->ShellColor = UVShellOptions->ShellColor;
+			UVShellEval->BackgroundColor = UVShellOptions->BackgroundColor;
+			Result = UVShellEval;
+			break;
+		}
 		case EGeometryScriptBakeTypes::Constant:
 		{
 			FGeometryScriptBakeType_Constant* ConstantOptions = static_cast<FGeometryScriptBakeType_Constant*>(Options.Options.Get());
@@ -592,6 +611,20 @@ namespace GeometryScriptBakeLocals
 				// Abort if any evaluators failed to build.
 				return nullptr;
 			}
+
+			// Post evaluator construction handling
+			switch (Options.BakeType)
+			{
+			case EGeometryScriptBakeTypes::UVShell:
+			{
+				FMeshUVShellMapEvaluator* UVShellEval = static_cast<FMeshUVShellMapEvaluator*>(Eval.Get());
+				UVShellEval->TexelSize = BakeDimensions.GetTexelSize();
+				break;
+			}
+			default:
+				break;
+			}
+			
 			Baker.AddEvaluator(Eval);
 		}
 
@@ -759,7 +792,7 @@ namespace GeometryScriptBakeLocals
 		{
 			auto IsValidBakeType = [&BakeVertexPrefix](EGeometryScriptBakeTypes BakeType, TArray<FGeometryScriptDebugMessage>* Debug)
 			{
-				const bool bIsValid = BakeType != EGeometryScriptBakeTypes::None;
+				const bool bIsValid = (BakeType != EGeometryScriptBakeTypes::None && BakeType != EGeometryScriptBakeTypes::UVShell);
 				if (!bIsValid)
 				{
 					const FText BakeTypeName = FText::FromName(StaticEnum<EGeometryScriptBakeTypes>()->GetNameByIndex(static_cast<int32>(BakeType)));
@@ -1080,6 +1113,25 @@ FGeometryScriptBakeTypeOptions UGeometryScriptLibrary_MeshBakeFunctions::MakeBak
 	return Output;
 }
 
+FGeometryScriptBakeTypeOptions UGeometryScriptLibrary_MeshBakeFunctions::MakeBakeTypeUVShell(
+	int SourceUVLayer,
+	float WireframeThickness,
+	FLinearColor WireframeColor,
+	FLinearColor ShellColor,
+	FLinearColor BackgroundColor)
+{
+	FGeometryScriptBakeTypeOptions Output;
+	Output.BakeType = EGeometryScriptBakeTypes::UVShell;
+	const TSharedPtr<FGeometryScriptBakeType_UVShell> UVShellOptions = MakeShared<FGeometryScriptBakeType_UVShell>();
+	UVShellOptions->SourceUVLayer = SourceUVLayer;
+	UVShellOptions->WireframeThickness = WireframeThickness;
+	UVShellOptions->WireframeColor = WireframeColor;
+	UVShellOptions->ShellColor = ShellColor;
+	UVShellOptions->BackgroundColor = BackgroundColor;
+	Output.Options = UVShellOptions;
+	return Output;
+}
+
 FGeometryScriptBakeTypeOptions UGeometryScriptLibrary_MeshBakeFunctions::MakeBakeTypeConstant(
 	float Value)
 {
diff --git a/Engine/Plugins/Runtime/GeometryScripting/Source/GeometryScriptingCore/Public/GeometryScript/MeshBakeFunctions.h b/Engine/Plugins/Runtime/GeometryScripting/Source/GeometryScriptingCore/Public/GeometryScript/MeshBakeFunctions.h
--- a/Engine/Plugins/Runtime/GeometryScripting/Source/GeometryScriptingCore/Public/GeometryScript/MeshBakeFunctions.h
+++ b/Engine/Plugins/Runtime/GeometryScripting/Source/GeometryScriptingCore/Public/GeometryScript/MeshBakeFunctions.h
@@ -75,7 +75,9 @@ enum class EGeometryScriptBakeTypes : uint8
 	/* Material IDs as unique colors */
 	MaterialID             UMETA(DisplayName = "Material ID"),
 	/* Constant value */
-	Constant
+	Constant,
+	/* UV Shell */
+	UVShell                UMETA(DisplayName = "UV Shell")
 };
 
 UENUM(BlueprintType)
@@ -129,7 +131,6 @@ enum class EGeometryScriptBakeCurvatureTypeMode : uint8
 	Gaussian
 };
 
-
 UENUM(BlueprintType)
 enum class EGeometryScriptBakeCurvatureColorMode : uint8
 {
@@ -141,7 +142,6 @@ enum class EGeometryScriptBakeCurvatureColorMode : uint8
 	RedGreenBlue
 };
 
-
 UENUM(BlueprintType)
 enum class EGeometryScriptBakeCurvatureClampMode : uint8
 {
@@ -171,6 +171,24 @@ struct GEOMETRYSCRIPTINGCORE_API FGeometryScriptBakeType_Curvature : public FGeo
 	EGeometryScriptBakeCurvatureClampMode Clamping = EGeometryScriptBakeCurvatureClampMode::None;
 };
 
+struct GEOMETRYSCRIPTINGCORE_API FGeometryScriptBakeType_UVShell : public FGeometryScriptBakeTypes
+{
+	/** The source UV layer to sample. */
+	int SourceUVLayer = 0;
+	
+	/** The pixel thickness of the wireframes. */
+	float WireframeThickness = 1.0f;
+
+	/** The color for wireframe samples. */
+	FLinearColor WireframeColor = FLinearColor::Blue;
+
+	/** The color for UV shell triangle fill samples. */
+	FLinearColor ShellColor = FLinearColor::Gray;
+
+	/** The color for samples external to UV shells. */
+	FLinearColor BackgroundColor = FLinearColor(0.0f, 0.0f, 0.0f, 0.0f);
+};
+
 struct GEOMETRYSCRIPTINGCORE_API FGeometryScriptBakeType_Texture : public FGeometryScriptBakeTypes
 {
 	/** Source mesh texture that is to be resampled into a new texture */
@@ -539,6 +557,14 @@ public:
 	UFUNCTION(BlueprintPure, Category = "GeometryScript|Bake")
 	static UPARAM(DisplayName="Bake Type Out") FGeometryScriptBakeTypeOptions MakeBakeTypeMaterialID();
 
+	UFUNCTION(BlueprintPure, Category = "GeometryScript|Bake")
+	static UPARAM(DisplayName="Bake Type Out") FGeometryScriptBakeTypeOptions MakeBakeTypeUVShell(
+		int SourceUVLayer = 0,
+		float WireframeThickness = 1.0f,
+		FLinearColor WireframeColor = FLinearColor::Blue,
+		FLinearColor ShellColor = FLinearColor::Gray,
+		FLinearColor BackgroundColor = FLinearColor(0.0f, 0.0f, 0.0f, 0.0f));
+
 	UFUNCTION(BlueprintPure, Category = "GeometryScript|Bake")
 	static UPARAM(DisplayName="Bake Type Out") FGeometryScriptBakeTypeOptions MakeBakeTypeConstant(
 		float Value = 0.0f);
