Add 'with collision' variants of geometry script append mesh primitive functions for the append box, sphere, and capsule methods -- returning a matching simple collision shape of the corresponding type

diff --git a/Engine/Plugins/Runtime/GeometryScripting/Source/GeometryScriptingCore/Private/MeshPrimitiveFunctions.cpp b/Engine/Plugins/Runtime/GeometryScripting/Source/GeometryScriptingCore/Private/MeshPrimitiveFunctions.cpp
--- a/Engine/Plugins/Runtime/GeometryScripting/Source/GeometryScriptingCore/Private/MeshPrimitiveFunctions.cpp
+++ b/Engine/Plugins/Runtime/GeometryScripting/Source/GeometryScriptingCore/Private/MeshPrimitiveFunctions.cpp
@@ -198,9 +198,221 @@ static void AppendPrimitiveMesh(
 	}
 }
 
+namespace UE::Private::AppendMeshHelper
+{
+	static UDynamicMesh* AppendBox(UDynamicMesh* TargetMesh,
+		FGeometryScriptSimpleCollision* SimpleCollision,
+		FGeometryScriptPrimitiveOptions PrimitiveOptions,
+		FTransform Transform,
+		float DimensionX,
+		float DimensionY,
+		float DimensionZ,
+		int32 StepsX,
+		int32 StepsY,
+		int32 StepsZ,
+		EGeometryScriptPrimitiveOriginMode Origin,
+		UGeometryScriptDebug* Debug)
+	{
+		if (TargetMesh == nullptr)
+		{
+			UE::Geometry::AppendError(Debug, EGeometryScriptErrorType::InvalidInputs, LOCTEXT("PrimitiveFunctions_AppendBox", "AppendBox: TargetMesh is Null"));
+			return TargetMesh;
+		}
+
+		UE::Geometry::FAxisAlignedBox3d ConvertBox(
+			FVector3d(-DimensionX / 2, -DimensionY / 2, 0),
+			FVector3d(DimensionX / 2, DimensionY / 2, DimensionZ));
+
+		// todo: if Steps X/Y/Z are zero, can use trivial box generator
+
+		FGridBoxMeshGenerator GridBoxGenerator;
+		GridBoxGenerator.Box = UE::Geometry::FOrientedBox3d(ConvertBox);
+		GridBoxGenerator.EdgeVertices = FIndex3i(FMath::Max(0, StepsX), FMath::Max(0, StepsY), FMath::Max(0, StepsZ));
+		GridBoxGenerator.bPolygroupPerQuad = (PrimitiveOptions.PolygroupMode == EGeometryScriptPrimitivePolygroupMode::PerQuad);
+		GridBoxGenerator.Generate();
+
+		FVector3d OriginShift = (Origin == EGeometryScriptPrimitiveOriginMode::Center) ? FVector3d(0, 0, -DimensionZ / 2) : FVector3d::Zero();
+		AppendPrimitive(TargetMesh, &GridBoxGenerator, Transform, PrimitiveOptions, OriginShift);
+
+		if (SimpleCollision)
+		{
+			FVector3d Dims = Transform.GetScale3D() * FVector3d(DimensionX, DimensionY, DimensionZ);
+			FKBoxElem BoxElem(Dims.X, Dims.Y, Dims.Z);
+			BoxElem.Center = Transform.TransformPosition(ConvertBox.Center() + OriginShift);
+			BoxElem.Rotation = Transform.Rotator();
+			SimpleCollision->AggGeom.BoxElems.Add(BoxElem);
+		}
+
+		return TargetMesh;
+	}
+
+	static UDynamicMesh* AppendBoundingBox(
+		UDynamicMesh* TargetMesh,
+		FGeometryScriptSimpleCollision* SimpleCollision,
+		FGeometryScriptPrimitiveOptions PrimitiveOptions,
+		FTransform Transform,
+		FBox Box,
+		int32 StepsX,
+		int32 StepsY,
+		int32 StepsZ,
+		UGeometryScriptDebug* Debug)
+	{
+		if (TargetMesh == nullptr)
+		{
+			UE::Geometry::AppendError(Debug, EGeometryScriptErrorType::InvalidInputs, LOCTEXT("PrimitiveFunctions_AppendBoundingBox", "AppendBoundingBox: TargetMesh is Null"));
+			return TargetMesh;
+		}
+
+		// todo: if Steps X/Y/Z are zero, can use trivial box generator
+
+		FGridBoxMeshGenerator GridBoxGenerator;
+		GridBoxGenerator.Box = UE::Geometry::FOrientedBox3d( UE::Geometry::FAxisAlignedBox3d(Box) );
+		GridBoxGenerator.EdgeVertices = FIndex3i(FMath::Max(0, StepsX), FMath::Max(0, StepsY), FMath::Max(0, StepsZ));
+		GridBoxGenerator.bPolygroupPerQuad = (PrimitiveOptions.PolygroupMode == EGeometryScriptPrimitivePolygroupMode::PerQuad);
+		GridBoxGenerator.Generate();
+
+		AppendPrimitive(TargetMesh, &GridBoxGenerator, Transform, PrimitiveOptions, FVector3d::Zero());
+
+		if (SimpleCollision)
+		{
+			FVector3d Dims = Transform.GetScale3D() * Box.GetSize();
+			FKBoxElem BoxElem(Dims.X, Dims.Y, Dims.Z);
+			BoxElem.Center = Transform.TransformPosition(Box.GetCenter());
+			BoxElem.Rotation = Transform.Rotator();
+			SimpleCollision->AggGeom.BoxElems.Add(BoxElem);
+		}
+
+		return TargetMesh;
+	}
+
+	void AppendSphereSimpleCollision(FKAggregateGeom& AggGeom, const FTransform& Transform, float Radius, FVector Center)
+	{
+		FKSphereElem Sphere;
+		Sphere.Center = Transform.TransformPosition(Center);
+		Sphere.Radius = Radius * Transform.GetMinimumAxisScale();
+		AggGeom.SphereElems.Add(Sphere);
+	}
+
+	UDynamicMesh* AppendSphereLatLong(
+		UDynamicMesh* TargetMesh,
+		FGeometryScriptSimpleCollision* SimpleCollision,
+		FGeometryScriptPrimitiveOptions PrimitiveOptions,
+		FTransform Transform,
+		float Radius,
+		int32 StepsPhi,
+		int32 StepsTheta,
+		EGeometryScriptPrimitiveOriginMode Origin,
+		UGeometryScriptDebug* Debug)
+	{
+		if (TargetMesh == nullptr)
+		{
+			UE::Geometry::AppendError(Debug, EGeometryScriptErrorType::InvalidInputs, LOCTEXT("PrimitiveFunctions_AppendSphereLatLong", "AppendSphereLatLong: TargetMesh is Null"));
+			return TargetMesh;
+		}
+
+		FSphereGenerator SphereGenerator;
+		SphereGenerator.Radius = FMath::Max(FMathf::ZeroTolerance, Radius);
+		SphereGenerator.NumPhi = FMath::Max(3, StepsPhi);
+		SphereGenerator.NumTheta = FMath::Max(3, StepsTheta);
+		SphereGenerator.bPolygroupPerQuad = (PrimitiveOptions.PolygroupMode == EGeometryScriptPrimitivePolygroupMode::PerQuad);
+		SphereGenerator.Generate();
+
+		FVector3d OriginShift = (Origin == EGeometryScriptPrimitiveOriginMode::Base) ? FVector3d(0, 0, Radius) : FVector3d::Zero();
+		AppendPrimitive(TargetMesh, &SphereGenerator, Transform, PrimitiveOptions, OriginShift);
+
+		if (SimpleCollision)
+		{
+			AppendSphereSimpleCollision(SimpleCollision->AggGeom, Transform, Radius, OriginShift);
+		}
+
+		return TargetMesh;
+	}
+
+	UDynamicMesh* AppendSphereBox(
+		UDynamicMesh* TargetMesh,
+		FGeometryScriptSimpleCollision* SimpleCollision,
+		FGeometryScriptPrimitiveOptions PrimitiveOptions,
+		FTransform Transform,
+		float Radius,
+		int32 StepsX,
+		int32 StepsY,
+		int32 StepsZ,
+		EGeometryScriptPrimitiveOriginMode Origin,
+		UGeometryScriptDebug* Debug)
+	{
+		if (TargetMesh == nullptr)
+		{
+			UE::Geometry::AppendError(Debug, EGeometryScriptErrorType::InvalidInputs, LOCTEXT("PrimitiveFunctions_AppendSphereBox", "AppendSphereBox: TargetMesh is Null"));
+			return TargetMesh;
+		}
+
+		FBoxSphereGenerator SphereGenerator;
+		SphereGenerator.Radius = FMath::Max(FMathf::ZeroTolerance, Radius);
+		SphereGenerator.EdgeVertices = FIndex3i(FMath::Max(0, StepsX), FMath::Max(0, StepsY), FMath::Max(0, StepsZ));
+		SphereGenerator.bPolygroupPerQuad = (PrimitiveOptions.PolygroupMode == EGeometryScriptPrimitivePolygroupMode::PerQuad);
+		SphereGenerator.Generate();
+
+		FVector3d OriginShift = (Origin == EGeometryScriptPrimitiveOriginMode::Base) ? FVector3d(0, 0, Radius) : FVector3d::Zero();
+		AppendPrimitive(TargetMesh, &SphereGenerator, Transform, PrimitiveOptions, OriginShift);
+
+		if (SimpleCollision)
+		{
+			AppendSphereSimpleCollision(SimpleCollision->AggGeom, Transform, Radius, OriginShift);
+		}
+
+		return TargetMesh;
+	}
+
+	UDynamicMesh* AppendCapsule(
+		UDynamicMesh* TargetMesh,
+		FGeometryScriptSimpleCollision* SimpleCollision,
+		FGeometryScriptPrimitiveOptions PrimitiveOptions,
+		FTransform Transform,
+		float Radius,
+		float LineLength,
+		int32 HemisphereSteps,
+		int32 CircleSteps,
+		int32 SegmentSteps,
+		EGeometryScriptPrimitiveOriginMode Origin,
+		UGeometryScriptDebug* Debug)
+	{
+		if (TargetMesh == nullptr)
+		{
+			UE::Geometry::AppendError(Debug, EGeometryScriptErrorType::InvalidInputs, LOCTEXT("PrimitiveFunctions_AppendCapsule", "AppendCapsule: TargetMesh is Null"));
+			return TargetMesh;
+		}
+
+		FCapsuleGenerator CapsuleGenerator;
+		CapsuleGenerator.Radius = FMath::Max(FMathf::ZeroTolerance, Radius);
+		CapsuleGenerator.SegmentLength = FMath::Max(FMathf::ZeroTolerance, LineLength);
+		CapsuleGenerator.NumHemisphereArcSteps = FMath::Max(2, HemisphereSteps);
+		CapsuleGenerator.NumCircleSteps = FMath::Max(3, CircleSteps);
+		CapsuleGenerator.NumSegmentSteps = FMath::Max(0, SegmentSteps);
+		CapsuleGenerator.bPolygroupPerQuad = (PrimitiveOptions.PolygroupMode == EGeometryScriptPrimitivePolygroupMode::PerQuad);
+		CapsuleGenerator.Generate();
+
+		FVector3d OriginShift = FVector3d::Zero();
+		OriginShift.Z = (Origin == EGeometryScriptPrimitiveOriginMode::Center) ? -(LineLength / 2) : Radius;
+		AppendPrimitive(TargetMesh, &CapsuleGenerator, Transform, PrimitiveOptions, OriginShift);
 
+		if (SimpleCollision)
+		{
+			FKSphylElem CapsuleElem;
+			FVector3d Center(0, 0, Origin == EGeometryScriptPrimitiveOriginMode::Center ? 0.0 : CapsuleGenerator.SegmentLength * .5 + CapsuleGenerator.Radius);
+			CapsuleElem.Center = Transform.TransformPosition(Center);
+			CapsuleElem.Rotation = Transform.Rotator();
+			FVector3d Scale = Transform.GetScale3D();
+			CapsuleElem.Length = Scale.Z * CapsuleGenerator.SegmentLength;
+			CapsuleElem.Radius = CapsuleGenerator.Radius * FMath::Max(FMath::Abs(Scale.X), FMath::Abs(Scale.Y));
+			SimpleCollision->AggGeom.SphylElems.Add(CapsuleElem);
+		}
+
+		return TargetMesh;
+	}
 
 
+}
+
 UDynamicMesh* UGeometryScriptLibrary_MeshPrimitiveFunctions::AppendBox(
 	UDynamicMesh* TargetMesh,
 	FGeometryScriptPrimitiveOptions PrimitiveOptions,
@@ -214,28 +426,24 @@ UDynamicMesh* UGeometryScriptLibrary_MeshPrimitiveFunctions::AppendBox(
 	EGeometryScriptPrimitiveOriginMode Origin,
 	UGeometryScriptDebug* Debug)
 {
-	if (TargetMesh == nullptr)
-	{
-		UE::Geometry::AppendError(Debug, EGeometryScriptErrorType::InvalidInputs, LOCTEXT("PrimitiveFunctions_AppendBox", "AppendBox: TargetMesh is Null"));
-		return TargetMesh;
-	}
-
-	UE::Geometry::FAxisAlignedBox3d ConvertBox(
-		FVector3d(-DimensionX / 2, -DimensionY / 2, 0),
-		FVector3d(DimensionX / 2, DimensionY / 2, DimensionZ));
-	
-	// todo: if Steps X/Y/Z are zero, can use trivial box generator
-
-	FGridBoxMeshGenerator GridBoxGenerator;
-	GridBoxGenerator.Box = UE::Geometry::FOrientedBox3d(ConvertBox);
-	GridBoxGenerator.EdgeVertices = FIndex3i(FMath::Max(0, StepsX), FMath::Max(0, StepsY), FMath::Max(0, StepsZ));
-	GridBoxGenerator.bPolygroupPerQuad = (PrimitiveOptions.PolygroupMode == EGeometryScriptPrimitivePolygroupMode::PerQuad);
-	GridBoxGenerator.Generate();
-
-	FVector3d OriginShift = (Origin == EGeometryScriptPrimitiveOriginMode::Center) ? FVector3d(0, 0, -DimensionZ/2) : FVector3d::Zero();
-	AppendPrimitive(TargetMesh, &GridBoxGenerator, Transform, PrimitiveOptions, OriginShift);
+	return UE::Private::AppendMeshHelper::AppendBox(TargetMesh, nullptr, PrimitiveOptions, Transform, DimensionX, DimensionY, DimensionZ, StepsX, StepsY, StepsZ, Origin, Debug);
+}
 
-	return TargetMesh;
+UDynamicMesh* UGeometryScriptLibrary_MeshPrimitiveFunctions::AppendBoxWithCollision(
+	UDynamicMesh* TargetMesh,
+	FGeometryScriptSimpleCollision& SimpleCollision,
+	FGeometryScriptPrimitiveOptions PrimitiveOptions,
+	FTransform Transform,
+	float DimensionX,
+	float DimensionY,
+	float DimensionZ,
+	int32 StepsX,
+	int32 StepsY,
+	int32 StepsZ,
+	EGeometryScriptPrimitiveOriginMode Origin,
+	UGeometryScriptDebug* Debug)
+{
+	return UE::Private::AppendMeshHelper::AppendBox(TargetMesh, &SimpleCollision, PrimitiveOptions, Transform, DimensionX, DimensionY, DimensionZ, StepsX, StepsY, StepsZ, Origin, Debug);
 }
 
 
@@ -249,23 +457,21 @@ UDynamicMesh* UGeometryScriptLibrary_MeshPrimitiveFunctions::AppendBoundingBox(
 	int32 StepsZ,
 	UGeometryScriptDebug* Debug)
 {
-	if (TargetMesh == nullptr)
-	{
-		UE::Geometry::AppendError(Debug, EGeometryScriptErrorType::InvalidInputs, LOCTEXT("PrimitiveFunctions_AppendBoundingBox", "AppendBoundingBox: TargetMesh is Null"));
-		return TargetMesh;
-	}
-
-	// todo: if Steps X/Y/Z are zero, can use trivial box generator
-
-	FGridBoxMeshGenerator GridBoxGenerator;
-	GridBoxGenerator.Box = UE::Geometry::FOrientedBox3d( UE::Geometry::FAxisAlignedBox3d(Box) );
-	GridBoxGenerator.EdgeVertices = FIndex3i(FMath::Max(0, StepsX), FMath::Max(0, StepsY), FMath::Max(0, StepsZ));
-	GridBoxGenerator.bPolygroupPerQuad = (PrimitiveOptions.PolygroupMode == EGeometryScriptPrimitivePolygroupMode::PerQuad);
-	GridBoxGenerator.Generate();
-
-	AppendPrimitive(TargetMesh, &GridBoxGenerator, Transform, PrimitiveOptions, FVector3d::Zero());
+	return UE::Private::AppendMeshHelper::AppendBoundingBox(TargetMesh, nullptr, PrimitiveOptions, Transform, Box, StepsX, StepsY, StepsZ, Debug);
+}
 
-	return TargetMesh;
+UDynamicMesh* UGeometryScriptLibrary_MeshPrimitiveFunctions::AppendBoundingBoxWithCollision(
+	UDynamicMesh* TargetMesh,
+	FGeometryScriptSimpleCollision& SimpleCollision,
+	FGeometryScriptPrimitiveOptions PrimitiveOptions,
+	FTransform Transform,
+	FBox Box,
+	int32 StepsX,
+	int32 StepsY,
+	int32 StepsZ,
+	UGeometryScriptDebug* Debug)
+{
+	return UE::Private::AppendMeshHelper::AppendBoundingBox(TargetMesh, &SimpleCollision, PrimitiveOptions, Transform, Box, StepsX, StepsY, StepsZ, Debug);
 }
 
 
@@ -280,26 +486,22 @@ UDynamicMesh* UGeometryScriptLibrary_MeshPrimitiveFunctions::AppendSphereLatLong
 	EGeometryScriptPrimitiveOriginMode Origin,
 	UGeometryScriptDebug* Debug)
 {
-	if (TargetMesh == nullptr)
-	{
-		UE::Geometry::AppendError(Debug, EGeometryScriptErrorType::InvalidInputs, LOCTEXT("PrimitiveFunctions_AppendSphereLatLong", "AppendSphereLatLong: TargetMesh is Null"));
-		return TargetMesh;
-	}
-
-	FSphereGenerator SphereGenerator;
-	SphereGenerator.Radius = FMath::Max(FMathf::ZeroTolerance, Radius);
-	SphereGenerator.NumPhi = FMath::Max(3, StepsPhi);
-	SphereGenerator.NumTheta = FMath::Max(3, StepsTheta);
-	SphereGenerator.bPolygroupPerQuad = (PrimitiveOptions.PolygroupMode == EGeometryScriptPrimitivePolygroupMode::PerQuad);
-	SphereGenerator.Generate();
-
-	FVector3d OriginShift = (Origin == EGeometryScriptPrimitiveOriginMode::Base) ? FVector3d(0, 0, Radius) : FVector3d::Zero();
-	AppendPrimitive(TargetMesh, &SphereGenerator, Transform, PrimitiveOptions, OriginShift);
-
-	return TargetMesh;
+	return UE::Private::AppendMeshHelper::AppendSphereLatLong(TargetMesh, nullptr, PrimitiveOptions, Transform, Radius, StepsPhi, StepsTheta, Origin, Debug);
 }
 
-
+UDynamicMesh* UGeometryScriptLibrary_MeshPrimitiveFunctions::AppendSphereLatLongWithCollision(
+	UDynamicMesh* TargetMesh,
+	FGeometryScriptSimpleCollision& SimpleCollision,
+	FGeometryScriptPrimitiveOptions PrimitiveOptions,
+	FTransform Transform,
+	float Radius,
+	int32 StepsPhi,
+	int32 StepsTheta,
+	EGeometryScriptPrimitiveOriginMode Origin,
+	UGeometryScriptDebug* Debug)
+{
+	return UE::Private::AppendMeshHelper::AppendSphereLatLong(TargetMesh, &SimpleCollision, PrimitiveOptions, Transform, Radius, StepsPhi, StepsTheta, Origin, Debug);
+}
 
 UDynamicMesh* UGeometryScriptLibrary_MeshPrimitiveFunctions::AppendSphereBox(
 	UDynamicMesh* TargetMesh,
@@ -312,24 +514,23 @@ UDynamicMesh* UGeometryScriptLibrary_MeshPrimitiveFunctions::AppendSphereBox(
 	EGeometryScriptPrimitiveOriginMode Origin,
 	UGeometryScriptDebug* Debug)
 {
-	if (TargetMesh == nullptr)
-	{
-		UE::Geometry::AppendError(Debug, EGeometryScriptErrorType::InvalidInputs, LOCTEXT("PrimitiveFunctions_AppendSphereBox", "AppendSphereBox: TargetMesh is Null"));
-		return TargetMesh;
-	}
-
-	FBoxSphereGenerator SphereGenerator;
-	SphereGenerator.Radius = FMath::Max(FMathf::ZeroTolerance, Radius);
-	SphereGenerator.EdgeVertices = FIndex3i(FMath::Max(0, StepsX), FMath::Max(0, StepsY), FMath::Max(0, StepsZ));
-	SphereGenerator.bPolygroupPerQuad = (PrimitiveOptions.PolygroupMode == EGeometryScriptPrimitivePolygroupMode::PerQuad);
-	SphereGenerator.Generate();
-
-	FVector3d OriginShift = (Origin == EGeometryScriptPrimitiveOriginMode::Base) ? FVector3d(0, 0, Radius) : FVector3d::Zero();
-	AppendPrimitive(TargetMesh, &SphereGenerator, Transform, PrimitiveOptions, OriginShift);
-
-	return TargetMesh;
+	return UE::Private::AppendMeshHelper::AppendSphereBox(TargetMesh, nullptr, PrimitiveOptions, Transform, Radius, StepsX, StepsY, StepsZ, Origin, Debug);
 }
 
+UDynamicMesh* UGeometryScriptLibrary_MeshPrimitiveFunctions::AppendSphereBoxWithCollision(
+	UDynamicMesh* TargetMesh,
+	FGeometryScriptSimpleCollision& SimpleCollision,
+	FGeometryScriptPrimitiveOptions PrimitiveOptions,
+	FTransform Transform,
+	float Radius,
+	int32 StepsX,
+	int32 StepsY,
+	int32 StepsZ,
+	EGeometryScriptPrimitiveOriginMode Origin,
+	UGeometryScriptDebug* Debug)
+{
+	return UE::Private::AppendMeshHelper::AppendSphereBox(TargetMesh, &SimpleCollision, PrimitiveOptions, Transform, Radius, StepsX, StepsY, StepsZ, Origin, Debug);
+}
 
 
 
@@ -345,26 +546,23 @@ UDynamicMesh* UGeometryScriptLibrary_MeshPrimitiveFunctions::AppendCapsule(
 	EGeometryScriptPrimitiveOriginMode Origin,
 	UGeometryScriptDebug* Debug)
 {
-	if (TargetMesh == nullptr)
-	{
-		UE::Geometry::AppendError(Debug, EGeometryScriptErrorType::InvalidInputs, LOCTEXT("PrimitiveFunctions_AppendCapsule", "AppendCapsule: TargetMesh is Null"));
-		return TargetMesh;
-	}
-
-	FCapsuleGenerator CapsuleGenerator;
-	CapsuleGenerator.Radius = FMath::Max(FMathf::ZeroTolerance, Radius);
-	CapsuleGenerator.SegmentLength = FMath::Max(FMathf::ZeroTolerance, LineLength);
-	CapsuleGenerator.NumHemisphereArcSteps = FMath::Max(2, HemisphereSteps);
-	CapsuleGenerator.NumCircleSteps = FMath::Max(3, CircleSteps);
-	CapsuleGenerator.NumSegmentSteps = FMath::Max(0, SegmentSteps);
-	CapsuleGenerator.bPolygroupPerQuad = (PrimitiveOptions.PolygroupMode == EGeometryScriptPrimitivePolygroupMode::PerQuad);
-	CapsuleGenerator.Generate();
-
-	FVector3d OriginShift = FVector3d::Zero();
-	OriginShift.Z = (Origin == EGeometryScriptPrimitiveOriginMode::Center) ? -(LineLength / 2) : Radius;
-	AppendPrimitive(TargetMesh, &CapsuleGenerator, Transform, PrimitiveOptions, OriginShift);
+	return UE::Private::AppendMeshHelper::AppendCapsule(TargetMesh, nullptr, PrimitiveOptions, Transform, Radius, LineLength, HemisphereSteps, CircleSteps, SegmentSteps, Origin, Debug);
+}
 
-	return TargetMesh;
+UDynamicMesh* UGeometryScriptLibrary_MeshPrimitiveFunctions::AppendCapsuleWithCollision(
+	UDynamicMesh* TargetMesh,
+	FGeometryScriptSimpleCollision& SimpleCollision,
+	FGeometryScriptPrimitiveOptions PrimitiveOptions,
+	FTransform Transform,
+	float Radius,
+	float LineLength,
+	int32 HemisphereSteps,
+	int32 CircleSteps,
+	int32 SegmentSteps,
+	EGeometryScriptPrimitiveOriginMode Origin,
+	UGeometryScriptDebug* Debug)
+{
+	return UE::Private::AppendMeshHelper::AppendCapsule(TargetMesh, &SimpleCollision, PrimitiveOptions, Transform, Radius, LineLength, HemisphereSteps, CircleSteps, SegmentSteps, Origin, Debug);
 }
 
 
diff --git a/Engine/Plugins/Runtime/GeometryScripting/Source/GeometryScriptingCore/Public/GeometryScript/MeshPrimitiveFunctions.h b/Engine/Plugins/Runtime/GeometryScripting/Source/GeometryScriptingCore/Public/GeometryScript/MeshPrimitiveFunctions.h
--- a/Engine/Plugins/Runtime/GeometryScripting/Source/GeometryScriptingCore/Public/GeometryScript/MeshPrimitiveFunctions.h
+++ b/Engine/Plugins/Runtime/GeometryScripting/Source/GeometryScriptingCore/Public/GeometryScript/MeshPrimitiveFunctions.h
@@ -172,6 +172,26 @@ public:
 		EGeometryScriptPrimitiveOriginMode Origin = EGeometryScriptPrimitiveOriginMode::Base,
 		UGeometryScriptDebug* Debug = nullptr);
 
+	/**
+	 * Appends a 3D box to the Target Mesh
+	 * Also creates matching simple collision
+	 */
+	UFUNCTION(BlueprintCallable, Category = "GeometryScript|Primitives", meta = (ScriptMethod))
+	static UPARAM(DisplayName = "Target Mesh") UDynamicMesh*
+	AppendBoxWithCollision(
+		UDynamicMesh* TargetMesh,
+		FGeometryScriptSimpleCollision &SimpleCollision,
+		FGeometryScriptPrimitiveOptions PrimitiveOptions,
+		FTransform Transform,
+		float DimensionX = 100,
+		float DimensionY = 100,
+		float DimensionZ = 100,
+		int32 StepsX = 0,
+		int32 StepsY = 0,
+		int32 StepsZ = 0,
+		EGeometryScriptPrimitiveOriginMode Origin = EGeometryScriptPrimitiveOriginMode::Base,
+		UGeometryScriptDebug* Debug = nullptr);
+
 	/**
 	 * Appends a 3D box to the Target Mesh with dimensions and origin taken from the input Box
 	 */
@@ -187,6 +207,23 @@ public:
 		int32 StepsZ = 0,
 		UGeometryScriptDebug* Debug = nullptr);
 
+	/**
+	 * Appends a 3D box to the Target Mesh with dimensions and origin taken from the input Box
+	 * Also creates matching simple collision
+	 */
+	UFUNCTION(BlueprintCallable, Category = "GeometryScript|Primitives", meta=(ScriptMethod))
+	static UPARAM(DisplayName = "Target Mesh") UDynamicMesh* 
+	AppendBoundingBoxWithCollision( 
+		UDynamicMesh* TargetMesh, 
+		FGeometryScriptSimpleCollision& SimpleCollision,
+		FGeometryScriptPrimitiveOptions PrimitiveOptions,
+		FTransform Transform,
+		FBox Box,
+		int32 StepsX = 0,
+		int32 StepsY = 0,
+		int32 StepsZ = 0,
+		UGeometryScriptDebug* Debug = nullptr);
+
 
 	/**
 	* Appends a 3D Sphere triangulated using latitude/longitude topology to the Target Mesh.
@@ -202,6 +239,24 @@ public:
 		int32 StepsTheta = 16,
 		EGeometryScriptPrimitiveOriginMode Origin = EGeometryScriptPrimitiveOriginMode::Center,
 		UGeometryScriptDebug* Debug = nullptr);
+	
+	/**
+	* Appends a 3D Sphere triangulated using latitude/longitude topology to the Target Mesh.
+	* Also creates matching simple collision. Note: If transform has non-uniform scale, collision shape will be a uniform-scale approximation
+	* matching the behavior of the physics system -- specifically, it will scale the sphere radius by the smallest axis scale.
+	*/
+	UFUNCTION(BlueprintCallable, Category = "GeometryScript|Primitives", meta=(ScriptMethod))
+	static UPARAM(DisplayName = "Target Mesh") UDynamicMesh* 
+	AppendSphereLatLongWithCollision( 
+		UDynamicMesh* TargetMesh, 
+		FGeometryScriptSimpleCollision& SimpleCollision,
+		FGeometryScriptPrimitiveOptions PrimitiveOptions,
+		FTransform Transform,
+		float Radius = 50,
+		int32 StepsPhi = 10,
+		int32 StepsTheta = 16,
+		EGeometryScriptPrimitiveOriginMode Origin = EGeometryScriptPrimitiveOriginMode::Center,
+		UGeometryScriptDebug* Debug = nullptr);
 
 	/**
 	* Appends a 3D sphere triangulated using box topology to the Target Mesh.
@@ -219,6 +274,25 @@ public:
 		EGeometryScriptPrimitiveOriginMode Origin = EGeometryScriptPrimitiveOriginMode::Center,
 		UGeometryScriptDebug* Debug = nullptr);
 
+	/**
+	 * Appends a 3D sphere triangulated using box topology to the Target Mesh.
+	 * Also creates matching simple collision. Note: If transform has non-uniform scale, collision shape will be a uniform-scale approximation
+	 * matching the behavior of the physics system -- specifically, it will scale the sphere radius by the smallest axis scale.
+	 */
+	UFUNCTION(BlueprintCallable, Category = "GeometryScript|Primitives", meta = (ScriptMethod))
+	static UPARAM(DisplayName = "Target Mesh") UDynamicMesh*
+	AppendSphereBoxWithCollision(
+		UDynamicMesh* TargetMesh,
+		FGeometryScriptSimpleCollision& SimpleCollision,
+		FGeometryScriptPrimitiveOptions PrimitiveOptions,
+		FTransform Transform,
+		float Radius = 50,
+		int32 StepsX = 6,
+		int32 StepsY = 6,
+		int32 StepsZ = 6,
+		EGeometryScriptPrimitiveOriginMode Origin = EGeometryScriptPrimitiveOriginMode::Center,
+		UGeometryScriptDebug* Debug = nullptr);
+
 	/**
 	* Appends a 3D Capsule to the Target Mesh.
 	*/ 
@@ -236,6 +310,26 @@ public:
 		EGeometryScriptPrimitiveOriginMode Origin = EGeometryScriptPrimitiveOriginMode::Base,
 		UGeometryScriptDebug* Debug = nullptr);
 
+	/**
+	 * Appends a 3D Capsule to the Target Mesh.
+	 * Also creates matching simple collision. Note: If transform has non-uniform scale, collision shape will be a uniform-scale approximation
+	 * matching the behavior of the physics system -- specifically, it will scale the radius by the larger of the X, Y axis scales, and the length by the Z axis scale
+	*/ 
+	UFUNCTION(BlueprintCallable, Category = "GeometryScript|Primitives", meta=(ScriptMethod))
+	static UPARAM(DisplayName = "Target Mesh") UDynamicMesh* 
+	AppendCapsuleWithCollision( 
+		UDynamicMesh* TargetMesh, 
+		FGeometryScriptSimpleCollision& SimpleCollision,
+		FGeometryScriptPrimitiveOptions PrimitiveOptions,
+		FTransform Transform,
+		float Radius = 30,
+		float LineLength = 75,
+		int32 HemisphereSteps = 5,
+		int32 CircleSteps = 8,
+		int32 SegmentSteps = 0,
+		EGeometryScriptPrimitiveOriginMode Origin = EGeometryScriptPrimitiveOriginMode::Base,
+		UGeometryScriptDebug* Debug = nullptr);
+
 	// Version of append capsule without SegmentSteps parameter
 	UE_DEPRECATED(5.5, "Use AppendCapsule with the SegmentSteps parameter, instead")
 	static UDynamicMesh* AppendCapsule(
