add 'reset' and 'combine array' methods to the geometry script simple collision shape management functions

diff --git a/Engine/Plugins/Runtime/GeometryScripting/Source/GeometryScriptingCore/Private/CollisionFunctions.cpp b/Engine/Plugins/Runtime/GeometryScripting/Source/GeometryScriptingCore/Private/CollisionFunctions.cpp
--- a/Engine/Plugins/Runtime/GeometryScripting/Source/GeometryScriptingCore/Private/CollisionFunctions.cpp
+++ b/Engine/Plugins/Runtime/GeometryScripting/Source/GeometryScriptingCore/Private/CollisionFunctions.cpp
@@ -525,6 +525,11 @@ void UGeometryScriptLibrary_CollisionFunctions::ResetDynamicMeshCollision(
 
 }
 
+void UGeometryScriptLibrary_CollisionFunctions::ResetSimpleCollision(FGeometryScriptSimpleCollision & SimpleCollision)
+{
+	SimpleCollision.AggGeom.EmptyElements();
+}
+
 
 FGeometryScriptSimpleCollision
 UGeometryScriptLibrary_CollisionFunctions::GenerateCollisionFromMesh(
@@ -914,6 +919,19 @@ void UGeometryScriptLibrary_CollisionFunctions::CombineSimpleCollision(
 	Collision.AggGeom.SkinnedLevelSetElems.Append(AppendCollision.AggGeom.SkinnedLevelSetElems);
 }
 
+
+void UGeometryScriptLibrary_CollisionFunctions::CombineSimpleCollisionArray(
+	const TArray<FGeometryScriptSimpleCollision>& SimpleCollisionArray,
+	FGeometryScriptSimpleCollision& SimpleCollision,
+	UGeometryScriptDebug* Debug
+)
+{
+	for (const FGeometryScriptSimpleCollision& Coll : SimpleCollisionArray)
+	{
+		CombineSimpleCollision(SimpleCollision, Coll, Debug);
+	}
+}
+
 void UGeometryScriptLibrary_CollisionFunctions::SimplifyConvexHulls(
 	FGeometryScriptSimpleCollision& SimpleCollision,
 	const FGeometryScriptConvexHullSimplificationOptions& SimplifyOptions,
diff --git a/Engine/Plugins/Runtime/GeometryScripting/Source/GeometryScriptingCore/Public/GeometryScript/CollisionFunctions.h b/Engine/Plugins/Runtime/GeometryScripting/Source/GeometryScriptingCore/Public/GeometryScript/CollisionFunctions.h
--- a/Engine/Plugins/Runtime/GeometryScripting/Source/GeometryScriptingCore/Public/GeometryScript/CollisionFunctions.h
+++ b/Engine/Plugins/Runtime/GeometryScripting/Source/GeometryScriptingCore/Public/GeometryScript/CollisionFunctions.h
@@ -402,6 +402,12 @@ public:
 		bool bEmitTransaction = false,
 		UGeometryScriptDebug* Debug = nullptr);
 
+	/**
+	 * Clears the Simple Collision shapes
+	 */
+	UFUNCTION(BlueprintCallable, Category = "GeometryScript|Collision")
+	static void ResetSimpleCollision(UPARAM(ref) FGeometryScriptSimpleCollision& SimpleCollision);
+
 	/*
 	 * Get the simple collision from a Primitive Component
 	 */
@@ -482,6 +488,16 @@ public:
 		const FGeometryScriptSimpleCollision& AppendCollision,
 		UGeometryScriptDebug* Debug = nullptr
 	);
+
+	/*
+	 * Combine the SimpleCollisionArray collision shapes into a single SimpleCollision
+	 */
+	UFUNCTION(BlueprintCallable, Category = "GeometryScript|Collision")
+	static void CombineSimpleCollisionArray(
+		const TArray<FGeometryScriptSimpleCollision>& SimpleCollisionArray,
+		FGeometryScriptSimpleCollision& SimpleCollision,
+		UGeometryScriptDebug* Debug = nullptr
+	);
 	
 	/**
 	 * Simplify any convex hulls in the given simple collision representation. Updates the passed-in Simple Collision.
