move geometry script out of beta

diff --git a/Engine/Plugins/Runtime/GeometryScripting/GeometryScripting.uplugin b/Engine/Plugins/Runtime/GeometryScripting/GeometryScripting.uplugin
--- a/Engine/Plugins/Runtime/GeometryScripting/GeometryScripting.uplugin
+++ b/Engine/Plugins/Runtime/GeometryScripting/GeometryScripting.uplugin
@@ -12,7 +12,7 @@
 	"SupportURL": "",
 	"CanContainContent": true,
 	"IsExperimentalVersion": false,
-	"IsBetaVersion":  true,
+	"IsBetaVersion":  false,
 	"Installed": false,
 	"Modules": [
 		{
