more consistent array/list naming for triangle arrays/lists:
array of integer tuples -> Triangle Tuple Array
GeometryScriptIndexList w/ triangle IDs -> Triangle ID List
GeometryScriptTriangleList -> Triangle List

diff --git a/Engine/Plugins/Runtime/GeometryScripting/Source/GeometryScriptingCore/Public/GeometryScript/ListUtilityFunctions.h b/Engine/Plugins/Runtime/GeometryScripting/Source/GeometryScriptingCore/Public/GeometryScript/ListUtilityFunctions.h
--- a/Engine/Plugins/Runtime/GeometryScripting/Source/GeometryScriptingCore/Public/GeometryScript/ListUtilityFunctions.h
+++ b/Engine/Plugins/Runtime/GeometryScripting/Source/GeometryScriptingCore/Public/GeometryScript/ListUtilityFunctions.h
@@ -90,13 +90,13 @@ public:
 	* Converts Triangle List to Triangle Array by populating with the appropriate integer triplets.
 	*/
 	UFUNCTION(BlueprintCallable, Category = "GeometryScript|ListUtils", meta=(ScriptMethod))
-	static void ConvertTriangleListToArray(FGeometryScriptTriangleList TriangleList, TArray<FIntVector>& TriangleArray);
+	static void ConvertTriangleListToArray(FGeometryScriptTriangleList TriangleList, TArray<FIntVector>& TriangleTupleArray);
 
 	/**
 	* Converts a Triangle Array of integer triplets to a Triangle List.
 	*/
 	UFUNCTION(BlueprintCallable, Category = "GeometryScript|ListUtils")
-	static void ConvertArrayToTriangleList(const TArray<FIntVector>& TriangleArray, FGeometryScriptTriangleList& TriangleList);
+	static void ConvertArrayToTriangleList(const TArray<FIntVector>& TriangleTupleArray, FGeometryScriptTriangleList& TriangleList);
 
 
 	/**
diff --git a/Engine/Plugins/Runtime/GeometryScripting/Source/GeometryScriptingCore/Public/GeometryScript/MeshBasicEditFunctions.h b/Engine/Plugins/Runtime/GeometryScripting/Source/GeometryScriptingCore/Public/GeometryScript/MeshBasicEditFunctions.h
--- a/Engine/Plugins/Runtime/GeometryScripting/Source/GeometryScriptingCore/Public/GeometryScript/MeshBasicEditFunctions.h
+++ b/Engine/Plugins/Runtime/GeometryScripting/Source/GeometryScriptingCore/Public/GeometryScript/MeshBasicEditFunctions.h
@@ -212,7 +212,7 @@ public:
 	static UPARAM(DisplayName = "Target Mesh") UDynamicMesh* 
 	DeleteTrianglesFromMesh( 
 		UDynamicMesh* TargetMesh, 
-		FGeometryScriptIndexList TriangleList,
+		FGeometryScriptIndexList TriangleIDList,
 		int& NumDeleted,
 		bool bDeferChangeNotifications = false );
 
diff --git a/Engine/Plugins/Runtime/GeometryScripting/Source/GeometryScriptingCore/Public/GeometryScript/MeshDecompositionFunctions.h b/Engine/Plugins/Runtime/GeometryScripting/Source/GeometryScriptingCore/Public/GeometryScript/MeshDecompositionFunctions.h
--- a/Engine/Plugins/Runtime/GeometryScripting/Source/GeometryScriptingCore/Public/GeometryScript/MeshDecompositionFunctions.h
+++ b/Engine/Plugins/Runtime/GeometryScripting/Source/GeometryScriptingCore/Public/GeometryScript/MeshDecompositionFunctions.h
@@ -72,7 +72,7 @@ public:
 	GetSubMeshFromMesh(  
 		UDynamicMesh* TargetMesh, 
 		UPARAM(DisplayName = "Copy To Submesh", ref) UDynamicMesh* StoreToSubmesh, 
-		FGeometryScriptIndexList TriangleList,
+		FGeometryScriptIndexList TriangleIDList,
 		UPARAM(DisplayName = "Copy To Submesh") UDynamicMesh*& StoreToSubmeshOut, 
 		UGeometryScriptDebug* Debug = nullptr);
 
