redo geometry script triangle array / lists name changes via meta tags to not break existing BPs

diff --git a/Engine/Plugins/Runtime/GeometryScripting/Source/GeometryScriptingCore/Public/GeometryScript/ListUtilityFunctions.h b/Engine/Plugins/Runtime/GeometryScripting/Source/GeometryScriptingCore/Public/GeometryScript/ListUtilityFunctions.h
--- a/Engine/Plugins/Runtime/GeometryScripting/Source/GeometryScriptingCore/Public/GeometryScript/ListUtilityFunctions.h
+++ b/Engine/Plugins/Runtime/GeometryScripting/Source/GeometryScriptingCore/Public/GeometryScript/ListUtilityFunctions.h
@@ -90,13 +90,13 @@ public:
 	* Converts Triangle List to Triangle Array by populating with the appropriate integer triplets.
 	*/
 	UFUNCTION(BlueprintCallable, Category = "GeometryScript|ListUtils", meta=(ScriptMethod))
-	static void ConvertTriangleListToArray(FGeometryScriptTriangleList TriangleList, TArray<FIntVector>& TriangleTupleArray);
+	static void ConvertTriangleListToArray(FGeometryScriptTriangleList TriangleList, UPARAM(DisplayName = "Triangle Tuple Array") TArray<FIntVector>& TriangleArray);
 
 	/**
 	* Converts a Triangle Array of integer triplets to a Triangle List.
 	*/
 	UFUNCTION(BlueprintCallable, Category = "GeometryScript|ListUtils")
-	static void ConvertArrayToTriangleList(const TArray<FIntVector>& TriangleTupleArray, FGeometryScriptTriangleList& TriangleList);
+	static void ConvertArrayToTriangleList(UPARAM(DisplayName = "Triangle Tuple Array") const TArray<FIntVector>& TriangleArray, FGeometryScriptTriangleList& TriangleList);
 
 
 	/**
diff --git a/Engine/Plugins/Runtime/GeometryScripting/Source/GeometryScriptingCore/Public/GeometryScript/MeshBasicEditFunctions.h b/Engine/Plugins/Runtime/GeometryScripting/Source/GeometryScriptingCore/Public/GeometryScript/MeshBasicEditFunctions.h
--- a/Engine/Plugins/Runtime/GeometryScripting/Source/GeometryScriptingCore/Public/GeometryScript/MeshBasicEditFunctions.h
+++ b/Engine/Plugins/Runtime/GeometryScripting/Source/GeometryScriptingCore/Public/GeometryScript/MeshBasicEditFunctions.h
@@ -212,7 +212,7 @@ public:
 	static UPARAM(DisplayName = "Target Mesh") UDynamicMesh* 
 	DeleteTrianglesFromMesh( 
 		UDynamicMesh* TargetMesh, 
-		FGeometryScriptIndexList TriangleIDList,
+		UPARAM(DisplayName = "Triangle ID List") FGeometryScriptIndexList TriangleList,
 		int& NumDeleted,
 		bool bDeferChangeNotifications = false );
 
diff --git a/Engine/Plugins/Runtime/GeometryScripting/Source/GeometryScriptingCore/Public/GeometryScript/MeshDecompositionFunctions.h b/Engine/Plugins/Runtime/GeometryScripting/Source/GeometryScriptingCore/Public/GeometryScript/MeshDecompositionFunctions.h
--- a/Engine/Plugins/Runtime/GeometryScripting/Source/GeometryScriptingCore/Public/GeometryScript/MeshDecompositionFunctions.h
+++ b/Engine/Plugins/Runtime/GeometryScripting/Source/GeometryScriptingCore/Public/GeometryScript/MeshDecompositionFunctions.h
@@ -72,7 +72,7 @@ public:
 	GetSubMeshFromMesh(  
 		UDynamicMesh* TargetMesh, 
 		UPARAM(DisplayName = "Copy To Submesh", ref) UDynamicMesh* StoreToSubmesh, 
-		FGeometryScriptIndexList TriangleIDList,
+		UPARAM(DisplayName = "Triangle ID List") FGeometryScriptIndexList TriangleList,
 		UPARAM(DisplayName = "Copy To Submesh") UDynamicMesh*& StoreToSubmeshOut, 
 		UGeometryScriptDebug* Debug = nullptr);
 
