Add more UV stats methods to geometry script: Get Mesh UV Area and Get Num UV Islands

diff --git a/Engine/Plugins/Runtime/GeometryScripting/Source/GeometryScriptingCore/Private/MeshQueryFunctions.cpp b/Engine/Plugins/Runtime/GeometryScripting/Source/GeometryScriptingCore/Private/MeshQueryFunctions.cpp
--- a/Engine/Plugins/Runtime/GeometryScripting/Source/GeometryScriptingCore/Private/MeshQueryFunctions.cpp
+++ b/Engine/Plugins/Runtime/GeometryScripting/Source/GeometryScriptingCore/Private/MeshQueryFunctions.cpp
@@ -447,6 +447,25 @@ FBox2D UGeometryScriptLibrary_MeshQueryFunctions::GetUVSetBoundingBox(UDynamicMe
 	});
 }
 
+double UGeometryScriptLibrary_MeshQueryFunctions::GetMeshUVArea(UDynamicMesh* TargetMesh, int UVChannelIndex, bool& bIsValidUVChannel)
+{
+	return SimpleMeshUVSetQuery<double>(TargetMesh, UVChannelIndex, bIsValidUVChannel, 0.0,
+		[](const FDynamicMesh3& Mesh, const FDynamicMeshUVOverlay& UVOverlay) -> double
+		{
+			double AreaSum = 0;
+			for (int32 TriangleID : Mesh.TriangleIndicesItr())
+			{
+				FTriangle2f Tri;
+				if (UVOverlay.IsSetTriangle(TriangleID))
+				{
+					UVOverlay.GetTriElements(TriangleID, Tri.V[0], Tri.V[1], Tri.V[2]);
+					AreaSum += (double)Tri.Area();
+				}
+			}
+			return AreaSum;
+		});
+}
+
 void UGeometryScriptLibrary_MeshQueryFunctions::GetTriangleUVs(UDynamicMesh* TargetMesh, int32 UVSetIndex, int32 TriangleID, 
 	FVector2D& UV1, FVector2D& UV2, FVector2D& UV3, bool& bHaveValidUVs)
 {
@@ -526,6 +545,28 @@ UDynamicMesh* UGeometryScriptLibrary_MeshQueryFunctions::GetAllUVSeamEdges(UDyna
 	return TargetMesh;
 }
 
+int32 UGeometryScriptLibrary_MeshQueryFunctions::GetNumUVIslands(UDynamicMesh* TargetMesh, int32 UVChannel, bool& bIsValidUVChannel)
+{
+	return SimpleMeshUVSetQuery<int32>(TargetMesh, UVChannel, bIsValidUVChannel, -1,
+		[](const FDynamicMesh3& Mesh, const FDynamicMeshUVOverlay& UVOverlay) -> int32
+		{
+			auto UVIslandPredicate = [&UVOverlay](int32 Triangle0, int32 Triangle1)
+				{
+					return UVOverlay.AreTrianglesConnected(Triangle0, Triangle1);
+				};
+
+			FMeshConnectedComponents UVComponents(&Mesh);
+			UVComponents.FindConnectedTriangles(UVIslandPredicate);
+			int32 NumIslands = 0;
+			// only count the islands with set triangles
+			for (const FMeshConnectedComponents::FComponent& Component : UVComponents.Components)
+			{
+				NumIslands += (int32)(!Component.Indices.IsEmpty() && UVOverlay.IsSetTriangle(Component.Indices[0]));
+			}
+			return NumIslands;
+		});
+}
+
 
 bool UGeometryScriptLibrary_MeshQueryFunctions::GetHasTriangleNormals( UDynamicMesh* TargetMesh )
 {
diff --git a/Engine/Plugins/Runtime/GeometryScripting/Source/GeometryScriptingCore/Public/GeometryScript/MeshQueryFunctions.h b/Engine/Plugins/Runtime/GeometryScripting/Source/GeometryScriptingCore/Public/GeometryScript/MeshQueryFunctions.h
--- a/Engine/Plugins/Runtime/GeometryScripting/Source/GeometryScriptingCore/Public/GeometryScript/MeshQueryFunctions.h
+++ b/Engine/Plugins/Runtime/GeometryScripting/Source/GeometryScriptingCore/Public/GeometryScript/MeshQueryFunctions.h
@@ -262,6 +262,20 @@ public:
 																			UPARAM(DisplayName = "UV Channel") int UVSetIndex, 
 																			UPARAM(DisplayName = "Is Valid UV Channel") bool& bIsValidUVSet, 
 																			UPARAM(DisplayName = "UV Channel Is Empty") bool& bUVSetIsEmpty );
+
+	/**
+	 * Gets the area of triangles in UV space for the given UV Channel.
+	 * 
+	 * @param TargetMesh The mesh to query.
+	 * @param UVChannel The UV channel to query
+	 * @param bIsValidUVChannel True, if the mesh has UVs for the given UVSetIndex.
+	 * @return The number of UV islands
+	 */
+	UFUNCTION(BlueprintCallable, Category = "GeometryScript|MeshQueries", meta = (ScriptMethod), DisplayName = "Get Mesh UV Area")
+	static UPARAM(DisplayName = "UV Area") double GetMeshUVArea(UDynamicMesh* TargetMesh,
+		UPARAM(DisplayName = "UV Channel") int UVChannel,
+		UPARAM(DisplayName = "Is Valid UV Channel") bool& bIsValidUVChannel);
+
 	/**
 	* Returns the UV values associated with the three vertices of the triangle in the specified UV Channel.
 	* If the Triangle does not exist in the mesh or if no UVs are set in the specified UV Channel for the triangle, the resulting values will be (0,0) and bHaveValidUVs will be set to false.
@@ -307,6 +321,20 @@ public:
 	static UPARAM(DisplayName = "Target Mesh") UDynamicMesh*
 	GetAllUVSeamEdges( UDynamicMesh* TargetMesh, UPARAM(DisplayName = "UV Channel") int32 UVSetIndex, bool& bHaveValidUVs, FGeometryScriptIndexList& ElementIDs);
 
+	/**
+	 * Returns the number of UV islands in a given UV channel.
+	 *
+	 * @param TargetMesh The mesh to query.
+	 * @param UVChannel The UV channel to query
+	 * @param bIsValidUVChannel True, if the mesh has UVs for the given UVSetIndex.
+	 * @return The number of UV islands
+	 */
+	UFUNCTION(BlueprintCallable, Category = "GeometryScript|MeshQueries", meta = (ScriptMethod), DisplayName = "Get Num UV Islands")
+	static UPARAM(DisplayName = "Num UV Islands") int32 GetNumUVIslands(
+		UDynamicMesh* TargetMesh, 
+		UPARAM(DisplayName = "UV Channel") int32 UVChannel, 
+		UPARAM(DisplayName = "Is Valid UV Channel") bool& bIsValidUVChannel);
+
 	
 	//
 	// Normal queries
