Add 'With Materials' versions of the geometry script append mesh methods, which take the materials arrays of the target and append meshes and makes the result mesh reference a combined materials array

diff --git a/Engine/Plugins/Runtime/GeometryScripting/Source/GeometryScriptingCore/Private/MeshBasicEditFunctions.cpp b/Engine/Plugins/Runtime/GeometryScripting/Source/GeometryScriptingCore/Private/MeshBasicEditFunctions.cpp
--- a/Engine/Plugins/Runtime/GeometryScripting/Source/GeometryScriptingCore/Private/MeshBasicEditFunctions.cpp
+++ b/Engine/Plugins/Runtime/GeometryScripting/Source/GeometryScriptingCore/Private/MeshBasicEditFunctions.cpp
@@ -379,49 +379,265 @@ void FGeometryScriptAppendMeshOptions::UpdateAttributesForCombineMode(FDynamicMe
 	// else the mode is UseTarget, which already corresponds to the default behavior for AppendMesh
 }
 
-UDynamicMesh* UGeometryScriptLibrary_MeshBasicEditFunctions::AppendMesh(
-	UDynamicMesh* TargetMesh,
-	UDynamicMesh* AppendMesh,
-	FTransform AppendTransform,
-	bool bDeferChangeNotifications,
-	FGeometryScriptAppendMeshOptions AppendOptions,
-	UGeometryScriptDebug* Debug)
+namespace UE::Private::AppendHelpers
 {
-	if (TargetMesh == nullptr)
+	static void AppendMaterials(const TArray<UMaterialInterface*>& TargetMeshMaterialList, const TArray<UMaterialInterface*>& AppendMeshMaterialList, bool bCompactAppendedMaterials,
+		TArray<int32>& OutAppendedMaterialRemap, TArray<UMaterialInterface*>& OutMaterialList)
 	{
-		UE::Geometry::AppendError(Debug, EGeometryScriptErrorType::InvalidInputs, LOCTEXT("AppendMesh_InvalidInput1", "AppendMesh: TargetMesh is Null"));
-		return TargetMesh;
+		OutMaterialList = TargetMeshMaterialList;
+		const int32 NumAppend = AppendMeshMaterialList.Num();
+		OutAppendedMaterialRemap.Init(-1, NumAppend);
+
+		for (int32 Idx = 0; Idx < NumAppend; ++Idx)
+		{
+			int32 FoundIdx;
+			if (bCompactAppendedMaterials && TargetMeshMaterialList.Find(AppendMeshMaterialList[Idx], FoundIdx))
+			{
+				OutAppendedMaterialRemap[Idx] = FoundIdx;
+			}
+			else
+			{
+				OutAppendedMaterialRemap[Idx] = OutMaterialList.Add(AppendMeshMaterialList[Idx]);
+			}
+		}
 	}
-	if (AppendMesh == nullptr)
+
+	static void ApplyMaterialRemap(const FMeshIndexMappings& Mappings, const TArray<int32>& AppendMaterialRemap, FDynamicMesh3& Mesh)
 	{
-		UE::Geometry::AppendError(Debug, EGeometryScriptErrorType::InvalidInputs, LOCTEXT("AppendMesh_InvalidInput2", "AppendMesh: AppendMesh is Null"));
-		return TargetMesh;
+		if (!Mesh.HasAttributes() || !Mesh.Attributes()->HasMaterialID())
+		{
+			return;
+		}
+
+		FDynamicMeshMaterialAttribute* MaterialAttrib = Mesh.Attributes()->GetMaterialID();
+		for (const TPair<int32, int32>& MapTID : Mappings.GetTriangleMap().GetForwardMap())
+		{
+			int32 OrigMID = MaterialAttrib->GetValue(MapTID.Value);
+			if (AppendMaterialRemap.IsValidIndex(OrigMID))
+			{
+				int32 RemapMID = AppendMaterialRemap[OrigMID];
+				MaterialAttrib->SetValue(MapTID.Value, RemapMID);
+			}
+		}
 	}
 
-	TargetMesh->EditMesh([&](FDynamicMesh3& AppendToMesh)
+	static UDynamicMesh* AppendMesh(
+		UDynamicMesh* TargetMesh,
+		UDynamicMesh* AppendMesh,
+		TArray<int32>* AppendMaterialRemap,
+		FTransform AppendTransform,
+		bool bDeferChangeNotifications,
+		FGeometryScriptAppendMeshOptions AppendOptions,
+		UGeometryScriptDebug* Debug)
 	{
-		AppendMesh->ProcessMesh([&](const FDynamicMesh3& OtherMesh)
+		if (TargetMesh == nullptr)
+		{
+			UE::Geometry::AppendError(Debug, EGeometryScriptErrorType::InvalidInputs, LOCTEXT("AppendMesh_InvalidInput1", "AppendMesh: TargetMesh is Null"));
+			return TargetMesh;
+		}
+		if (AppendMesh == nullptr)
+		{
+			UE::Geometry::AppendError(Debug, EGeometryScriptErrorType::InvalidInputs, LOCTEXT("AppendMesh_InvalidInput2", "AppendMesh: AppendMesh is Null"));
+			return TargetMesh;
+		}
+
+		TargetMesh->EditMesh([&AppendOptions, &AppendMaterialRemap, &AppendMesh, &AppendTransform](FDynamicMesh3& AppendToMesh)
+		{
+			AppendMesh->ProcessMesh([&AppendOptions, &AppendMaterialRemap, &AppendToMesh, &AppendTransform](const FDynamicMesh3& OtherMesh)
+			{
+				AppendOptions.UpdateAttributesForCombineMode(AppendToMesh, OtherMesh);
+				FTransformSRT3d XForm(AppendTransform);
+				FMeshIndexMappings TmpMappings;
+				FDynamicMeshEditor Editor(&AppendToMesh);
+				const FDynamicMesh3* UseOtherMesh = &OtherMesh;
+				FDynamicMesh3 TmpMesh;
+				if (UseOtherMesh == &AppendToMesh)
+				{
+					TmpMesh = OtherMesh;	// need  to make a copy if we are appending to ourself
+					UseOtherMesh = &TmpMesh;
+				}
+				Editor.AppendMesh(UseOtherMesh, TmpMappings,
+					[&XForm](int, const FVector3d& Position) { return XForm.TransformPosition(Position); },
+					[&XForm](int, const FVector3d& Normal) { return XForm.TransformNormal(Normal); });
+				if (AppendMaterialRemap)
+				{
+					ApplyMaterialRemap(TmpMappings, *AppendMaterialRemap, AppendToMesh);
+				}
+			});
+		}, EDynamicMeshChangeType::GeneralEdit, EDynamicMeshAttributeChangeFlags::Unknown, bDeferChangeNotifications);
+
+		return TargetMesh;
+	}
+
+	static UDynamicMesh* AppendMeshTransformed(
+		UDynamicMesh* TargetMesh,
+		UDynamicMesh* AppendMesh,
+		TArray<int32>* AppendMaterialRemap,
+		const TArray<FTransform>& AppendTransforms, 
+		FTransform ConstantTransform,
+		bool bConstantTransformIsRelative,
+		bool bDeferChangeNotifications,
+		FGeometryScriptAppendMeshOptions AppendOptions,
+		UGeometryScriptDebug* Debug)
+	{
+		if (TargetMesh == nullptr)
+		{
+			UE::Geometry::AppendError(Debug, EGeometryScriptErrorType::InvalidInputs, LOCTEXT("AppendMeshTransformed_InvalidInput1", "AppendMeshTransformed: TargetMesh is Null"));
+			return TargetMesh;
+		}
+		if (AppendMesh == nullptr)
+		{
+			UE::Geometry::AppendError(Debug, EGeometryScriptErrorType::InvalidInputs, LOCTEXT("AppendMeshTransformed_InvalidInput2", "AppendMeshTransformed: AppendMesh is Null"));
+			return TargetMesh;
+		}
+		if (AppendTransforms.IsEmpty())
+		{
+			UE::Geometry::AppendError(Debug, EGeometryScriptErrorType::InvalidInputs, LOCTEXT("AppendMeshTransformed_NoTransforms", "AppendMeshTransformed: AppendTransforms array is empty"));
+			return TargetMesh;
+		}
+
+		TargetMesh->EditMesh([&AppendOptions, &AppendMaterialRemap, &AppendMesh, &AppendTransforms, bConstantTransformIsRelative, &ConstantTransform](FDynamicMesh3& AppendToMesh)
+		{
+			AppendMesh->ProcessMesh([&AppendOptions, &AppendMaterialRemap, &AppendToMesh, &AppendTransforms, bConstantTransformIsRelative, &ConstantTransform](const FDynamicMesh3& OtherMesh)
+			{
+				AppendOptions.UpdateAttributesForCombineMode(AppendToMesh, OtherMesh);
+				FMeshIndexMappings TmpMappings;
+				FDynamicMeshEditor Editor(&AppendToMesh);
+				const FDynamicMesh3* UseOtherMesh = &OtherMesh;
+				FDynamicMesh3 TmpMesh;
+				if (UseOtherMesh == &AppendToMesh)
+				{
+					TmpMesh = OtherMesh;	// need  to make a copy if we are appending to ourself
+					UseOtherMesh = &TmpMesh;
+				}
+				for (FTransform AppendTransform : AppendTransforms)
+				{
+					FTransformSequence3d TransformSequence;
+
+					if (bConstantTransformIsRelative)
+					{
+						TransformSequence.Append(ConstantTransform);
+						TransformSequence.Append(AppendTransform);
+					}
+					else
+					{
+						// want to apply the constant transform's rotate/scale after
+						// the main transform rotate/scale, so the main positioning 
+						// translation has to be deferred until after that
+						FVector Translation = AppendTransform.GetLocation();
+						AppendTransform.SetTranslation(FVector::Zero());
+
+						TransformSequence.Append(AppendTransform);
+						TransformSequence.Append(ConstantTransform);
+						TransformSequence.Append(FTransform(Translation));
+					}
+
+					Editor.AppendMesh(UseOtherMesh, TmpMappings,
+						[&TransformSequence](int, const FVector3d& Position) { return TransformSequence.TransformPosition(Position); },
+						[&TransformSequence](int, const FVector3d& Normal) { return TransformSequence.TransformNormal(Normal); });
+					if (AppendMaterialRemap)
+					{
+						ApplyMaterialRemap(TmpMappings, *AppendMaterialRemap, AppendToMesh);
+					}
+					TmpMappings.Reset();
+				}
+			});
+		}, EDynamicMeshChangeType::GeneralEdit, EDynamicMeshAttributeChangeFlags::Unknown, bDeferChangeNotifications);
+
+		return TargetMesh;
+	}
+
+	static UDynamicMesh* AppendMeshRepeated(
+		UDynamicMesh* TargetMesh,
+		UDynamicMesh* AppendMesh,
+		TArray<int32>* AppendMaterialRemap,
+		FTransform AppendTransform,
+		int32 RepeatCount,
+		bool bApplyTransformToFirstInstance,
+		bool bDeferChangeNotifications,
+		FGeometryScriptAppendMeshOptions AppendOptions,
+		UGeometryScriptDebug* Debug)
+	{
+		if (TargetMesh == nullptr)
+		{
+			UE::Geometry::AppendError(Debug, EGeometryScriptErrorType::InvalidInputs, LOCTEXT("AppendMeshRepeated_InvalidInput1", "AppendMeshRepeated: TargetMesh is Null"));
+			return TargetMesh;
+		}
+		if (AppendMesh == nullptr)
+		{
+			UE::Geometry::AppendError(Debug, EGeometryScriptErrorType::InvalidInputs, LOCTEXT("AppendMeshRepeated_InvalidInput2", "AppendMeshRepeated: AppendMesh is Null"));
+			return TargetMesh;
+		}
+		if (RepeatCount > 0)
 		{
-			AppendOptions.UpdateAttributesForCombineMode(AppendToMesh, OtherMesh);
 			FTransformSRT3d XForm(AppendTransform);
-			FMeshIndexMappings TmpMappings;
-			FDynamicMeshEditor Editor(&AppendToMesh);
-			const FDynamicMesh3* UseOtherMesh = &OtherMesh;
 			FDynamicMesh3 TmpMesh;
-			if (UseOtherMesh == &AppendToMesh)
+			AppendMesh->ProcessMesh([&TmpMesh](const FDynamicMesh3& OtherMesh) { TmpMesh.Copy(OtherMesh); });
+			if (AppendMaterialRemap && TmpMesh.HasAttributes() && TmpMesh.Attributes()->HasMaterialID())
+			{
+				FDynamicMeshMaterialAttribute* MaterialAttrib = TmpMesh.Attributes()->GetMaterialID();
+				for (int32 TID : TmpMesh.TriangleIndicesItr())
+				{
+					int32 OrigMID = MaterialAttrib->GetValue(TID);
+					if (AppendMaterialRemap->IsValidIndex(OrigMID))
+					{
+						int32 NewMID = (*AppendMaterialRemap)[OrigMID];
+						MaterialAttrib->SetValue(TID, NewMID);
+					}
+				}
+			}
+			if (bApplyTransformToFirstInstance)
 			{
-				TmpMesh = OtherMesh;	// need  to make a copy if we are appending to ourself
-				UseOtherMesh = &TmpMesh;
+				MeshTransforms::ApplyTransform(TmpMesh, XForm, true);
 			}
-			Editor.AppendMesh(UseOtherMesh, TmpMappings,
-				[&](int, const FVector3d& Position) { return XForm.TransformPosition(Position); },
-				[&](int, const FVector3d& Normal) { return XForm.TransformNormal(Normal); });
-		});
-	}, EDynamicMeshChangeType::GeneralEdit, EDynamicMeshAttributeChangeFlags::Unknown, bDeferChangeNotifications);
+			TargetMesh->EditMesh([&AppendOptions, &XForm, &TmpMesh, RepeatCount](FDynamicMesh3& AppendToMesh)
+			{
+				AppendOptions.UpdateAttributesForCombineMode(AppendToMesh, TmpMesh);
+				FMeshIndexMappings TmpMappings;
+				FDynamicMeshEditor Editor(&AppendToMesh);
+				for (int32 k = 0; k < RepeatCount; ++k)
+				{
+					Editor.AppendMesh(&TmpMesh, TmpMappings);
+					if (k < RepeatCount)
+					{
+						MeshTransforms::ApplyTransform(TmpMesh, XForm, true);
+						TmpMappings.Reset();
+					}
+				}
+			}, EDynamicMeshChangeType::GeneralEdit, EDynamicMeshAttributeChangeFlags::Unknown, bDeferChangeNotifications);
+		}
 
-	return TargetMesh;
+		return TargetMesh;
+	}
 }
 
+UDynamicMesh* UGeometryScriptLibrary_MeshBasicEditFunctions::AppendMesh(
+	UDynamicMesh* TargetMesh,
+	UDynamicMesh* AppendMesh,
+	FTransform AppendTransform,
+	bool bDeferChangeNotifications,
+	FGeometryScriptAppendMeshOptions AppendOptions,
+	UGeometryScriptDebug* Debug)
+{
+	return UE::Private::AppendHelpers::AppendMesh(TargetMesh, AppendMesh, nullptr, AppendTransform, bDeferChangeNotifications, AppendOptions, Debug);
+}
+
+UDynamicMesh* UGeometryScriptLibrary_MeshBasicEditFunctions::AppendMeshWithMaterials(
+	UDynamicMesh* TargetMesh,
+	const TArray<UMaterialInterface*>& TargetMeshMaterialList,
+	UDynamicMesh* AppendMesh,
+	const TArray<UMaterialInterface*>& AppendMeshMaterialList,
+	TArray<UMaterialInterface*>& ResultMeshMaterialList,
+	FTransform AppendTransform,
+	bool bDeferChangeNotifications,
+	FGeometryScriptAppendMeshOptions AppendOptions,
+	bool bCompactAppendedMaterials,
+	UGeometryScriptDebug* Debug)
+{
+	TArray<int32> AppendMaterialRemap;
+	UE::Private::AppendHelpers::AppendMaterials(TargetMeshMaterialList, AppendMeshMaterialList, bCompactAppendedMaterials, AppendMaterialRemap, ResultMeshMaterialList);
+	return UE::Private::AppendHelpers::AppendMesh(TargetMesh, AppendMesh, &AppendMaterialRemap, AppendTransform, bDeferChangeNotifications, AppendOptions, Debug);
+}
 
 UDynamicMesh* UGeometryScriptLibrary_MeshBasicEditFunctions::AppendMeshTransformed(
 	UDynamicMesh* TargetMesh,
@@ -433,67 +649,27 @@ UDynamicMesh* UGeometryScriptLibrary_MeshBasicEditFunctions::AppendMeshTransform
 	FGeometryScriptAppendMeshOptions AppendOptions,
 	UGeometryScriptDebug* Debug)
 {
-	if (TargetMesh == nullptr)
-	{
-		UE::Geometry::AppendError(Debug, EGeometryScriptErrorType::InvalidInputs, LOCTEXT("AppendMeshTransformed_InvalidInput1", "AppendMeshTransformed: TargetMesh is Null"));
-		return TargetMesh;
-	}
-	if (AppendMesh == nullptr)
-	{
-		UE::Geometry::AppendError(Debug, EGeometryScriptErrorType::InvalidInputs, LOCTEXT("AppendMeshTransformed_InvalidInput2", "AppendMeshTransformed: AppendMesh is Null"));
-		return TargetMesh;
-	}
-	if (AppendTransforms.IsEmpty())
-	{
-		UE::Geometry::AppendError(Debug, EGeometryScriptErrorType::InvalidInputs, LOCTEXT("AppendMeshTransformed_NoTransforms", "AppendMeshTransformed: AppendTransforms array is empty"));
-		return TargetMesh;
-	}
-
-	TargetMesh->EditMesh([&](FDynamicMesh3& AppendToMesh)
-	{
-		AppendMesh->ProcessMesh([&](const FDynamicMesh3& OtherMesh)
-		{
-			AppendOptions.UpdateAttributesForCombineMode(AppendToMesh, OtherMesh);
-			FMeshIndexMappings TmpMappings;
-			FDynamicMeshEditor Editor(&AppendToMesh);
-			const FDynamicMesh3* UseOtherMesh = &OtherMesh;
-			FDynamicMesh3 TmpMesh;
-			if (UseOtherMesh == &AppendToMesh)
-			{
-				TmpMesh = OtherMesh;	// need  to make a copy if we are appending to ourself
-				UseOtherMesh = &TmpMesh;
-			}
-			for (FTransform AppendTransform : AppendTransforms)
-			{
-				FTransformSequence3d TransformSequence;
-
-				if (bConstantTransformIsRelative)
-				{
-					TransformSequence.Append(ConstantTransform);
-					TransformSequence.Append(AppendTransform);
-				}
-				else
-				{
-					// want to apply the constant transform's rotate/scale after
-					// the main transform rotate/scale, so the main positioning 
-					// translation has to be deferred until after that
-					FVector Translation = AppendTransform.GetLocation();
-					AppendTransform.SetTranslation(FVector::Zero());
-
-					TransformSequence.Append(AppendTransform);
-					TransformSequence.Append(ConstantTransform);
-					TransformSequence.Append(FTransform(Translation));
-				}
+	return UE::Private::AppendHelpers::AppendMeshTransformed(TargetMesh, AppendMesh, nullptr, AppendTransforms, ConstantTransform, bConstantTransformIsRelative, bDeferChangeNotifications, AppendOptions, Debug);
+}
 
-				Editor.AppendMesh(UseOtherMesh, TmpMappings,
-					[&](int, const FVector3d& Position) { return TransformSequence.TransformPosition(Position); },
-					[&](int, const FVector3d& Normal) { return TransformSequence.TransformNormal(Normal); });
-				TmpMappings.Reset();
-			}
-		});
-	}, EDynamicMeshChangeType::GeneralEdit, EDynamicMeshAttributeChangeFlags::Unknown, bDeferChangeNotifications);
 
-	return TargetMesh;
+UDynamicMesh* UGeometryScriptLibrary_MeshBasicEditFunctions::AppendMeshTransformedWithMaterials(
+	UDynamicMesh* TargetMesh,
+	const TArray<UMaterialInterface*>& TargetMeshMaterialList,
+	UDynamicMesh* AppendMesh,
+	const TArray<UMaterialInterface*>& AppendMeshMaterialList,
+	TArray<UMaterialInterface*>& ResultMeshMaterialList,
+	const TArray<FTransform>& AppendTransforms,
+	FTransform ConstantTransform,
+	bool bConstantTransformIsRelative,
+	bool bDeferChangeNotifications,
+	FGeometryScriptAppendMeshOptions AppendOptions,
+	bool bCompactAppendedMaterials,
+	UGeometryScriptDebug* Debug)
+{
+	TArray<int32> AppendMaterialRemap;
+	UE::Private::AppendHelpers::AppendMaterials(TargetMeshMaterialList, AppendMeshMaterialList, bCompactAppendedMaterials, AppendMaterialRemap, ResultMeshMaterialList);
+	return UE::Private::AppendHelpers::AppendMeshTransformed(TargetMesh, AppendMesh, &AppendMaterialRemap, AppendTransforms, ConstantTransform, bConstantTransformIsRelative, bDeferChangeNotifications, AppendOptions, Debug);
 }
 
 
@@ -508,46 +684,27 @@ UDynamicMesh* UGeometryScriptLibrary_MeshBasicEditFunctions::AppendMeshRepeated(
 	FGeometryScriptAppendMeshOptions AppendOptions,
 	UGeometryScriptDebug* Debug)
 {
-	if (TargetMesh == nullptr)
-	{
-		UE::Geometry::AppendError(Debug, EGeometryScriptErrorType::InvalidInputs, LOCTEXT("AppendMeshRepeated_InvalidInput1", "AppendMeshRepeated: TargetMesh is Null"));
-		return TargetMesh;
-	}
-	if (AppendMesh == nullptr)
-	{
-		UE::Geometry::AppendError(Debug, EGeometryScriptErrorType::InvalidInputs, LOCTEXT("AppendMeshRepeated_InvalidInput2", "AppendMeshRepeated: AppendMesh is Null"));
-		return TargetMesh;
-	}
-	if (RepeatCount > 0)
-	{
-		FTransformSRT3d XForm(AppendTransform);
-		FDynamicMesh3 TmpMesh;
-		AppendMesh->ProcessMesh([&](const FDynamicMesh3& OtherMesh) { TmpMesh.Copy(OtherMesh); });
-		if (bApplyTransformToFirstInstance)
-		{
-			MeshTransforms::ApplyTransform(TmpMesh, XForm, true);
-		}
-		TargetMesh->EditMesh([&](FDynamicMesh3& AppendToMesh)
-		{
-			AppendOptions.UpdateAttributesForCombineMode(AppendToMesh, TmpMesh);
-			FMeshIndexMappings TmpMappings;
-			FDynamicMeshEditor Editor(&AppendToMesh);
-			for (int32 k = 0; k < RepeatCount; ++k)
-			{
-				Editor.AppendMesh(&TmpMesh, TmpMappings);
-				if (k < RepeatCount)
-				{
-					MeshTransforms::ApplyTransform(TmpMesh, XForm, true);
-					TmpMappings.Reset();
-				}
-			}
-		}, EDynamicMeshChangeType::GeneralEdit, EDynamicMeshAttributeChangeFlags::Unknown, bDeferChangeNotifications);
-	}
-
-	return TargetMesh;
+	return UE::Private::AppendHelpers::AppendMeshRepeated(TargetMesh, AppendMesh, nullptr, AppendTransform, RepeatCount, bApplyTransformToFirstInstance, bDeferChangeNotifications, AppendOptions, Debug);
 }
 
-
+UDynamicMesh* UGeometryScriptLibrary_MeshBasicEditFunctions::AppendMeshRepeatedWithMaterials(
+	UDynamicMesh* TargetMesh,
+	const TArray<UMaterialInterface*>& TargetMeshMaterialList,
+	UDynamicMesh* AppendMesh,
+	const TArray<UMaterialInterface*>& AppendMeshMaterialList,
+	TArray<UMaterialInterface*>& ResultMeshMaterialList,
+	FTransform AppendTransform,
+	int RepeatCount,
+	bool bApplyTransformToFirstInstance,
+	bool bDeferChangeNotifications,
+	FGeometryScriptAppendMeshOptions AppendOptions,
+	bool bCompactAppendedMaterials,
+	UGeometryScriptDebug* Debug)
+{
+	TArray<int32> AppendMaterialRemap;
+	UE::Private::AppendHelpers::AppendMaterials(TargetMeshMaterialList, AppendMeshMaterialList, bCompactAppendedMaterials, AppendMaterialRemap, ResultMeshMaterialList);
+	return UE::Private::AppendHelpers::AppendMeshRepeated(TargetMesh, AppendMesh, &AppendMaterialRemap, AppendTransform, RepeatCount, bApplyTransformToFirstInstance, bDeferChangeNotifications, AppendOptions, Debug);
+}
 
 
 
diff --git a/Engine/Plugins/Runtime/GeometryScripting/Source/GeometryScriptingCore/Public/GeometryScript/MeshBasicEditFunctions.h b/Engine/Plugins/Runtime/GeometryScripting/Source/GeometryScriptingCore/Public/GeometryScript/MeshBasicEditFunctions.h
--- a/Engine/Plugins/Runtime/GeometryScripting/Source/GeometryScriptingCore/Public/GeometryScript/MeshBasicEditFunctions.h
+++ b/Engine/Plugins/Runtime/GeometryScripting/Source/GeometryScriptingCore/Public/GeometryScript/MeshBasicEditFunctions.h
@@ -242,6 +242,25 @@ public:
 		FGeometryScriptAppendMeshOptions AppendOptions = FGeometryScriptAppendMeshOptions(),
 		UGeometryScriptDebug* Debug = nullptr);
 
+	/**
+	 * Apply Append Transform to Append Mesh and then add its geometry to the Target Mesh.
+	 * Also combines materials lists of the Target and Append meshes, and updates the output mesh materials to reference the combined list.
+	 * @param AppendOptions Control how details like mesh attributes are handled when one mesh is appended to another.
+	 */
+	UFUNCTION(BlueprintCallable, Category = "GeometryScript|MeshEdits", meta = (ScriptMethod))
+	static UPARAM(DisplayName = "Target Mesh") UDynamicMesh*
+	AppendMeshWithMaterials(
+		UDynamicMesh* TargetMesh,
+		const TArray<UMaterialInterface*>& TargetMeshMaterialList,
+		UDynamicMesh* AppendMesh,
+		const TArray<UMaterialInterface*>& AppendMeshMaterialList,
+		TArray<UMaterialInterface*>& ResultMeshMaterialList,
+		FTransform AppendTransform,
+		bool bDeferChangeNotifications = false,
+		FGeometryScriptAppendMeshOptions AppendOptions = FGeometryScriptAppendMeshOptions(),
+		bool bCompactAppendedMaterials = true,
+		UGeometryScriptDebug* Debug = nullptr);
+
 	/**
 	 * For each transform in AppendTransforms, apply the transform to AppendMesh and then add its geometry to the TargetMesh.
 	 * @param ConstantTransform the Constant transform will be applied after each Append transform
@@ -260,6 +279,30 @@ public:
 		FGeometryScriptAppendMeshOptions AppendOptions = FGeometryScriptAppendMeshOptions(),
 		UGeometryScriptDebug* Debug = nullptr);
 
+
+	/**
+	 * For each transform in AppendTransforms, apply the transform to AppendMesh and then add its geometry to the TargetMesh.
+	 * Also combines materials lists of the Target and Append meshes, and updates the output mesh materials to reference the combined list.
+	 * @param ConstantTransform the Constant transform will be applied after each Append transform
+	 * @param bConstantTransformIsRelative if true, the Constant transform is applied "in the frame" of the Append Transform, otherwise it is applied as a second transform in local coordinates (ie rotate around the AppendTransform X axis, vs around the local X axis)
+	 * @param AppendOptions Control how details like mesh attributes are handled when one mesh is appended to another
+	 */
+	UFUNCTION(BlueprintCallable, Category = "GeometryScript|MeshEdits", meta=(ScriptMethod))
+	static UPARAM(DisplayName = "Target Mesh") UDynamicMesh* 
+	AppendMeshTransformedWithMaterials( 
+		UDynamicMesh* TargetMesh, 
+		const TArray<UMaterialInterface*>& TargetMeshMaterialList,
+		UDynamicMesh* AppendMesh, 
+		const TArray<UMaterialInterface*>& AppendMeshMaterialList,
+		TArray<UMaterialInterface*>& ResultMeshMaterialList,
+		const TArray<FTransform>& AppendTransforms, 
+		FTransform ConstantTransform,
+		bool bConstantTransformIsRelative = true,
+		bool bDeferChangeNotifications = false,
+		FGeometryScriptAppendMeshOptions AppendOptions = FGeometryScriptAppendMeshOptions(),
+		bool bCompactAppendedMaterials = true,
+		UGeometryScriptDebug* Debug = nullptr);
+
 	/**
 	 * Repeatedly apply AppendTransform to the AppendMesh, each time adding the geometry to TargetMesh.
 	 * @param RepeatCount number of times to repeat the transform-append cycle
@@ -278,6 +321,29 @@ public:
 		FGeometryScriptAppendMeshOptions AppendOptions = FGeometryScriptAppendMeshOptions(),
 		UGeometryScriptDebug* Debug = nullptr);
 
+	/**
+	 * Repeatedly apply AppendTransform to the AppendMesh, each time adding the geometry to TargetMesh.
+	 * Also combines materials lists of the Target and Append meshes, and updates the output mesh materials to reference the combined list.
+	 * @param RepeatCount number of times to repeat the transform-append cycle
+	 * @param bApplyTransformToFirstInstance if true, the AppendTransform is applied before the first mesh append, otherwise it is applied after
+	 * @param AppendOptions Control how details like mesh attributes are handled when one mesh is appended to another
+	 */
+	UFUNCTION(BlueprintCallable, Category = "GeometryScript|MeshEdits", meta=(ScriptMethod))
+	static UPARAM(DisplayName = "Target Mesh") UDynamicMesh* 
+	AppendMeshRepeatedWithMaterials( 
+		UDynamicMesh* TargetMesh, 
+		const TArray<UMaterialInterface*>& TargetMeshMaterialList,
+		UDynamicMesh* AppendMesh, 
+		const TArray<UMaterialInterface*>& AppendMeshMaterialList,
+		TArray<UMaterialInterface*>& ResultMeshMaterialList,
+		FTransform AppendTransform, 
+		int RepeatCount = 1,
+		bool bApplyTransformToFirstInstance = true,
+		bool bDeferChangeNotifications = false,
+		FGeometryScriptAppendMeshOptions AppendOptions = FGeometryScriptAppendMeshOptions(),
+		bool bCompactAppendedMaterials = true,
+		UGeometryScriptDebug* Debug = nullptr);
+
 
 	/**
 	* Adds a set of vertices/triangles to the mesh, with Normals, UVs, and Colors; returning the new triangles indices 
