make some geometry script bone weight struct members 'editanywhere' so they are editable when exposed as variables

diff --git a/Engine/Plugins/Runtime/GeometryScripting/Source/GeometryScriptingCore/Public/GeometryScript/MeshBoneWeightFunctions.h b/Engine/Plugins/Runtime/GeometryScripting/Source/GeometryScriptingCore/Public/GeometryScript/MeshBoneWeightFunctions.h
--- a/Engine/Plugins/Runtime/GeometryScripting/Source/GeometryScriptingCore/Public/GeometryScript/MeshBoneWeightFunctions.h
+++ b/Engine/Plugins/Runtime/GeometryScripting/Source/GeometryScriptingCore/Public/GeometryScript/MeshBoneWeightFunctions.h
@@ -24,10 +24,10 @@ struct GEOMETRYSCRIPTINGCORE_API FGeometryScriptBoneWeight
 	{
 	}
 
-	UPROPERTY(BlueprintReadWrite, Category = BoneWeights)
+	UPROPERTY(EditAnywhere, BlueprintReadWrite, Category = BoneWeights)
 	int32 BoneIndex = 0;
 
-	UPROPERTY(BlueprintReadWrite, Category = BoneWeights)
+	UPROPERTY(EditAnywhere, BlueprintReadWrite, Category = BoneWeights)
 	float Weight = 0;
 };
 
@@ -37,7 +37,7 @@ struct GEOMETRYSCRIPTINGCORE_API FGeometryScriptBoneWeightProfile
 {
 	GENERATED_BODY()
 
-	UPROPERTY(BlueprintReadWrite, Category = BoneWeights)
+	UPROPERTY(EditAnywhere, BlueprintReadWrite, Category = BoneWeights)
 	FName ProfileName = FSkeletalMeshAttributes::DefaultSkinWeightProfileName;
 
 	FName GetProfileName() const { return ProfileName; }
