add geometry script methods to create 'slot name to material' maps from components, static meshes, and skeletal meshes, and to convert such maps back to arrays of materials. (these maps are expected by the 'create new skeletal mesh' method)

diff --git a/Engine/Plugins/Runtime/GeometryScripting/Source/GeometryScriptingCore/Private/MeshMaterialFunctions.cpp b/Engine/Plugins/Runtime/GeometryScripting/Source/GeometryScriptingCore/Private/MeshMaterialFunctions.cpp
--- a/Engine/Plugins/Runtime/GeometryScripting/Source/GeometryScriptingCore/Private/MeshMaterialFunctions.cpp
+++ b/Engine/Plugins/Runtime/GeometryScripting/Source/GeometryScriptingCore/Private/MeshMaterialFunctions.cpp
@@ -8,6 +8,10 @@
 #include "UDynamicMesh.h"
 #include "Polygroups/PolygroupSet.h"
 
+#include "Components/PrimitiveComponent.h"
+#include "Engine/StaticMesh.h"
+#include "Engine/SkeletalMesh.h"
+
 #include UE_INLINE_GENERATED_CPP_BY_NAME(MeshMaterialFunctions)
 
 using namespace UE::Geometry;
@@ -804,7 +808,73 @@ UDynamicMesh* UGeometryScriptLibrary_MeshMaterialFunctions::CompactMaterialIDs(
 	return TargetMesh;
 }
 
+TMap<FName, UMaterialInterface*> UGeometryScriptLibrary_MeshMaterialFunctions::GetMaterialSlotToMaterialMapFromComponent(const UPrimitiveComponent* PrimitiveComponent)
+{
+	TMap<FName, UMaterialInterface*> ToRet;
+	TArray<FName> SlotNames = PrimitiveComponent->GetMaterialSlotNames();
+	ToRet.Reserve(SlotNames.Num());
+	for (const FName& SlotName : SlotNames)
+	{
+		ToRet.Add(SlotName, PrimitiveComponent->GetMaterialByName(SlotName));
+	}
+	return ToRet;
+}
+
+
+TMap<FName, UMaterialInterface*> UGeometryScriptLibrary_MeshMaterialFunctions::GetMaterialSlotToMaterialMapFromStaticMesh(const UStaticMesh* StaticMesh, bool bUseImportedSlotNames)
+{
+	TMap<FName, UMaterialInterface*> ToRet;
+	const TArray<FStaticMaterial>& Materials = StaticMesh->GetStaticMaterials();
+	ToRet.Reserve(Materials.Num());
+	for (const FStaticMaterial& Mat : Materials)
+	{
+#if WITH_EDITORONLY_DATA
+		if (bUseImportedSlotNames)
+		{
+			ToRet.Add(Mat.ImportedMaterialSlotName, Mat.MaterialInterface);
+		}
+		else
+#endif
+		{
+			ToRet.Add(Mat.MaterialSlotName, Mat.MaterialInterface);
+		}
+	}
+	return ToRet;
+}
 
+TMap<FName, UMaterialInterface*> UGeometryScriptLibrary_MeshMaterialFunctions::GetMaterialSlotToMaterialMapFromSkeletalMesh(const USkeletalMesh* SkeletalMesh, bool bUseImportedSlotNames)
+{
+	TMap<FName, UMaterialInterface*> ToRet;
+	const TArray<FSkeletalMaterial>& Materials = SkeletalMesh->GetMaterials();
+	ToRet.Reserve(Materials.Num());
+	for (const FSkeletalMaterial& Mat : Materials)
+	{
+#if WITH_EDITORONLY_DATA
+		if (bUseImportedSlotNames)
+		{
+			ToRet.Add(Mat.ImportedMaterialSlotName, Mat.MaterialInterface);
+		}
+		else
+#endif
+		{
+			ToRet.Add(Mat.MaterialSlotName, Mat.MaterialInterface);
+		}
+	}
+	return ToRet;
+}
+
+
+TArray<UMaterialInterface*> UGeometryScriptLibrary_MeshMaterialFunctions::Conv_MaterialsMapToMaterialsArray(const TMap<FName, UMaterialInterface*>& MaterialsMap)
+{
+	TArray<UMaterialInterface*> ToRet;
+	ToRet.Reserve(MaterialsMap.Num());
+	for (const TPair<FName, UMaterialInterface*>& NameMat : MaterialsMap)
+	{
+		// Note we add the materials even if they are null, so we always have the same number of materials in the list as were in the map
+		ToRet.Add(NameMat.Value);
+	}
+	return ToRet;
+}
 
 
 #undef LOCTEXT_NAMESPACE
diff --git a/Engine/Plugins/Runtime/GeometryScripting/Source/GeometryScriptingCore/Public/GeometryScript/MeshMaterialFunctions.h b/Engine/Plugins/Runtime/GeometryScripting/Source/GeometryScriptingCore/Public/GeometryScript/MeshMaterialFunctions.h
--- a/Engine/Plugins/Runtime/GeometryScripting/Source/GeometryScriptingCore/Public/GeometryScript/MeshMaterialFunctions.h
+++ b/Engine/Plugins/Runtime/GeometryScripting/Source/GeometryScriptingCore/Public/GeometryScript/MeshMaterialFunctions.h
@@ -235,4 +235,30 @@ public:
 		bool bRemoveDuplicateMaterials = false,
 		UGeometryScriptDebug* Debug = nullptr);
 
+	/**
+	 * Get a map from material slot names to materials from a given component
+	 */
+	UFUNCTION(BlueprintCallable, Category = "GeometryScript|Materials", meta = (ScriptMethod))
+	static UPARAM(DisplayName = "Slot Name to Material Map") TMap<FName, UMaterialInterface*> GetMaterialSlotToMaterialMapFromComponent(const UPrimitiveComponent* PrimitiveComponent);
+
+	/**
+	 * Get a map from material slot names to materials from a given static mesh
+	 * @param bUseImportedSlotNames If available (i.e., if editor-only data is available), use the imported slot names instead of the current slot names as the map keys.
+	 */
+	UFUNCTION(BlueprintCallable, Category = "GeometryScript|Materials", meta = (ScriptMethod))
+	static UPARAM(DisplayName = "Slot Name to Material Map") TMap<FName, UMaterialInterface*> GetMaterialSlotToMaterialMapFromStaticMesh(const UStaticMesh* StaticMesh, bool bUseImportedSlotNames = false);
+
+	/**
+	 * Get a map from material slot names to materials from a given skeletal mesh
+	 * @param bUseImportedSlotNames If available (i.e., if editor-only data is available), use the imported slot names instead of the current slot names as the map keys.
+	 */
+	UFUNCTION(BlueprintCallable, Category = "GeometryScript|Materials", meta = (ScriptMethod))
+	static UPARAM(DisplayName = "Slot Name to Material Map") TMap<FName, UMaterialInterface*> GetMaterialSlotToMaterialMapFromSkeletalMesh(const USkeletalMesh* SkeletalMesh, bool bUseImportedSlotNames = false);
+
+	/**
+	 * Converts materials map to materials array. Null materials will be kept in the list, and the list will have the same number of elements as the map.
+	 */
+	UFUNCTION(BlueprintPure, meta = (DisplayName = "Material Map To Material List", CompactNodeTitle = "->", BlueprintAutocast), Category = "GeometryScript|Materials")
+	static UPARAM(DisplayName = "Material List") TArray<UMaterialInterface*> Conv_MaterialsMapToMaterialsArray(const TMap<FName, UMaterialInterface*>& MaterialMap);
+
 };
\ No newline at end of file
