Make CopyMeshToStaticMesh have the same bUseSectionMaterials parameter as the CopyMeshFromStaticMesh, and change the default behavior for both to use section materials.

diff --git a/Engine/Plugins/Runtime/GeometryScripting/Source/GeometryScriptingCore/Private/MeshAssetFunctions.cpp b/Engine/Plugins/Runtime/GeometryScripting/Source/GeometryScriptingCore/Private/MeshAssetFunctions.cpp
--- a/Engine/Plugins/Runtime/GeometryScripting/Source/GeometryScriptingCore/Private/MeshAssetFunctions.cpp
+++ b/Engine/Plugins/Runtime/GeometryScripting/Source/GeometryScriptingCore/Private/MeshAssetFunctions.cpp
@@ -128,6 +128,7 @@ UDynamicMesh*  UGeometryScriptLibrary_StaticMeshFunctions::CopyMeshToStaticMesh(
 	FGeometryScriptCopyMeshToAssetOptions Options,
 	FGeometryScriptMeshWriteLOD TargetLOD,
 	EGeometryScriptOutcomePins& Outcome,
+	bool bUseSectionMaterials,
 	UGeometryScriptDebug* Debug)
 {
 	Outcome = EGeometryScriptOutcomePins::Failure;
@@ -260,12 +261,9 @@ UDynamicMesh*  UGeometryScriptLibrary_StaticMeshFunctions::CopyMeshToStaticMesh(
 
 		FConversionToMeshDescriptionOptions ConversionOptions;
 		FDynamicMeshToMeshDescription Converter(ConversionOptions);
-		if (!Options.bReplaceMaterials)
+		if (!bUseSectionMaterials && !Options.bReplaceMaterials)
 		{
-			UE::Conversion::EMeshLODType LODType;
-			int32 LODIndex;
-			ConvertGeometryScriptWriteLOD(TargetLOD, LODType, LODIndex);
-			TArray<int32> MaterialIDMap = UE::Conversion::GetPolygonGroupToMaterialIndexMap(ToStaticMeshAsset, LODType, LODIndex);
+			TArray<int32> MaterialIDMap = UE::Conversion::GetPolygonGroupToMaterialIndexMap(ToStaticMeshAsset, UE::Conversion::EMeshLODType::SourceModel, UseLODIndex);
 			Converter.SetMaterialIDMapFromInverseMap(MaterialIDMap);
 		}
 		FromDynamicMesh->ProcessMesh([&](const FDynamicMesh3& ReadMesh)
diff --git a/Engine/Plugins/Runtime/GeometryScripting/Source/GeometryScriptingCore/Public/GeometryScript/MeshAssetFunctions.h b/Engine/Plugins/Runtime/GeometryScripting/Source/GeometryScriptingCore/Public/GeometryScript/MeshAssetFunctions.h
--- a/Engine/Plugins/Runtime/GeometryScripting/Source/GeometryScriptingCore/Public/GeometryScript/MeshAssetFunctions.h
+++ b/Engine/Plugins/Runtime/GeometryScripting/Source/GeometryScriptingCore/Public/GeometryScript/MeshAssetFunctions.h
@@ -159,8 +159,10 @@ public:
 	* Extracts a Dynamic Mesh from a Static Mesh Asset. 
 	* 
 	* Note that the LOD Index in RequestedLOD will be silently clamped to the available number of LODs (SourceModel or RenderData)
+	*
+	* @param bUseSectionMaterials Whether to use the mesh section indices as material IDs. If true, use GetSectionMaterialListFromStaticMesh to get the corresponding materials. If false, use GetMaterialListFromStaticMesh to get the materials instead.
 	*/
-	//~ Note this is V2 to change the default behavior to use the normal materials list, rather than the section materials list.
+	//~ Note this V2 version adds the bUseSectionMaterials parameter
 	UFUNCTION(BlueprintCallable, Category = "GeometryScript|StaticMesh", meta = (DisplayName = "Copy Mesh From Static Mesh", ExpandEnumAsExecs = "Outcome"))
 	static UPARAM(DisplayName = "Dynamic Mesh") UDynamicMesh* 
 	CopyMeshFromStaticMeshV2(
@@ -169,11 +171,11 @@ public:
 		FGeometryScriptCopyMeshFromAssetOptions AssetOptions,
 		FGeometryScriptMeshReadLOD RequestedLOD,
 		EGeometryScriptOutcomePins& Outcome,
-		bool bUseSectionMaterials = false,
+		bool bUseSectionMaterials = true,
 		UGeometryScriptDebug* Debug = nullptr);
 
 	/**
-	* Extracts a Dynamic Mesh from a Static Mesh Asset, using section material indices for the materials.
+	* Extracts a Dynamic Mesh from a Static Mesh Asset, using section indices for the material IDs -- use GetSectionMaterialListFromStaticMesh to get the corresponding materials.
 	*
 	* Note that the LOD Index in RequestedLOD will be silently clamped to the available number of LODs (SourceModel or RenderData)
 	*/
@@ -189,20 +191,36 @@ public:
 	{
 		return CopyMeshFromStaticMeshV2(FromStaticMeshAsset, ToDynamicMesh, AssetOptions, RequestedLOD, Outcome, true, Debug);
 	}
-
-	/** 
-	* Updates a Static Mesh Asset with new geometry converted from a Dynamic Mesh.
+	
+	/**
+	* Updates a Static Mesh Asset with new geometry converted from a Dynamic Mesh
+	*
+	* @param bUseSectionMaterials Whether to assume Dynamic Mesh material IDs are section indices in the target Static Mesh. Should match the value passed to CopyMeshFromStaticMesh. Has no effect if replacing the asset materials.
 	*/
 	UFUNCTION(BlueprintCallable, Category = "GeometryScript|StaticMesh", meta = (ExpandEnumAsExecs = "Outcome"))
-	static UPARAM(DisplayName = "Dynamic Mesh") UDynamicMesh* 
+	static UPARAM(DisplayName = "Dynamic Mesh") UDynamicMesh*
 	CopyMeshToStaticMesh(
-		UDynamicMesh* FromDynamicMesh, 
-		UStaticMesh* ToStaticMeshAsset, 
+		UDynamicMesh* FromDynamicMesh,
+		UStaticMesh* ToStaticMeshAsset,
 		FGeometryScriptCopyMeshToAssetOptions Options,
 		FGeometryScriptMeshWriteLOD TargetLOD,
 		EGeometryScriptOutcomePins& Outcome,
+		bool bUseSectionMaterials = true,
 		UGeometryScriptDebug* Debug = nullptr);
 
+	//~ Version for C++ api compatibility without bUseSectionMaterials parameter
+	UE_DEPRECATED(5.5, "Use the version of this function with a bUseSectionMaterials parameter")
+	static UDynamicMesh* CopyMeshToStaticMesh(
+		UDynamicMesh* FromDynamicMesh,
+		UStaticMesh* ToStaticMeshAsset,
+		FGeometryScriptCopyMeshToAssetOptions Options,
+		FGeometryScriptMeshWriteLOD TargetLOD,
+		EGeometryScriptOutcomePins& Outcome,
+		UGeometryScriptDebug* Debug)
+	{
+		return CopyMeshToStaticMesh(FromDynamicMesh, ToStaticMeshAsset, Options, TargetLOD, Outcome, true, Debug);
+	}
+
 
     /** 
 	* Extracts the Material List and corresponding Material Indices from the specified LOD of the Static Mesh Asset. 
