Geoscripting: Fix morph target writing to use the non-manifold conversion when going back to the original mesh descript.

diff --git a/Engine/Plugins/Runtime/GeometryScripting/Source/GeometryScriptingCore/Private/MeshAssetFunctions.cpp b/Engine/Plugins/Runtime/GeometryScripting/Source/GeometryScriptingCore/Private/MeshAssetFunctions.cpp
--- a/Engine/Plugins/Runtime/GeometryScripting/Source/GeometryScriptingCore/Private/MeshAssetFunctions.cpp
+++ b/Engine/Plugins/Runtime/GeometryScripting/Source/GeometryScriptingCore/Private/MeshAssetFunctions.cpp
@@ -25,6 +25,7 @@
 #include "StaticMeshLODResourcesToDynamicMesh.h"
 #include "AssetUtils/StaticMeshMaterialUtil.h"
 #include "ConversionUtils/SceneComponentToDynamicMesh.h"
+#include "DynamicMesh/NonManifoldMappingSupport.h"
 
 #include UE_INLINE_GENERATED_CPP_BY_NAME(MeshAssetFunctions)
 
@@ -638,6 +639,10 @@ UDynamicMesh* UGeometryScriptLibrary_StaticMeshFunctions::CopyMeshFromSkeletalMe
 
 		FDynamicMesh3 NewMesh;
 		FMeshDescriptionToDynamicMesh Converter;
+
+		// Leave this on, since the set morph target node uses this. 
+		Converter.bVIDsFromNonManifoldMeshDescriptionAttr = true;
+		
 		Converter.Convert(SourceMesh, NewMesh, AssetOptions.bRequestTangents);
 	
 		ToDynamicMesh->SetMesh(MoveTemp(NewMesh));
@@ -910,14 +915,32 @@ UDynamicMesh* UGeometryScriptLibrary_StaticMeshFunctions::CopyMorphTargetToSkele
 		return FromDynamicMesh;
 	}
 
-	// Morph targets must be compact and have the same number of vertices as the skeletal asset mesh description
-	if (FromDynamicMesh->GetMeshRef().MaxVertexID() != FromDynamicMesh->GetMeshRef().VertexCount() ||
-		FromDynamicMesh->GetMeshRef().MaxVertexID() != MeshDescription->Vertices().Num())
+	// If the dynamic mesh has non-manifold information, use that to figure out what the original vertex count was.
+	// Otherwise, we assume that they have a 1:1 match.
+	const FDynamicMesh3& SourceMesh = FromDynamicMesh->GetMeshRef();
+	const FNonManifoldMappingSupport NonManifoldMappingSupport(SourceMesh);
+	int32 SourceVertexCount;
+
+	if (NonManifoldMappingSupport.IsNonManifoldVertexInSource())
+	{
+		TSet<int32> UniqueVertices;
+		for (int32 SourceVID = 0; SourceVID < SourceMesh.VertexCount(); ++SourceVID)
+		{
+			UniqueVertices.Add(NonManifoldMappingSupport.GetOriginalNonManifoldVertexID(SourceVID));
+		}
+
+		SourceVertexCount = UniqueVertices.Num();
+	}
+	else
+	{
+		SourceVertexCount = SourceMesh.VertexCount();
+	}
+	if (MeshDescription->Vertices().Num() != SourceVertexCount)
 	{
 		UE::Geometry::AppendError(Debug, EGeometryScriptErrorType::InvalidInputs, LOCTEXT("CopyMorphTargetToSkeletalMesh_InvalidMorphTargetGeometry", "CopyMorphTargetToSkeletalMesh: Morph target mesh doesnt have the same number of vertices as the skeletal mesh."));
 		return FromDynamicMesh;
 	}
-
+	
 	FSkeletalMeshAttributes MeshAttributes(*MeshDescription);
 	MeshAttributes.Register();
 
@@ -940,17 +963,19 @@ UDynamicMesh* UGeometryScriptLibrary_StaticMeshFunctions::CopyMorphTargetToSkele
 		}
 	}
 
+	
 	TVertexAttributesRef<FVector3f> PositionDelta = MeshAttributes.GetVertexMorphPositionDelta(MorphTargetName);
 	TVertexAttributesRef<FVector3f> VertexPositions = MeshAttributes.GetVertexPositions();
 
-	for (int32 VID = 0; VID < FromDynamicMesh->GetMeshRef().MaxVertexID(); ++VID)
+	for (int32 SourceVID = 0; SourceVID < SourceMesh.VertexCount(); ++SourceVID)
 	{
-		const FVector3d V0 = FromDynamicMesh->GetMeshRef().GetVertex(VID);
-		const FVector3f V1 = VertexPositions[VID];
+		const int32 TargetVID = NonManifoldMappingSupport.GetOriginalNonManifoldVertexID(SourceVID);
+		
+		const FVector3d V0 = FromDynamicMesh->GetMeshRef().GetVertex(SourceVID);
+		const FVector3f V1 = VertexPositions[TargetVID];
 
-		PositionDelta.Set(VID, FVector3f(float(V0[0]) - V1[0], 
-										 float(V0[1]) - V1[1], 
-										 float(V0[2]) - V1[2]));
+		const FVector3f Delta = FVector3f{V0} - V1;
+		PositionDelta.Set(TargetVID, Delta);
 	}
 
 	ToSkeletalMeshAsset->CommitMeshDescription(TargetLOD.LODIndex);
