Geoscripting: Fix initial color when averaging triangle vertex colors into vertex colors.

#rnx

diff --git a/Engine/Plugins/Runtime/GeometryScripting/Source/GeometryScriptingCore/Private/MeshVertexColorFunctions.cpp b/Engine/Plugins/Runtime/GeometryScripting/Source/GeometryScriptingCore/Private/MeshVertexColorFunctions.cpp
--- a/Engine/Plugins/Runtime/GeometryScripting/Source/GeometryScriptingCore/Private/MeshVertexColorFunctions.cpp
+++ b/Engine/Plugins/Runtime/GeometryScripting/Source/GeometryScriptingCore/Private/MeshVertexColorFunctions.cpp
@@ -286,7 +286,7 @@ UDynamicMesh* UGeometryScriptLibrary_MeshVertexColorFunctions::GetMeshPerVertexC
 	{
 		TargetMesh->ProcessMesh([&](const FDynamicMesh3& ReadMesh)
 		{
-			Colors.Init(FLinearColor::Black, ReadMesh.MaxVertexID());
+			Colors.Init(FLinearColor::Transparent, ReadMesh.MaxVertexID());
 			bHasVertexIDGaps = ! ReadMesh.IsCompactV();
 
 			if (ReadMesh.HasAttributes() && ReadMesh.Attributes()->HasPrimaryColors() )
