add a method to get section materials for skeletal meshes

diff --git a/Engine/Plugins/Runtime/GeometryScripting/Source/GeometryScriptingCore/Private/MeshAssetFunctions.cpp b/Engine/Plugins/Runtime/GeometryScripting/Source/GeometryScriptingCore/Private/MeshAssetFunctions.cpp
--- a/Engine/Plugins/Runtime/GeometryScripting/Source/GeometryScriptingCore/Private/MeshAssetFunctions.cpp
+++ b/Engine/Plugins/Runtime/GeometryScripting/Source/GeometryScriptingCore/Private/MeshAssetFunctions.cpp
@@ -536,6 +536,86 @@ void UGeometryScriptLibrary_StaticMeshFunctions::GetSectionMaterialListFromStati
 #endif
 }
 
+void UGeometryScriptLibrary_StaticMeshFunctions::GetLODMaterialListFromSkeletalMesh(
+	USkeletalMesh* FromSkeletalMeshAsset,
+	FGeometryScriptMeshReadLOD RequestedLOD,
+	TArray<UMaterialInterface*>& MaterialList,
+	TArray<int32>& MaterialIndex,
+	TArray<FName>& MaterialSlotNames,
+	EGeometryScriptOutcomePins& Outcome,
+	UGeometryScriptDebug* Debug)
+{
+	Outcome = EGeometryScriptOutcomePins::Failure;
+
+	if (FromSkeletalMeshAsset == nullptr)
+	{
+		UE::Geometry::AppendError(Debug, EGeometryScriptErrorType::InvalidInputs, LOCTEXT("GetLODMaterialListFromSkeletalMesh_InvalidInput1", "GetLODMaterialListFromSkeletalMesh: FromSkeletalMeshAsset is Null"));
+		return;
+	}
+
+#if WITH_EDITOR
+
+	if (RequestedLOD.LODType == EGeometryScriptLODType::HiResSourceModel)
+	{
+		UE::Geometry::AppendError(Debug, EGeometryScriptErrorType::InvalidInputs, LOCTEXT("GetLODMaterialListFromSkeletalMesh_LODNotAvailable", "GetLODMaterialListFromSkeletalMesh: Requested LOD is not available"));
+		return;
+	}
+
+	int32 UseLODIndex = FMath::Clamp(RequestedLOD.LODIndex, 0, FromSkeletalMeshAsset->GetLODNum() - 1);
+
+	const TArray<FSkeletalMaterial>& Mats = FromSkeletalMeshAsset->GetMaterials();
+	const int32 NumMats = Mats.Num();
+
+	// Get the material mapping via the LODInfo struct
+	const FSkeletalMeshLODInfo* LODInfo = FromSkeletalMeshAsset->GetLODInfo(UseLODIndex);
+	if (LODInfo && !LODInfo->LODMaterialMap.IsEmpty())
+	{
+		const TArray<int32>& Map = LODInfo->LODMaterialMap;
+		const int32 NumSectionMat = Map.Num();
+		MaterialList.Reset(NumSectionMat);
+		MaterialIndex.Reset(NumSectionMat);
+		MaterialSlotNames.Reset(NumSectionMat);
+		for (int32 Idx = 0; Idx < NumSectionMat; ++Idx)
+		{
+			int32 MatIdx = Map[Idx];
+			if (MatIdx == INDEX_NONE) // by convention, INDEX_NONE means the index is mapped to itself
+			{
+				MatIdx = FMath::Min(Idx, NumMats - 1);
+			}
+			MaterialIndex.Add(MatIdx);
+			if (Mats.IsValidIndex(MatIdx))
+			{
+				MaterialList.Add(Mats[MatIdx].MaterialInterface);
+				MaterialSlotNames.Add(Mats[MatIdx].MaterialSlotName);
+			}
+			else
+			{
+				MaterialList.Add(nullptr);
+				MaterialSlotNames.Add(FName());
+			}
+		}
+	}
+	// if the LODMaterialMap is not there or is empty, materials are identity-mapped
+	else
+	{
+		MaterialList.Reset(NumMats);
+		MaterialIndex.Reset(NumMats);
+		MaterialSlotNames.Reset(NumMats);
+		for (int32 Idx = 0; Idx < NumMats; ++Idx)
+		{
+			MaterialIndex.Add(Idx);
+			MaterialList.Add(Mats[Idx].MaterialInterface);
+			MaterialSlotNames.Add(Mats[Idx].MaterialSlotName);
+		}
+	}
+
+	Outcome = EGeometryScriptOutcomePins::Success;
+
+#else
+	UE::Geometry::AppendError(Debug, EGeometryScriptErrorType::InvalidInputs, LOCTEXT("GetLODMaterialListFromSkeletalMesh_EditorOnly", "GetLODMaterialListFromSkeletalMesh: Not available at Runtime"));
+#endif
+}
+
 
 namespace UELocal
 {
diff --git a/Engine/Plugins/Runtime/GeometryScripting/Source/GeometryScriptingCore/Public/GeometryScript/MeshAssetFunctions.h b/Engine/Plugins/Runtime/GeometryScripting/Source/GeometryScriptingCore/Public/GeometryScript/MeshAssetFunctions.h
--- a/Engine/Plugins/Runtime/GeometryScripting/Source/GeometryScriptingCore/Public/GeometryScript/MeshAssetFunctions.h
+++ b/Engine/Plugins/Runtime/GeometryScripting/Source/GeometryScriptingCore/Public/GeometryScript/MeshAssetFunctions.h
@@ -245,7 +245,27 @@ public:
 		UGeometryScriptDebug* Debug = nullptr);
 
 	/**
-	 * Get the asset materials from the static mesh asset. These will match the DynamicMesh material if CopyMeshToStaticMesh
+	* Extracts the Material List and corresponding Material Indices from the specified LOD of the Skeletal Mesh Asset.
+	* If Copy Mesh To Skeletal Mesh was used to create a Dynamic Mesh, then the returned Material List can be passed directly to a Dynamic Mesh Component.
+	*
+	* @param MaterialIndex this returned array is the same size as MaterialList, with each value the index of that Material in the Skeletal Mesh's Material List
+	* @param MateriaSlotNames this returned array is the same size as MaterialList, with each value the Slot Name of that Material in the Skeletal Mesh's Material List
+	*
+	* Note that the LOD Index in RequestedLOD will be silently clamped to the available number of LODs
+	*/
+	UFUNCTION(BlueprintCallable, Category = "GeometryScript|SkeletalMesh", meta = (DisplayName = "Get LOD Material List From Skeletal Mesh", ExpandEnumAsExecs = "Outcome"))
+	static void
+	GetLODMaterialListFromSkeletalMesh(
+		USkeletalMesh* FromSkeletalMeshAsset,
+		FGeometryScriptMeshReadLOD RequestedLOD,
+		TArray<UMaterialInterface*>& MaterialList,
+		TArray<int32>& MaterialIndex,
+		TArray<FName>& MaterialSlotNames,
+		EGeometryScriptOutcomePins& Outcome,
+		UGeometryScriptDebug* Debug = nullptr);
+
+	/**
+	 * Get the asset materials from the static mesh asset. These will match the DynamicMesh material if CopyMeshFromStaticMesh
 	 * was used to create a DynamicMesh with bUseSectionMaterials=false
 	 */
 	UFUNCTION(BlueprintCallable, Category = "GeometryScript|StaticMesh")
