More consistently also add 'EditAnywhere' tag to geometry script options structs members that are 'BlueprintReadWrite' so they are editable when used as variables

diff --git a/Engine/Plugins/Runtime/GeometryScripting/Source/GeometryScriptingCore/Public/GeometryScript/GeometryScriptTypes.h b/Engine/Plugins/Runtime/GeometryScripting/Source/GeometryScriptingCore/Public/GeometryScript/GeometryScriptTypes.h
--- a/Engine/Plugins/Runtime/GeometryScripting/Source/GeometryScriptingCore/Public/GeometryScript/GeometryScriptTypes.h
+++ b/Engine/Plugins/Runtime/GeometryScripting/Source/GeometryScriptingCore/Public/GeometryScript/GeometryScriptTypes.h
@@ -85,10 +85,10 @@ struct GEOMETRYSCRIPTINGCORE_API FGeometryScriptMeshReadLOD
 {
 	GENERATED_BODY()
 public:
-	UPROPERTY(BlueprintReadWrite, Category = LOD)
+	UPROPERTY(EditAnywhere, BlueprintReadWrite, Category = LOD)
 	EGeometryScriptLODType LODType = EGeometryScriptLODType::MaxAvailable;
 
-	UPROPERTY(BlueprintReadWrite, Category = LOD)
+	UPROPERTY(EditAnywhere, BlueprintReadWrite, Category = LOD)
 	int32 LODIndex = 0;
 };
 
@@ -98,10 +98,10 @@ struct GEOMETRYSCRIPTINGCORE_API FGeometryScriptMeshWriteLOD
 {
 	GENERATED_BODY()
 public:
-	UPROPERTY(BlueprintReadWrite, Category = LOD)
+	UPROPERTY(EditAnywhere, BlueprintReadWrite, Category = LOD)
 	bool bWriteHiResSource = false;
 
-	UPROPERTY(BlueprintReadWrite, Category = LOD)
+	UPROPERTY(EditAnywhere, BlueprintReadWrite, Category = LOD)
 	int32 LODIndex = 0;
 };
 
@@ -190,13 +190,13 @@ struct GEOMETRYSCRIPTINGCORE_API FGeometryScriptTriangle
 {
 	GENERATED_BODY()
 public:
-	UPROPERTY(BlueprintReadWrite, Category = Options)
+	UPROPERTY(EditAnywhere, BlueprintReadWrite, Category = Options)
 	FVector Vector0 = FVector::ZeroVector;
 
-	UPROPERTY(BlueprintReadWrite, Category = Options)
+	UPROPERTY(EditAnywhere, BlueprintReadWrite, Category = Options)
 	FVector Vector1 = FVector::ZeroVector;
 
-	UPROPERTY(BlueprintReadWrite, Category = Options)
+	UPROPERTY(EditAnywhere, BlueprintReadWrite, Category = Options)
 	FVector Vector2 = FVector::ZeroVector;
 };
 
@@ -206,16 +206,16 @@ struct GEOMETRYSCRIPTINGCORE_API FGeometryScriptTrianglePoint
 {
 	GENERATED_BODY()
 public:
-	UPROPERTY(BlueprintReadWrite, Category = Options)
+	UPROPERTY(EditAnywhere, BlueprintReadWrite, Category = Options)
 	bool bValid = false;
 
-	UPROPERTY(BlueprintReadWrite, Category = Options)
+	UPROPERTY(EditAnywhere, BlueprintReadWrite, Category = Options)
 	int TriangleID = -1;
 
-	UPROPERTY(BlueprintReadWrite, Category = Options)
+	UPROPERTY(EditAnywhere, BlueprintReadWrite, Category = Options)
 	FVector Position = FVector::ZeroVector;
 
-	UPROPERTY(BlueprintReadWrite, Category = Options)
+	UPROPERTY(EditAnywhere, BlueprintReadWrite, Category = Options)
 	FVector BaryCoords = FVector::ZeroVector;
 };
 
@@ -225,13 +225,13 @@ struct GEOMETRYSCRIPTINGCORE_API FGeometryScriptUVTriangle
 {
 	GENERATED_BODY()
 public:
-	UPROPERTY(BlueprintReadWrite, Category = Options)
+	UPROPERTY(EditAnywhere, BlueprintReadWrite, Category = Options)
 	FVector2D UV0 = FVector2D::ZeroVector;
 
-	UPROPERTY(BlueprintReadWrite, Category = Options)
+	UPROPERTY(EditAnywhere, BlueprintReadWrite, Category = Options)
 	FVector2D UV1 = FVector2D::ZeroVector;
 
-	UPROPERTY(BlueprintReadWrite, Category = Options)
+	UPROPERTY(EditAnywhere, BlueprintReadWrite, Category = Options)
 	FVector2D UV2 = FVector2D::ZeroVector;
 };
 
@@ -247,16 +247,16 @@ struct GEOMETRYSCRIPTINGCORE_API FGeometryScriptColorFlags
 {
 	GENERATED_BODY()
 public:
-	UPROPERTY(BlueprintReadWrite, Category = Color)
+	UPROPERTY(EditAnywhere, BlueprintReadWrite, Category = Color)
 	bool bRed = true;
 
-	UPROPERTY(BlueprintReadWrite, Category = Color)
+	UPROPERTY(EditAnywhere, BlueprintReadWrite, Category = Color)
 	bool bGreen = true;
 
-	UPROPERTY(BlueprintReadWrite, Category = Color)
+	UPROPERTY(EditAnywhere, BlueprintReadWrite, Category = Color)
 	bool bBlue = true;
 
-	UPROPERTY(BlueprintReadWrite, Category = Color)
+	UPROPERTY(EditAnywhere, BlueprintReadWrite, Category = Color)
 	bool bAlpha = true;
 
 	bool AllSet() const { return bRed && bGreen && bBlue && bAlpha; }
@@ -278,11 +278,11 @@ struct GEOMETRYSCRIPTINGCORE_API FGeometryScriptGroupLayer
 	GENERATED_BODY()
 public:
 	/** If true,the default/standard PolyGroup Layer is used */
-	UPROPERTY(BlueprintReadWrite, Category = LOD)
+	UPROPERTY(EditAnywhere, BlueprintReadWrite, Category = LOD)
 	bool bDefaultLayer = true;
 
 	/** Index of an extended PolyGroup Layer (which may or may not exist on any given Mesh) */
-	UPROPERTY(BlueprintReadWrite, Category = LOD)
+	UPROPERTY(EditAnywhere, BlueprintReadWrite, Category = LOD)
 	int ExtendedLayerIndex = 0;
 };
 
@@ -315,7 +315,7 @@ struct GEOMETRYSCRIPTINGCORE_API FGeometryScriptIndexList
 {
 	GENERATED_BODY()
 
-	UPROPERTY(BlueprintReadWrite, Category = IndexList)
+	UPROPERTY(EditAnywhere, BlueprintReadWrite, Category = IndexList)
 	EGeometryScriptIndexType IndexType = EGeometryScriptIndexType::Any;
 public:
 	TSharedPtr<TArray<int>> List;
@@ -796,13 +796,13 @@ struct GEOMETRYSCRIPTINGCORE_API FGeometryScriptDebugMessage
 {
 	GENERATED_BODY()
 public:
-	UPROPERTY(BlueprintReadWrite, Category = Options)
+	UPROPERTY(EditAnywhere, BlueprintReadWrite, Category = Options)
 	EGeometryScriptDebugMessageType MessageType = EGeometryScriptDebugMessageType::ErrorMessage;
 
-	UPROPERTY(BlueprintReadWrite, Category = Options)
+	UPROPERTY(EditAnywhere, BlueprintReadWrite, Category = Options)
 	EGeometryScriptErrorType ErrorType = EGeometryScriptErrorType::UnknownError;
 
-	UPROPERTY(BlueprintReadWrite, Category = Options)
+	UPROPERTY(EditAnywhere, BlueprintReadWrite, Category = Options)
 	FText Message = FText();
 };
 
@@ -812,7 +812,7 @@ class GEOMETRYSCRIPTINGCORE_API UGeometryScriptDebug : public UObject
 {
 	GENERATED_BODY()
 public:
-	UPROPERTY(BlueprintReadWrite, Category = Messages)
+	UPROPERTY(EditAnywhere, BlueprintReadWrite, Category = Messages)
 	TArray<FGeometryScriptDebugMessage> Messages;
 
 
diff --git a/Engine/Plugins/Runtime/GeometryScripting/Source/GeometryScriptingCore/Public/GeometryScript/MeshBakeFunctions.h b/Engine/Plugins/Runtime/GeometryScripting/Source/GeometryScriptingCore/Public/GeometryScript/MeshBakeFunctions.h
--- a/Engine/Plugins/Runtime/GeometryScripting/Source/GeometryScriptingCore/Public/GeometryScript/MeshBakeFunctions.h
+++ b/Engine/Plugins/Runtime/GeometryScripting/Source/GeometryScriptingCore/Public/GeometryScript/MeshBakeFunctions.h
@@ -274,7 +274,7 @@ struct GEOMETRYSCRIPTINGCORE_API FGeometryScriptBakeVertexOptions
 	bool bSplitAtNormalSeams = false;
 
 	/** If true, compute a separate vertex color for each unique UV on a vertex. */
-	UPROPERTY(BlueprintReadWrite, Category = Options, meta=(DisplayName = "Split at UV Seams"))
+	UPROPERTY(EditAnywhere, BlueprintReadWrite, Category = Options, meta=(DisplayName = "Split at UV Seams"))
 	bool bSplitAtUVSeams = false;
 
 	/** Maximum allowed distance for the projection from target mesh to source mesh for the sample to be considered valid.
@@ -318,7 +318,7 @@ struct GEOMETRYSCRIPTINGCORE_API FGeometryScriptBakeTargetMeshOptions
 {
 	GENERATED_BODY();
 
-	UPROPERTY(BlueprintReadWrite, Category = Options, meta = (DisplayName="Target UV Channel"))
+	UPROPERTY(EditAnywhere, BlueprintReadWrite, Category = Options, meta = (DisplayName="Target UV Channel"))
 	int TargetUVLayer = 0;
 };
 
@@ -330,7 +330,7 @@ struct GEOMETRYSCRIPTINGCORE_API FGeometryScriptBakeSourceMeshOptions
 	UPROPERTY(EditAnywhere, BlueprintReadWrite, Category = Options)
 	TObjectPtr<UTexture2D> SourceNormalMap = nullptr;
 
-	UPROPERTY(BlueprintReadWrite, Category = Options, meta = (DisplayName="Source Normal UV Channel"))
+	UPROPERTY(EditAnywhere, BlueprintReadWrite, Category = Options, meta = (DisplayName="Source Normal UV Channel"))
 	int SourceNormalUVLayer = 0;
 
 	UPROPERTY(EditAnywhere, BlueprintReadWrite, Category = Options)
diff --git a/Engine/Plugins/Runtime/GeometryScripting/Source/GeometryScriptingCore/Public/GeometryScript/MeshBoneWeightFunctions.h b/Engine/Plugins/Runtime/GeometryScripting/Source/GeometryScriptingCore/Public/GeometryScript/MeshBoneWeightFunctions.h
--- a/Engine/Plugins/Runtime/GeometryScripting/Source/GeometryScriptingCore/Public/GeometryScript/MeshBoneWeightFunctions.h
+++ b/Engine/Plugins/Runtime/GeometryScripting/Source/GeometryScriptingCore/Public/GeometryScript/MeshBoneWeightFunctions.h
@@ -192,27 +192,27 @@ struct GEOMETRYSCRIPTINGCORE_API FGeometryScriptBoneInfo
 	GENERATED_BODY()
 
 	/** Index of the bone in the skeletal hierarchy. */
-	UPROPERTY(BlueprintReadWrite, Category = Bone)
+	UPROPERTY(EditAnywhere, BlueprintReadWrite, Category = Bone)
 	int Index = INDEX_NONE;
 
 	/** Bone name. */
-	UPROPERTY(BlueprintReadWrite, Category = Bone)
+	UPROPERTY(EditAnywhere, BlueprintReadWrite, Category = Bone)
 	FName Name = NAME_None;
 	
 	/** Parent bone index. */
-	UPROPERTY(BlueprintReadWrite, Category = Bone)
+	UPROPERTY(EditAnywhere, BlueprintReadWrite, Category = Bone)
 	int ParentIndex = INDEX_NONE;
 
 	/** Local/bone space reference transform. */
-	UPROPERTY(BlueprintReadWrite, Category = Bone)
+	UPROPERTY(EditAnywhere, BlueprintReadWrite, Category = Bone)
 	FTransform LocalTransform = FTransform::Identity;
 
 	/** Global/world space reference transform. */
-	UPROPERTY(BlueprintReadWrite, Category = Bone)
+	UPROPERTY(EditAnywhere, BlueprintReadWrite, Category = Bone)
 	FTransform WorldTransform = FTransform::Identity;
 
 	/** Bone color. */
-	UPROPERTY(BlueprintReadWrite, Category = Bone)
+	UPROPERTY(EditAnywhere, BlueprintReadWrite, Category = Bone)
 	FLinearColor Color = FLinearColor::White;
 };
 
diff --git a/Engine/Plugins/Runtime/GeometryScripting/Source/GeometryScriptingCore/Public/GeometryScript/MeshDeformFunctions.h b/Engine/Plugins/Runtime/GeometryScripting/Source/GeometryScriptingCore/Public/GeometryScript/MeshDeformFunctions.h
--- a/Engine/Plugins/Runtime/GeometryScripting/Source/GeometryScriptingCore/Public/GeometryScript/MeshDeformFunctions.h
+++ b/Engine/Plugins/Runtime/GeometryScripting/Source/GeometryScriptingCore/Public/GeometryScript/MeshDeformFunctions.h
@@ -21,7 +21,7 @@ public:
 	bool bSymmetricExtents = true;
 
 	/** Lower extent used when bSymmetricExtents = false */
-	UPROPERTY(BlueprintReadWrite, Category = Options, meta = (EditCondition = "bSymmetricExtents == false"))
+	UPROPERTY(EditAnywhere, BlueprintReadWrite, Category = Options, meta = (EditCondition = "bSymmetricExtents == false"))
 	float LowerExtent = 10;
 
 	/** If true, the Bend is "centered" at the Origin, ie the regions on either side of the extents are rigidly transformed. If false, the Bend begins at the start of the Lower Extents, and the "lower" region is not affected. */
@@ -41,7 +41,7 @@ public:
 	bool bSymmetricExtents = true;
 
 	/** Lower extent used when bSymmetricExtents = false */
-	UPROPERTY(BlueprintReadWrite, Category = Options, meta = (EditCondition = "bSymmetricExtents == false"))
+	UPROPERTY(EditAnywhere, BlueprintReadWrite, Category = Options, meta = (EditCondition = "bSymmetricExtents == false"))
 	float LowerExtent = 10;
 
 	/** If true, the Twist is "centered" at the Origin, ie the regions on either side of the extents are rigidly transformed. If false, the Twist begins at the start of the Lower Extents, and the "lower" region is not affected. */
@@ -72,7 +72,7 @@ public:
 	bool bSymmetricExtents = true;
 
 	/** Lower extent used when bSymmetricExtents = false */
-	UPROPERTY(BlueprintReadWrite, Category = Options, meta = (EditCondition = "bSymmetricExtents == false"))
+	UPROPERTY(EditAnywhere, BlueprintReadWrite, Category = Options, meta = (EditCondition = "bSymmetricExtents == false"))
 	float LowerExtent = 10;
 
 	/** Determines the profile used as a displacement */
diff --git a/Engine/Plugins/Runtime/GeometryScripting/Source/GeometryScriptingCore/Public/GeometryScript/MeshModelingFunctions.h b/Engine/Plugins/Runtime/GeometryScripting/Source/GeometryScriptingCore/Public/GeometryScript/MeshModelingFunctions.h
--- a/Engine/Plugins/Runtime/GeometryScripting/Source/GeometryScriptingCore/Public/GeometryScript/MeshModelingFunctions.h
+++ b/Engine/Plugins/Runtime/GeometryScripting/Source/GeometryScriptingCore/Public/GeometryScript/MeshModelingFunctions.h
@@ -217,26 +217,26 @@ public:
 	 * If true the set of beveled PolyGroup edges is limited to those that 
 	 * are fully or partially contained within the (transformed) FilterBox
 	 */
-	UPROPERTY(BlueprintReadWrite, Category = FilterShape, AdvancedDisplay)
+	UPROPERTY(EditAnywhere, BlueprintReadWrite, Category = FilterShape, AdvancedDisplay)
 	bool bApplyFilterBox = false;
 
 	/**
 	 * Bounding Box used for edge filtering
 	 */
-	UPROPERTY(BlueprintReadWrite, Category = FilterShape, AdvancedDisplay)
+	UPROPERTY(EditAnywhere, BlueprintReadWrite, Category = FilterShape, AdvancedDisplay)
 	FBox FilterBox = FBox(EForceInit::ForceInit);
 
 	/**
 	 * Transform applied to the FilterBox
 	 */
-	UPROPERTY(BlueprintReadWrite, Category = FilterShape, AdvancedDisplay)
+	UPROPERTY(EditAnywhere, BlueprintReadWrite, Category = FilterShape, AdvancedDisplay)
 	FTransform FilterBoxTransform = FTransform::Identity;
 
 	/**
 	 * If true, then only PolyGroup edges that are fully contained within the filter box will be beveled,
 	 * otherwise the edge will be beveled if any vertex is within the filter box.
 	 */
-	UPROPERTY(BlueprintReadWrite, Category = FilterShape, AdvancedDisplay)
+	UPROPERTY(EditAnywhere, BlueprintReadWrite, Category = FilterShape, AdvancedDisplay)
 	bool bFullyContained = true;
 };
 
diff --git a/Engine/Plugins/Runtime/GeometryScripting/Source/GeometryScriptingCore/Public/GeometryScript/MeshRemeshFunctions.h b/Engine/Plugins/Runtime/GeometryScripting/Source/GeometryScriptingCore/Public/GeometryScript/MeshRemeshFunctions.h
--- a/Engine/Plugins/Runtime/GeometryScripting/Source/GeometryScriptingCore/Public/GeometryScript/MeshRemeshFunctions.h
+++ b/Engine/Plugins/Runtime/GeometryScripting/Source/GeometryScriptingCore/Public/GeometryScript/MeshRemeshFunctions.h
@@ -57,63 +57,63 @@ struct GEOMETRYSCRIPTINGCORE_API FGeometryScriptRemeshOptions
 public:
 
 	/** When enabled, all mesh attributes are discarded, so UV and Normal Seams can be ignored. New per-vertex normals are computed. */
-	UPROPERTY(BlueprintReadWrite, Category = Options)
+	UPROPERTY(EditAnywhere, BlueprintReadWrite, Category = Options)
 	bool bDiscardAttributes = false;
 
 	/** When enabled, mesh vertices are projected back onto the input mesh surface during Remeshing, preserving the shape */
-	UPROPERTY(BlueprintReadWrite, Category = Options)
+	UPROPERTY(EditAnywhere, BlueprintReadWrite, Category = Options)
 	bool bReprojectToInputMesh = true;
 
 	/** Type of 3D Mesh Smoothing to apply during Remeshing. Disable by setting SmoothingRate = 0 */
-	UPROPERTY(BlueprintReadWrite, Category = Options)
+	UPROPERTY(EditAnywhere, BlueprintReadWrite, Category = Options)
 	EGeometryScriptRemeshSmoothingType SmoothingType = EGeometryScriptRemeshSmoothingType::Mixed;
 
 	/** Smoothing Rate/Speed. Faster Smoothing results in a more regular mesh, but also more potential for undesirable 3D shape change and UV distortion */
-	UPROPERTY(BlueprintReadWrite, Category = Options, meta = (UIMin = 0, UIMax = 1, ClampMin = 0, ClampMax = 1))
+	UPROPERTY(EditAnywhere, BlueprintReadWrite, Category = Options, meta = (UIMin = 0, UIMax = 1, ClampMin = 0, ClampMax = 1))
 	float SmoothingRate = 0.25f;
 
 	/** Constraints on the open mesh boundary/border edges */
-	UPROPERTY(BlueprintReadWrite, Category = Options)
+	UPROPERTY(EditAnywhere, BlueprintReadWrite, Category = Options)
 	EGeometryScriptRemeshEdgeConstraintType MeshBoundaryConstraint = EGeometryScriptRemeshEdgeConstraintType::Free;
 
 	/** Constraints on the mesh boundary/border edges between different PolyGroups of the Mesh */
-	UPROPERTY(BlueprintReadWrite, Category = Options)
+	UPROPERTY(EditAnywhere, BlueprintReadWrite, Category = Options)
 	EGeometryScriptRemeshEdgeConstraintType GroupBoundaryConstraint = EGeometryScriptRemeshEdgeConstraintType::Free;
 
 	/** Constraints on the mesh boundary/border edges between different Material Results of the Mesh */
-	UPROPERTY(BlueprintReadWrite, Category = Options)
+	UPROPERTY(EditAnywhere, BlueprintReadWrite, Category = Options)
 	EGeometryScriptRemeshEdgeConstraintType MaterialBoundaryConstraint = EGeometryScriptRemeshEdgeConstraintType::Free;
 
 	/** Enable/Disable Edge Flips during Remeshing. Disabling flips will significantly reduce the output mesh quality */
-	UPROPERTY(BlueprintReadWrite, Category = Options)
+	UPROPERTY(EditAnywhere, BlueprintReadWrite, Category = Options)
 	bool bAllowFlips = true;
 
 	/** Enable/Disable Edge Splits during Remeshing. Disabling Splits will prevent the mesh density from increasing. */
-	UPROPERTY(BlueprintReadWrite, Category = Options)
+	UPROPERTY(EditAnywhere, BlueprintReadWrite, Category = Options)
 	bool bAllowSplits = true;
 
 	/** Enable/Disable Edge Collapses during Remeshing. Disabling Collapses will prevent the mesh density from decreasing. */
-	UPROPERTY(BlueprintReadWrite, Category = Options)
+	UPROPERTY(EditAnywhere, BlueprintReadWrite, Category = Options)
 	bool bAllowCollapses = true;
 
 	/** When Enabled, Flips and Collapses will be skipped if they would flip any triangle face normals */
-	UPROPERTY(BlueprintReadWrite, Category = Options)
+	UPROPERTY(EditAnywhere, BlueprintReadWrite, Category = Options)
 	bool bPreventNormalFlips = true;
 
 	/** When Enabled, Flips and Collapses will be skipped if they would create tiny degenerate triangles */
-	UPROPERTY(BlueprintReadWrite, Category = Options)
+	UPROPERTY(EditAnywhere, BlueprintReadWrite, Category = Options)
 	bool bPreventTinyTriangles = true;
 
 	/** By default, remeshing is accelerated by tracking a queue of edges that need to be processed. This is signficantly faster but can produce a lower quality output. Enable this option to use a more expensive strategy that guarantees maximum quality.  */
-	UPROPERTY(BlueprintReadWrite, Category = Options)
+	UPROPERTY(EditAnywhere, BlueprintReadWrite, Category = Options)
 	bool bUseFullRemeshPasses = false;
 
 	/** Maximum Number of iterations of the Remeshing Strategy to apply to the Mesh. More iterations are generally more expensive (much moreso with bUseFullRemeshPasses = true) */
-	UPROPERTY(BlueprintReadWrite, Category = Options, meta = (UIMin = 0, ClampMin = 0))
+	UPROPERTY(EditAnywhere, BlueprintReadWrite, Category = Options, meta = (UIMin = 0, ClampMin = 0))
 	int32 RemeshIterations = 20;
 
 	/** If enabled, the output mesh is automatically compacted to remove gaps in the index space. This is expensive and can be disabled by advanced users. */
-	UPROPERTY(BlueprintReadWrite, Category = Options)
+	UPROPERTY(EditAnywhere, BlueprintReadWrite, Category = Options)
 	bool bAutoCompact = true;
 };
 
@@ -127,15 +127,15 @@ struct GEOMETRYSCRIPTINGCORE_API FGeometryScriptUniformRemeshOptions
 	GENERATED_BODY()
 public:
 	/** Method used to define target/goal of Uniform Remeshing */
-	UPROPERTY(BlueprintReadWrite, Category = Options)
+	UPROPERTY(EditAnywhere, BlueprintReadWrite, Category = Options)
 	EGeometryScriptUniformRemeshTargetType TargetType = EGeometryScriptUniformRemeshTargetType::TriangleCount;
 
 	/** Approximate Target Triangle Count, combined with mesh surface area to derive a TargetEdgeLength */
-	UPROPERTY(BlueprintReadWrite, Category = Options, meta = (UIMin = 0, ClampMin = 0, EditCondition = "TargetType == EGeometryScriptUniformRemeshTargetType::TriangleCount"))
+	UPROPERTY(EditAnywhere, BlueprintReadWrite, Category = Options, meta = (UIMin = 0, ClampMin = 0, EditCondition = "TargetType == EGeometryScriptUniformRemeshTargetType::TriangleCount"))
 	int32 TargetTriangleCount = 5000;
 
 	/** Explicit Target Edge Length that is desired in the output uniform mesh */
-	UPROPERTY(BlueprintReadWrite, Category = Options, meta = (UIMin = 0, ClampMin = 0, EditCondition = "TargetType == EGeometryScriptUniformRemeshTargetType::TargetEdgeLength"))
+	UPROPERTY(EditAnywhere, BlueprintReadWrite, Category = Options, meta = (UIMin = 0, ClampMin = 0, EditCondition = "TargetType == EGeometryScriptUniformRemeshTargetType::TargetEdgeLength"))
 	float TargetEdgeLength = 1.0f;
 };
 
diff --git a/Engine/Plugins/Runtime/GeometryScripting/Source/GeometryScriptingCore/Public/GeometryScript/MeshSimplifyFunctions.h b/Engine/Plugins/Runtime/GeometryScripting/Source/GeometryScriptingCore/Public/GeometryScript/MeshSimplifyFunctions.h
--- a/Engine/Plugins/Runtime/GeometryScripting/Source/GeometryScriptingCore/Public/GeometryScript/MeshSimplifyFunctions.h
+++ b/Engine/Plugins/Runtime/GeometryScripting/Source/GeometryScriptingCore/Public/GeometryScript/MeshSimplifyFunctions.h
@@ -19,7 +19,7 @@ public:
 	float AngleThreshold = 0.001;
 
 	/** If enabled, the simplified mesh is automatically compacted to remove gaps in the index space. This is expensive and can be disabled by advanced users. */
-	UPROPERTY(BlueprintReadWrite, Category = Options)
+	UPROPERTY(EditAnywhere, BlueprintReadWrite, Category = Options)
 	bool bAutoCompact = true;
 };
 
@@ -33,7 +33,7 @@ public:
 	float AngleThreshold = 0.001;
 
 	/** If enabled, the simplified mesh is automatically compacted to remove gaps in the index space. This is expensive and can be disabled by advanced users. */
-	UPROPERTY(BlueprintReadWrite, Category = Options)
+	UPROPERTY(EditAnywhere, BlueprintReadWrite, Category = Options)
 	bool bAutoCompact = true;
 };
 
@@ -72,7 +72,7 @@ public:
 	bool bRetainQuadricMemory = false;
 
 	/** If enabled, the simplified mesh is automatically compacted to remove gaps in the index space. This is expensive and can be disabled by advanced users. */
-	UPROPERTY(BlueprintReadWrite, Category = Options)
+	UPROPERTY(EditAnywhere, BlueprintReadWrite, Category = Options)
 	bool bAutoCompact = true;
 };
 
diff --git a/Engine/Plugins/Runtime/GeometryScripting/Source/GeometryScriptingCore/Public/GeometryScript/MeshVoxelFunctions.h b/Engine/Plugins/Runtime/GeometryScripting/Source/GeometryScriptingCore/Public/GeometryScript/MeshVoxelFunctions.h
--- a/Engine/Plugins/Runtime/GeometryScripting/Source/GeometryScriptingCore/Public/GeometryScript/MeshVoxelFunctions.h
+++ b/Engine/Plugins/Runtime/GeometryScripting/Source/GeometryScriptingCore/Public/GeometryScript/MeshVoxelFunctions.h
@@ -30,11 +30,11 @@ public:
 	EGeometryScriptGridSizingMethod SizeMethod = EGeometryScriptGridSizingMethod::GridResolution;
 
 	/** Use a specific grid cell size, and construct a grid with dimensions large enough to contain the target object */
-	UPROPERTY(BlueprintReadWrite, Category = Options, meta = (EditCondition = "SizeMethod == EGeometryScriptGridSizingMethod::GridCellSize"));
+	UPROPERTY(EditAnywhere, BlueprintReadWrite, Category = Options, meta = (EditCondition = "SizeMethod == EGeometryScriptGridSizingMethod::GridCellSize"));
 	float GridCellSize = 0.5;
 
 	/** Use a specific grid resolution, with the grid cell size derived form the target object bounds such that this is the number of cells along the longest box dimension */
-	UPROPERTY(BlueprintReadWrite, Category = Options, meta = (EditCondition = "SizeMethod == EGeometryScriptGridSizingMethod::GridResolution"))
+	UPROPERTY(EditAnywhere, BlueprintReadWrite, Category = Options, meta = (EditCondition = "SizeMethod == EGeometryScriptGridSizingMethod::GridResolution"))
 	int GridResolution = 64;
 };
 
@@ -60,11 +60,11 @@ public:
 	int SurfaceSearchSteps = 3;
 
 	/** When enabled, regions of the input mesh that have open boundaries (ie "shells") are thickened by extruding them into closed solids. This may be expensive on large meshes. */
-	UPROPERTY(BlueprintReadWrite, Category = Options)
+	UPROPERTY(EditAnywhere, BlueprintReadWrite, Category = Options)
 	bool bThickenShells = false;
 
 	/** Open Shells are Thickened by offsetting vertices along their averaged vertex normals by this amount. Dimension is but clamped to twice the grid cell size. */
-	UPROPERTY(BlueprintReadWrite, Category = Options)
+	UPROPERTY(EditAnywhere, BlueprintReadWrite, Category = Options)
 	double ShellThickness = 1.0;
 
 };
diff --git a/Engine/Plugins/Runtime/GeometryScripting/Source/GeometryScriptingCore/Public/GeometryScript/PolyPathFunctions.h b/Engine/Plugins/Runtime/GeometryScripting/Source/GeometryScriptingCore/Public/GeometryScript/PolyPathFunctions.h
--- a/Engine/Plugins/Runtime/GeometryScripting/Source/GeometryScriptingCore/Public/GeometryScript/PolyPathFunctions.h
+++ b/Engine/Plugins/Runtime/GeometryScripting/Source/GeometryScriptingCore/Public/GeometryScript/PolyPathFunctions.h
@@ -35,10 +35,10 @@ struct GEOMETRYSCRIPTINGCORE_API FGeometryScriptSplineSamplingOptions
 {
 	GENERATED_BODY()
 public:
-	UPROPERTY(BlueprintReadWrite, Category = Options, meta = (EditCondition = "SampleSpacing != EGeometryScriptSampleSpacing::ErrorTolerance"))
+	UPROPERTY(EditAnywhere, BlueprintReadWrite, Category = Options, meta = (EditCondition = "SampleSpacing != EGeometryScriptSampleSpacing::ErrorTolerance"))
 	int32 NumSamples = 10;
 
-	UPROPERTY(BlueprintReadWrite, Category = Options, meta = (EditCondition = "SampleSpacing == EGeometryScriptSampleSpacing::ErrorTolerance"))
+	UPROPERTY(EditAnywhere, BlueprintReadWrite, Category = Options, meta = (EditCondition = "SampleSpacing == EGeometryScriptSampleSpacing::ErrorTolerance"))
 	float ErrorTolerance = 1;
 
 	UPROPERTY(EditAnywhere, BlueprintReadWrite, Category = Options)
@@ -52,11 +52,11 @@ public:
 	EGeometryScriptEvaluateSplineRange RangeMethod = EGeometryScriptEvaluateSplineRange::FullSpline;
 
 	// If not evaluating the full spline, where to start sampling. Expressed in units based on the EvaluateRange value.
-	UPROPERTY(BlueprintReadWrite, Category = Options, meta = (EditCondition = "EvaluateRange != EGeometryScriptEvaluateSplineRange::FullSpline"))
+	UPROPERTY(EditAnywhere, BlueprintReadWrite, Category = Options, meta = (EditCondition = "EvaluateRange != EGeometryScriptEvaluateSplineRange::FullSpline"))
 	float RangeStart = 0;
 
 	// If not evaluating the full spline, where to stop sampling. Expressed in units based on the EvaluateRange value.
-	UPROPERTY(BlueprintReadWrite, Category = Options, meta = (EditCondition = "EvaluateRange != EGeometryScriptEvaluateSplineRange::FullSpline"))
+	UPROPERTY(EditAnywhere, BlueprintReadWrite, Category = Options, meta = (EditCondition = "EvaluateRange != EGeometryScriptEvaluateSplineRange::FullSpline"))
 	float RangeEnd = 1;
 
 };
diff --git a/Engine/Plugins/Runtime/GeometryScripting/Source/GeometryScriptingCore/Public/GeometryScript/PolygonFunctions.h b/Engine/Plugins/Runtime/GeometryScripting/Source/GeometryScriptingCore/Public/GeometryScript/PolygonFunctions.h
--- a/Engine/Plugins/Runtime/GeometryScripting/Source/GeometryScriptingCore/Public/GeometryScript/PolygonFunctions.h
+++ b/Engine/Plugins/Runtime/GeometryScripting/Source/GeometryScriptingCore/Public/GeometryScript/PolygonFunctions.h
@@ -40,23 +40,23 @@ struct GEOMETRYSCRIPTINGCORE_API FGeometryScriptPolygonOffsetOptions
 public:
 
 	// How to join / extend corners between two edges
-	UPROPERTY(BlueprintReadWrite, Category = Options)
+	UPROPERTY(EditAnywhere, BlueprintReadWrite, Category = Options)
 	EGeometryScriptPolyOffsetJoinType JoinType = EGeometryScriptPolyOffsetJoinType::Square;
 
 	// if JoinType is Miter, limits how far the miter can extend
-	UPROPERTY(BlueprintReadWrite, Category = Options)
+	UPROPERTY(EditAnywhere, BlueprintReadWrite, Category = Options)
 	double MiterLimit = 2.0;
 
 	// Whether to apply the offset to both sides of the polygon, i.e. adding an inner hole to any polygon. If false, the offset is only applied to one side.
-	UPROPERTY(BlueprintReadWrite, Category = Options)
+	UPROPERTY(EditAnywhere, BlueprintReadWrite, Category = Options)
 	bool bOffsetBothSides = false;
 
 	// Scales the default number of vertices (per radian) used for round joins.
-	UPROPERTY(BlueprintReadWrite, Category = Options, meta = (ClampMin = "0"))
+	UPROPERTY(EditAnywhere, BlueprintReadWrite, Category = Options, meta = (ClampMin = "0"))
 	double StepsPerRadianScale = 1.0;
 
 	// Maximum vertices per radian for round joins. Only applied if > 0.
-	UPROPERTY(BlueprintReadWrite, Category = Options, meta = (ClampMin = "-1"))
+	UPROPERTY(EditAnywhere, BlueprintReadWrite, Category = Options, meta = (ClampMin = "-1"))
 	double MaximumStepsPerRadian = 10.0;
 };
 
@@ -68,23 +68,23 @@ struct GEOMETRYSCRIPTINGCORE_API FGeometryScriptOpenPathOffsetOptions
 public:
 
 	// How to join / extend corners between two edges
-	UPROPERTY(BlueprintReadWrite, Category = Options)
+	UPROPERTY(EditAnywhere, BlueprintReadWrite, Category = Options)
 	EGeometryScriptPolyOffsetJoinType JoinType = EGeometryScriptPolyOffsetJoinType::Square;
 
 	// if JoinType is Miter, limits how far the miter can extend
-	UPROPERTY(BlueprintReadWrite, Category = Options)
+	UPROPERTY(EditAnywhere, BlueprintReadWrite, Category = Options)
 	double MiterLimit = 2.0;
 
 	// How the ends of a path should be closed off
-	UPROPERTY(BlueprintReadWrite, Category = Options)
+	UPROPERTY(EditAnywhere, BlueprintReadWrite, Category = Options)
 	EGeometryScriptPathOffsetEndType EndType = EGeometryScriptPathOffsetEndType::Square;
 
 	// Scales the default number of vertices (per radian) used for round joins and ends.
-	UPROPERTY(BlueprintReadWrite, Category = Options, meta = (ClampMin = "0"))
+	UPROPERTY(EditAnywhere, BlueprintReadWrite, Category = Options, meta = (ClampMin = "0"))
 	double StepsPerRadianScale = 1.0;
 
 	// Maximum vertices per radian for round joins and ends. Only applied if > 0.
-	UPROPERTY(BlueprintReadWrite, Category = Options, meta = (ClampMin = "-1"))
+	UPROPERTY(EditAnywhere, BlueprintReadWrite, Category = Options, meta = (ClampMin = "-1"))
 	double MaximumStepsPerRadian = 10.0;
 
 };
diff --git a/Engine/Plugins/Runtime/GeometryScripting/Source/GeometryScriptingCore/Public/GeometryScript/VolumeTextureBakeFunctions.h b/Engine/Plugins/Runtime/GeometryScripting/Source/GeometryScriptingCore/Public/GeometryScript/VolumeTextureBakeFunctions.h
--- a/Engine/Plugins/Runtime/GeometryScripting/Source/GeometryScriptingCore/Public/GeometryScript/VolumeTextureBakeFunctions.h
+++ b/Engine/Plugins/Runtime/GeometryScripting/Source/GeometryScriptingCore/Public/GeometryScript/VolumeTextureBakeFunctions.h
@@ -38,23 +38,23 @@ public:
 	// Whether to compute distances only in a band around the surface (faster) or compute the full grid
 	// Note: If full grid is computed, the distances will still be more accurately computed in the narrow band
 	// In narrow band mode, values outside the band will have a large magnitude with the correct sign
-	UPROPERTY(BlueprintReadWrite, Category = Options)
+	UPROPERTY(EditAnywhere, BlueprintReadWrite, Category = Options)
 	EDistanceFieldComputeMode ComputeMode = EDistanceFieldComputeMode::NarrowBand;
 
 	// Width of the narrow band where distances are computed accurately
-	UPROPERTY(BlueprintReadWrite, Category = Options)
+	UPROPERTY(EditAnywhere, BlueprintReadWrite, Category = Options)
 	float NarrowBandWidth = 2;
 
 	// Whether Narrow Band Width is expressed as a number of voxels (rounded up to nearest int) or a distance
-	UPROPERTY(BlueprintReadWrite, Category = Options)
+	UPROPERTY(EditAnywhere, BlueprintReadWrite, Category = Options)
 	EDistanceFieldUnits NarrowBandUnits = EDistanceFieldUnits::NumberOfVoxels;
 
 	// Number of voxels to use along each axis
-	UPROPERTY(BlueprintReadWrite, Category = Options)
+	UPROPERTY(EditAnywhere, BlueprintReadWrite, Category = Options)
 	FIntVector VoxelsPerDimensions = FIntVector(32,32,32);
 
 	// Whether to round voxel count on each dimension up to the nearest power of two
-	UPROPERTY(BlueprintReadWrite, Category = Options)
+	UPROPERTY(EditAnywhere, BlueprintReadWrite, Category = Options)
 	bool bRequirePower2 = true;
 };
 
@@ -65,11 +65,11 @@ struct FDistanceFieldToTextureSettings
 public:
 
 	// Scale values by this amount before writing them to the texture
-	UPROPERTY(BlueprintReadWrite, Category = Options)
+	UPROPERTY(EditAnywhere, BlueprintReadWrite, Category = Options)
 	float Scale = 1;
 
 	// Offset values by this amount before writing them to the texture (after applying Scale)
-	UPROPERTY(BlueprintReadWrite, Category = Options)
+	UPROPERTY(EditAnywhere, BlueprintReadWrite, Category = Options)
 	float Offset = 0;
 };
 
