redo asset -> material list + slot names functions to a more consistent form, and add helpers to conver to/from material map format as a separate step

diff --git a/Engine/Plugins/PCGInterops/PCGGeometryScriptInterop/Source/PCGGeometryScriptInterop/Private/Elements/PCGMeshSampler.cpp b/Engine/Plugins/PCGInterops/PCGGeometryScriptInterop/Source/PCGGeometryScriptInterop/Private/Elements/PCGMeshSampler.cpp
--- a/Engine/Plugins/PCGInterops/PCGGeometryScriptInterop/Source/PCGGeometryScriptInterop/Private/Elements/PCGMeshSampler.cpp
+++ b/Engine/Plugins/PCGInterops/PCGGeometryScriptInterop/Source/PCGGeometryScriptInterop/Private/Elements/PCGMeshSampler.cpp
@@ -481,10 +481,11 @@ bool FPCGMeshSamplerElement::PrepareDataInternal(FPCGContext* InContext) const
 			Context->DynamicMeshes.Add(FPCGContext::NewObject_AnyThread<UDynamicMesh>(Context));
 			FGeometryScriptCopyMeshFromAssetOptions Options{};
 
-			UGeometryScriptLibrary_StaticMeshFunctions::CopyMeshFromStaticMesh(StaticMesh, Context->DynamicMeshes.Last(), Options, MeshReadLOD, Outcome, Debug);
+			UGeometryScriptLibrary_StaticMeshFunctions::CopyMeshFromStaticMeshV2(StaticMesh, Context->DynamicMeshes.Last(), Options, MeshReadLOD, Outcome, false, Debug);
 			if (Outcome == EGeometryScriptOutcomePins::Success && Settings->bOutputMaterialInfo)
 			{
-				UGeometryScriptLibrary_StaticMeshFunctions::GetMaterialListFromStaticMesh(StaticMesh, Context->AssetMaterialList.Emplace_GetRef());
+				TArray<FName> MaterialSlotNames;
+				UGeometryScriptLibrary_StaticMeshFunctions::GetMaterialListFromStaticMesh(StaticMesh, Context->AssetMaterialList.Emplace_GetRef(), MaterialSlotNames);
 			}
 		}
 		else
diff --git a/Engine/Plugins/Runtime/GeometryScripting/Source/GeometryScriptingCore/Private/MeshAssetFunctions.cpp b/Engine/Plugins/Runtime/GeometryScripting/Source/GeometryScriptingCore/Private/MeshAssetFunctions.cpp
--- a/Engine/Plugins/Runtime/GeometryScripting/Source/GeometryScriptingCore/Private/MeshAssetFunctions.cpp
+++ b/Engine/Plugins/Runtime/GeometryScripting/Source/GeometryScriptingCore/Private/MeshAssetFunctions.cpp
@@ -458,8 +458,9 @@ int UGeometryScriptLibrary_StaticMeshFunctions::GetNumStaticMeshLODsOfType(
 
 
 void UGeometryScriptLibrary_StaticMeshFunctions::GetMaterialListFromStaticMesh(
-	UStaticMesh* FromStaticMeshAsset,
+	const UStaticMesh* FromStaticMeshAsset,
 	TArray<UMaterialInterface*>& MaterialList,
+	TArray<FName>& MaterialSlotNames,
 	UGeometryScriptDebug* Debug)
 {
 	if (FromStaticMeshAsset == nullptr)
@@ -469,10 +470,67 @@ void UGeometryScriptLibrary_StaticMeshFunctions::GetMaterialListFromStaticMesh(
 	}
 
 	const TArray<FStaticMaterial>& AssetMaterials = FromStaticMeshAsset->GetStaticMaterials();
+	MaterialList.Reset(AssetMaterials.Num());
+	MaterialSlotNames.Reset(AssetMaterials.Num());
 	for (int32 k = 0; k < AssetMaterials.Num(); ++k)
 	{
 		MaterialList.Add(AssetMaterials[k].MaterialInterface);
+		MaterialSlotNames.Add(AssetMaterials[k].MaterialSlotName);
+	}
+}
+
+void UGeometryScriptLibrary_StaticMeshFunctions::GetMaterialListFromSkeletalMesh(
+	const USkeletalMesh* FromSkeletalMeshAsset,
+	TArray<UMaterialInterface*>& MaterialList,
+	TArray<FName>& MaterialSlotNames,
+	UGeometryScriptDebug* Debug)
+{
+	if (FromSkeletalMeshAsset == nullptr)
+	{
+		UE::Geometry::AppendError(Debug, EGeometryScriptErrorType::InvalidInputs, LOCTEXT("GetMaterialListFromSkeletalMesh_InvalidInput1", "GetMaterialListFromSkeletalMesh: FromSkeletalMeshAsset is Null"));
+		return;
+	}
+
+	const TArray<FSkeletalMaterial>& AssetMaterials = FromSkeletalMeshAsset->GetMaterials();
+	MaterialList.Reset(AssetMaterials.Num());
+	MaterialSlotNames.Reset(AssetMaterials.Num());
+	for (int32 k = 0; k < AssetMaterials.Num(); ++k)
+	{
+		MaterialList.Add(AssetMaterials[k].MaterialInterface);
+		MaterialSlotNames.Add(AssetMaterials[k].MaterialSlotName);
+	}
+}
+
+void UGeometryScriptLibrary_StaticMeshFunctions::ConvertMaterialMapToMaterialList(const TMap<FName, UMaterialInterface*>& MaterialMap,
+	TArray<UMaterialInterface*>& MaterialList,
+	TArray<FName>& MaterialSlotNames)
+{
+	MaterialList.Reset(MaterialMap.Num());
+	MaterialSlotNames.Reset(MaterialMap.Num());
+	for (const TPair<FName, UMaterialInterface*> NameMat : MaterialMap)
+	{
+		MaterialList.Add(NameMat.Value);
+		MaterialSlotNames.Add(NameMat.Key);
+	}
+}
+
+TMap<FName, UMaterialInterface*> UGeometryScriptLibrary_StaticMeshFunctions::ConvertMaterialListToMaterialMap(const TArray<UMaterialInterface*>& MaterialList, const TArray<FName>& MaterialSlotNames)
+{
+	TMap<FName, UMaterialInterface*> ToRet;
+	if (MaterialSlotNames.Num() != MaterialList.Num())
+	{
+		UE_LOG(LogGeometry, Warning, TEXT("ConvertMaterialListToMaterialMap: Number of Material Slot Names does not match number of Materials"));
+	}
+
+	ToRet.Reserve(MaterialList.Num());
+	for (int32 Idx = 0; Idx < MaterialList.Num(); ++Idx)
+	{
+		UMaterialInterface* Mat = MaterialList[Idx];
+		// If we have fewer slot names than materials, we will have warned user via above AppendWarning, but make up a slot name so that we still have all materials in the map
+		FName SlotName = MaterialSlotNames.IsValidIndex(Idx) ? MaterialSlotNames[Idx] : FName(FString::Printf(TEXT("%s_%d"), *((Mat) ? Mat->GetName() : TEXT("Material")), Idx));
+		ToRet.Add(SlotName, Mat);
 	}
+	return ToRet;
 }
 
 void UGeometryScriptLibrary_StaticMeshFunctions::GetSectionMaterialListFromStaticMesh(
diff --git a/Engine/Plugins/Runtime/GeometryScripting/Source/GeometryScriptingCore/Private/MeshMaterialFunctions.cpp b/Engine/Plugins/Runtime/GeometryScripting/Source/GeometryScriptingCore/Private/MeshMaterialFunctions.cpp
--- a/Engine/Plugins/Runtime/GeometryScripting/Source/GeometryScriptingCore/Private/MeshMaterialFunctions.cpp
+++ b/Engine/Plugins/Runtime/GeometryScripting/Source/GeometryScriptingCore/Private/MeshMaterialFunctions.cpp
@@ -8,10 +8,6 @@
 #include "UDynamicMesh.h"
 #include "Polygroups/PolygroupSet.h"
 
-#include "Components/PrimitiveComponent.h"
-#include "Engine/StaticMesh.h"
-#include "Engine/SkeletalMesh.h"
-
 #include UE_INLINE_GENERATED_CPP_BY_NAME(MeshMaterialFunctions)
 
 using namespace UE::Geometry;
@@ -808,73 +804,4 @@ UDynamicMesh* UGeometryScriptLibrary_MeshMaterialFunctions::CompactMaterialIDs(
 	return TargetMesh;
 }
 
-TMap<FName, UMaterialInterface*> UGeometryScriptLibrary_MeshMaterialFunctions::GetMaterialSlotToMaterialMapFromComponent(const UPrimitiveComponent* PrimitiveComponent)
-{
-	TMap<FName, UMaterialInterface*> ToRet;
-	TArray<FName> SlotNames = PrimitiveComponent->GetMaterialSlotNames();
-	ToRet.Reserve(SlotNames.Num());
-	for (const FName& SlotName : SlotNames)
-	{
-		ToRet.Add(SlotName, PrimitiveComponent->GetMaterialByName(SlotName));
-	}
-	return ToRet;
-}
-
-
-TMap<FName, UMaterialInterface*> UGeometryScriptLibrary_MeshMaterialFunctions::GetMaterialSlotToMaterialMapFromStaticMesh(const UStaticMesh* StaticMesh, bool bUseImportedSlotNames)
-{
-	TMap<FName, UMaterialInterface*> ToRet;
-	const TArray<FStaticMaterial>& Materials = StaticMesh->GetStaticMaterials();
-	ToRet.Reserve(Materials.Num());
-	for (const FStaticMaterial& Mat : Materials)
-	{
-#if WITH_EDITORONLY_DATA
-		if (bUseImportedSlotNames)
-		{
-			ToRet.Add(Mat.ImportedMaterialSlotName, Mat.MaterialInterface);
-		}
-		else
-#endif
-		{
-			ToRet.Add(Mat.MaterialSlotName, Mat.MaterialInterface);
-		}
-	}
-	return ToRet;
-}
-
-TMap<FName, UMaterialInterface*> UGeometryScriptLibrary_MeshMaterialFunctions::GetMaterialSlotToMaterialMapFromSkeletalMesh(const USkeletalMesh* SkeletalMesh, bool bUseImportedSlotNames)
-{
-	TMap<FName, UMaterialInterface*> ToRet;
-	const TArray<FSkeletalMaterial>& Materials = SkeletalMesh->GetMaterials();
-	ToRet.Reserve(Materials.Num());
-	for (const FSkeletalMaterial& Mat : Materials)
-	{
-#if WITH_EDITORONLY_DATA
-		if (bUseImportedSlotNames)
-		{
-			ToRet.Add(Mat.ImportedMaterialSlotName, Mat.MaterialInterface);
-		}
-		else
-#endif
-		{
-			ToRet.Add(Mat.MaterialSlotName, Mat.MaterialInterface);
-		}
-	}
-	return ToRet;
-}
-
-
-TArray<UMaterialInterface*> UGeometryScriptLibrary_MeshMaterialFunctions::Conv_MaterialsMapToMaterialsArray(const TMap<FName, UMaterialInterface*>& MaterialsMap)
-{
-	TArray<UMaterialInterface*> ToRet;
-	ToRet.Reserve(MaterialsMap.Num());
-	for (const TPair<FName, UMaterialInterface*>& NameMat : MaterialsMap)
-	{
-		// Note we add the materials even if they are null, so we always have the same number of materials in the list as were in the map
-		ToRet.Add(NameMat.Value);
-	}
-	return ToRet;
-}
-
-
 #undef LOCTEXT_NAMESPACE
diff --git a/Engine/Plugins/Runtime/GeometryScripting/Source/GeometryScriptingCore/Public/GeometryScript/MeshAssetFunctions.h b/Engine/Plugins/Runtime/GeometryScripting/Source/GeometryScriptingCore/Public/GeometryScript/MeshAssetFunctions.h
--- a/Engine/Plugins/Runtime/GeometryScripting/Source/GeometryScriptingCore/Public/GeometryScript/MeshAssetFunctions.h
+++ b/Engine/Plugins/Runtime/GeometryScripting/Source/GeometryScriptingCore/Public/GeometryScript/MeshAssetFunctions.h
@@ -270,10 +270,39 @@ public:
 	 */
 	UFUNCTION(BlueprintCallable, Category = "GeometryScript|StaticMesh")
 	static void 
-	GetMaterialListFromStaticMesh(UStaticMesh* FromStaticMeshAsset,
+	GetMaterialListFromStaticMesh(const UStaticMesh* FromStaticMeshAsset,
 		TArray<UMaterialInterface*>& MaterialList,
+		TArray<FName>& MaterialSlotNames,
+		UGeometryScriptDebug* Debug = nullptr);
+
+	/**
+	 * Get the asset materials from the skeletal mesh asset.
+	 * Note: For LOD-specific materials, use GetLODMaterialListFromSkeletalMesh instead.
+	 */
+	UFUNCTION(BlueprintCallable, Category = "GeometryScript|SkeletalMesh")
+	static void
+	GetMaterialListFromSkeletalMesh(const USkeletalMesh* FromSkeletalMeshAsset,
+		TArray<UMaterialInterface*>& MaterialList,
+		TArray<FName>& MaterialSlotNames,
 		UGeometryScriptDebug* Debug = nullptr);
 
+	/**
+	 * Converts material map to a material list and a slot names list. Null materials will be kept in the list, and the list will have the same number of elements as the map.
+	 */
+	UFUNCTION(BlueprintCallable, Category = "GeometryScript|Materials")
+	static void ConvertMaterialMapToMaterialList(const TMap<FName, UMaterialInterface*>& MaterialMap,
+		TArray<UMaterialInterface*>& MaterialList,
+		TArray<FName>& MaterialSlotNames);
+
+	/**
+	 * Converts material list and slot names list to material map, which is the format expected by CreateNewSkeletalMeshAssetFromMesh.
+	 * Material List and Material Slot Names should have the same length. However, if there are fewer slot names than materials, 
+	 * slot names will be auto-generated (as '[Name of material]_[Index]', or 'Material_[Index]' for null materials)
+	 */
+	UFUNCTION(BlueprintCallable, Category = "GeometryScript|Materials")
+	static UPARAM(DisplayName = "Material Map") TMap<FName, UMaterialInterface*> ConvertMaterialListToMaterialMap(const TArray<UMaterialInterface*>& MaterialList,
+		const TArray<FName>& MaterialSlotNames);
+
 	/** 
 	* Extracts a Dynamic Mesh from a Skeletal Mesh Asset. 
 	*/
diff --git a/Engine/Plugins/Runtime/GeometryScripting/Source/GeometryScriptingCore/Public/GeometryScript/MeshMaterialFunctions.h b/Engine/Plugins/Runtime/GeometryScripting/Source/GeometryScriptingCore/Public/GeometryScript/MeshMaterialFunctions.h
--- a/Engine/Plugins/Runtime/GeometryScripting/Source/GeometryScriptingCore/Public/GeometryScript/MeshMaterialFunctions.h
+++ b/Engine/Plugins/Runtime/GeometryScripting/Source/GeometryScriptingCore/Public/GeometryScript/MeshMaterialFunctions.h
@@ -235,30 +235,6 @@ public:
 		bool bRemoveDuplicateMaterials = false,
 		UGeometryScriptDebug* Debug = nullptr);
 
-	/**
-	 * Get a map from material slot names to materials from a given component
-	 */
-	UFUNCTION(BlueprintCallable, Category = "GeometryScript|Materials", meta = (ScriptMethod))
-	static UPARAM(DisplayName = "Slot Name to Material Map") TMap<FName, UMaterialInterface*> GetMaterialSlotToMaterialMapFromComponent(const UPrimitiveComponent* PrimitiveComponent);
-
-	/**
-	 * Get a map from material slot names to materials from a given static mesh
-	 * @param bUseImportedSlotNames If available (i.e., if editor-only data is available), use the imported slot names instead of the current slot names as the map keys.
-	 */
-	UFUNCTION(BlueprintCallable, Category = "GeometryScript|Materials", meta = (ScriptMethod))
-	static UPARAM(DisplayName = "Slot Name to Material Map") TMap<FName, UMaterialInterface*> GetMaterialSlotToMaterialMapFromStaticMesh(const UStaticMesh* StaticMesh, bool bUseImportedSlotNames = false);
 
-	/**
-	 * Get a map from material slot names to materials from a given skeletal mesh
-	 * @param bUseImportedSlotNames If available (i.e., if editor-only data is available), use the imported slot names instead of the current slot names as the map keys.
-	 */
-	UFUNCTION(BlueprintCallable, Category = "GeometryScript|Materials", meta = (ScriptMethod))
-	static UPARAM(DisplayName = "Slot Name to Material Map") TMap<FName, UMaterialInterface*> GetMaterialSlotToMaterialMapFromSkeletalMesh(const USkeletalMesh* SkeletalMesh, bool bUseImportedSlotNames = false);
-
-	/**
-	 * Converts materials map to materials array. Null materials will be kept in the list, and the list will have the same number of elements as the map.
-	 */
-	UFUNCTION(BlueprintPure, meta = (DisplayName = "Material Map To Material List", CompactNodeTitle = "->", BlueprintAutocast), Category = "GeometryScript|Materials")
-	static UPARAM(DisplayName = "Material List") TArray<UMaterialInterface*> Conv_MaterialsMapToMaterialsArray(const TMap<FName, UMaterialInterface*>& MaterialMap);
 
 };
\ No newline at end of file
