add autocreaterefterm to geometry script options structs that were passed by const ref, so users don't have to connect the pin to get default options for these nodes

diff --git a/Engine/Plugins/Runtime/GeometryScripting/Source/GeometryScriptingCore/Public/GeometryScript/CollisionFunctions.h b/Engine/Plugins/Runtime/GeometryScripting/Source/GeometryScriptingCore/Public/GeometryScript/CollisionFunctions.h
--- a/Engine/Plugins/Runtime/GeometryScripting/Source/GeometryScriptingCore/Public/GeometryScript/CollisionFunctions.h
+++ b/Engine/Plugins/Runtime/GeometryScripting/Source/GeometryScriptingCore/Public/GeometryScript/CollisionFunctions.h
@@ -568,7 +568,7 @@ public:
 	 * @param Options		Options controlling the convex decomposition
 	 * @return				The resulting convex hulls as simple collision shapes
 	 */
-	UFUNCTION(BlueprintCallable, Category = "GeometryScript|Collision")
+	UFUNCTION(BlueprintCallable, Category = "GeometryScript|Collision", meta = (AutoCreateRefTerm = "Options"))
 	static UPARAM(DisplayName = "Convex Decomposition") FGeometryScriptSimpleCollision ComputeNavigableConvexDecomposition(
 		const UDynamicMesh* TargetMesh,
 		const FNavigableConvexDecompositionOptions& Options,
diff --git a/Engine/Plugins/Runtime/GeometryScripting/Source/GeometryScriptingCore/Public/GeometryScript/MeshSamplingFunctions.h b/Engine/Plugins/Runtime/GeometryScripting/Source/GeometryScriptingCore/Public/GeometryScript/MeshSamplingFunctions.h
--- a/Engine/Plugins/Runtime/GeometryScripting/Source/GeometryScriptingCore/Public/GeometryScript/MeshSamplingFunctions.h
+++ b/Engine/Plugins/Runtime/GeometryScripting/Source/GeometryScriptingCore/Public/GeometryScript/MeshSamplingFunctions.h
@@ -155,7 +155,7 @@ public:
 	 * @param Box     Bounding Box containing the scene to be captured
 	 * @param Options Defines the Camera viewing directions into the box and other Camera parameters
 	 */
-	UFUNCTION(BlueprintCallable, Category = "GeometryScript|MeshSampling")
+	UFUNCTION(BlueprintCallable, Category = "GeometryScript|MeshSampling", meta = (AutoCreateRefTerm = "Options"))
 	static void
 	ComputeRenderCaptureCamerasForBox(
 		TArray<FGeometryScriptRenderCaptureCamera>& Cameras,
diff --git a/Engine/Plugins/Runtime/GeometryScripting/Source/GeometryScriptingCore/Public/GeometryScript/PointSetFunctions.h b/Engine/Plugins/Runtime/GeometryScripting/Source/GeometryScriptingCore/Public/GeometryScript/PointSetFunctions.h
--- a/Engine/Plugins/Runtime/GeometryScripting/Source/GeometryScriptingCore/Public/GeometryScript/PointSetFunctions.h
+++ b/Engine/Plugins/Runtime/GeometryScripting/Source/GeometryScriptingCore/Public/GeometryScript/PointSetFunctions.h
@@ -86,7 +86,7 @@ public:
 	 * Use K-Means clustering to cluster the given points into a target number of clusters,
 	 * and return an array with a cluster index per point.
 	 */
-	UFUNCTION(BlueprintCallable, Category = "GeometryScript|PointSet")
+	UFUNCTION(BlueprintCallable, Category = "GeometryScript|PointSet", meta = (AutoCreateRefTerm = "Options"))
 	static void KMeansClusterToIDs(
 		const TArray<FVector>& Points,
 		const FGeometryScriptPointClusteringOptions& Options,
@@ -96,7 +96,7 @@ public:
 	 * Use K-Means clustering to cluster the given points into a target number of clusters,
 	 * and return the clusters as an array of lists of point indices.
 	 */
-	UFUNCTION(BlueprintCallable, Category = "GeometryScript|PointSet")
+	UFUNCTION(BlueprintCallable, Category = "GeometryScript|PointSet", meta = (AutoCreateRefTerm = "Options"))
 	static void KMeansClusterToArrays(
 		const TArray<FVector>& Points,
 		const FGeometryScriptPointClusteringOptions& Options,
@@ -110,7 +110,7 @@ public:
 	 * (2) if uniform spacing is requested, points will be ordered to have an octree-style coverage --
 	 *     so the first 8 points will cover the 8 octants (where samples are available) and the subsequent points will progressively fill in the space
 	 */
-	UFUNCTION(BlueprintCallable, Category = "GeometryScript|PointSet")
+	UFUNCTION(BlueprintCallable, Category = "GeometryScript|PointSet", meta = (AutoCreateRefTerm = "Options"))
 	static void DownsamplePoints(
 		const TArray<FVector>& Points,
 		const FGeometryScriptPointPriorityOptions& Options,
@@ -143,7 +143,7 @@ public:
 	 * Convert an array of points from 3D to 2D, by transforming into the given ReferenceFrame and taking the X,Y coordinates
 	 * Note that to transform into the ReferenceFrame, we apply the inverse of the ReferenceFrame's transform.
 	 */
-	UFUNCTION(BlueprintCallable, Category = "GeometryScript|PointSet")
+	UFUNCTION(BlueprintCallable, Category = "GeometryScript|PointSet", meta = (AutoCreateRefTerm = "Options"))
 	static void FlattenPoints(
 		const TArray<FVector>& PointsIn3D,
 		TArray<FVector2D>& PointsIn2D,
@@ -153,7 +153,7 @@ public:
 	/**
 	 * Convert an array of points from 2D to 3D, by transforming out of the given ReferenceFrame, with the given Height for the non-flat axis (default Z).
 	 */
-	UFUNCTION(BlueprintCallable, Category = "GeometryScript|PointSet")
+	UFUNCTION(BlueprintCallable, Category = "GeometryScript|PointSet", meta = (AutoCreateRefTerm = "Options"))
 	static void UnflattenPoints(
 		const TArray<FVector2D>& PointsIn2D,
 		TArray<FVector>& PointsIn3D,
