Geoscripting: Remap bone indices on skin bindings of the source dynamesh to the skeletal mesh asset in CopyMeshToSkeletalMesh.

diff --git a/Engine/Plugins/Runtime/GeometryScripting/Source/GeometryScriptingCore/Private/MeshAssetFunctions.cpp b/Engine/Plugins/Runtime/GeometryScripting/Source/GeometryScriptingCore/Private/MeshAssetFunctions.cpp
--- a/Engine/Plugins/Runtime/GeometryScripting/Source/GeometryScriptingCore/Private/MeshAssetFunctions.cpp
+++ b/Engine/Plugins/Runtime/GeometryScripting/Source/GeometryScriptingCore/Private/MeshAssetFunctions.cpp
@@ -22,9 +22,11 @@
 
 #include "MeshDescriptionToDynamicMesh.h"
 #include "DynamicMeshToMeshDescription.h"
+#include "SkeletalMeshOperations.h"
 #include "StaticMeshLODResourcesToDynamicMesh.h"
 #include "AssetUtils/StaticMeshMaterialUtil.h"
 #include "ConversionUtils/SceneComponentToDynamicMesh.h"
+#include "DynamicMesh/DynamicBoneAttribute.h"
 #include "DynamicMesh/NonManifoldMappingSupport.h"
 
 #include UE_INLINE_GENERATED_CPP_BY_NAME(MeshAssetFunctions)
@@ -909,6 +911,31 @@ UDynamicMesh* UGeometryScriptLibrary_StaticMeshFunctions::CopyMeshToSkeletalMesh
 		UE::Geometry::AppendError(Debug, EGeometryScriptErrorType::InvalidInputs, LOCTEXT("CopyMeshToSkeletalMesh_TargetMeshDescription", "CopyMeshToSkeletalMesh: Failed to generate the mesh data for the Target LOD Index"));
 		return FromDynamicMesh;
 	}
+
+	// Verify that the bones on the dynamic mesh are a proper subset of the bones on the skeletal mesh. The order is not important, since
+	// we re-order as needed below. If the mesh has no bones, then we create, or get, the default skin weight attribute and bind everything 
+	// to root, since we can't verify that any current skin binding is valid.
+	TArray<int32> BoneRemapping;
+	const FDynamicMesh3& Mesh{FromDynamicMesh->GetMeshRef()}; 
+	if (Mesh.HasAttributes() && Mesh.Attributes()->HasBones())
+	{
+		const FDynamicMeshBoneNameAttribute* SrcBoneNames = Mesh.Attributes()->GetBoneNames();
+		TArray<FName> DstBoneNames = ToSkeletalMeshAsset->GetRefSkeleton().GetRawRefBoneNames();
+		for (int32 SrcBoneIndex = 0; SrcBoneIndex < SrcBoneNames->Num(); ++SrcBoneIndex)
+		{
+			const FName SrcBoneName = SrcBoneNames->GetValue(SrcBoneIndex);
+			const int32 DstBoneIndex = DstBoneNames.IndexOfByKey(SrcBoneName);
+			if (DstBoneIndex == INDEX_NONE)
+			{
+				UE::Geometry::AppendError(Debug, EGeometryScriptErrorType::InvalidInputs,
+					FText::Format(LOCTEXT("CopyMeshToSkeletalMesh_MissingBonesOnAsset", "CopyMeshToSkeletalMesh: Source geometry contains bone '{0}' which does not exist on the skeletal mesh asset ({1})."),
+						FText::FromName(SrcBoneName), FText::FromString(ToSkeletalMeshAsset->GetPackage()->GetPathName())));
+				return FromDynamicMesh;
+			}
+
+			BoneRemapping.Add(DstBoneIndex);
+		}
+	}
 	
 	FSkeletalMeshAttributes MeshAttributes(*MeshDescription);
 	MeshAttributes.Register();
@@ -921,7 +948,27 @@ UDynamicMesh* UGeometryScriptLibrary_StaticMeshFunctions::CopyMeshToSkeletalMesh
 	{
 		Converter.Convert(&ReadMesh, *MeshDescription, !Options.bEnableRecomputeTangents);
 	});
+	
+	if (!BoneRemapping.IsEmpty())
+	{
+		FSkeletalMeshOperations::RemapBoneIndicesOnSkinWeightAttribute(*MeshDescription, BoneRemapping);
+	}
+	else
+	{
+		using namespace UE::AnimationCore;
+		FBoneWeight RootWeight(0, 1.0f);
+		FBoneWeights RootBinding = FBoneWeights::Create({RootWeight});
+		for (const FName AttributeName: MeshAttributes.GetSkinWeightProfileNames())
+		{
+			FSkinWeightsVertexAttributesRef SkinWeights(MeshAttributes.GetVertexSkinWeights(AttributeName));
 
+			for (FVertexID VertexID: MeshDescription->Vertices().GetElementIDs())
+			{
+				SkinWeights.Set(VertexID, RootBinding);
+			}
+		}
+	}
+	
 	FSkeletalMeshLODInfo* SkeletalLODInfo = ToSkeletalMeshAsset->GetLODInfo(TargetLOD.LODIndex);
 	SkeletalLODInfo->BuildSettings.bRecomputeNormals = Options.bEnableRecomputeNormals;
 	SkeletalLODInfo->BuildSettings.bRecomputeTangents = Options.bEnableRecomputeTangents;
diff --git a/Engine/Source/Runtime/SkeletalMeshDescription/Private/SkeletalMeshOperations.cpp b/Engine/Source/Runtime/SkeletalMeshDescription/Private/SkeletalMeshOperations.cpp
--- a/Engine/Source/Runtime/SkeletalMeshDescription/Private/SkeletalMeshOperations.cpp
+++ b/Engine/Source/Runtime/SkeletalMeshDescription/Private/SkeletalMeshOperations.cpp
@@ -231,4 +231,42 @@ bool FSkeletalMeshOperations::CopySkinWeightAttributeFromMesh(
 	return true;
 }
 
+bool FSkeletalMeshOperations::RemapBoneIndicesOnSkinWeightAttribute(FMeshDescription& InMesh, TConstArrayView<int32> InBoneIndexMapping)
+{
+	using namespace UE::AnimationCore;
+	
+	FSkeletalMeshAttributes MeshAttributes(InMesh);
+
+	// Don't renormalize, since we are not changing the weights or order.
+	FBoneWeightsSettings Settings;
+	Settings.SetNormalizeType(EBoneWeightNormalizeType::None);
+	
+	TArray<FBoneWeight> NewBoneWeights;
+	for (const FName AttributeName: MeshAttributes.GetSkinWeightProfileNames())
+	{
+		FSkinWeightsVertexAttributesRef SkinWeights(MeshAttributes.GetVertexSkinWeights(AttributeName));
+
+		for (FVertexID VertexID: InMesh.Vertices().GetElementIDs())
+		{
+			FVertexBoneWeights OldBoneWeights = SkinWeights.Get(VertexID);
+			NewBoneWeights.Reset(OldBoneWeights.Num());
+
+			for (FBoneWeight BoneWeight: OldBoneWeights)
+			{
+				if (!ensure(InBoneIndexMapping.IsValidIndex(BoneWeight.GetBoneIndex())))
+				{
+					return false;
+				}
+				
+				BoneWeight.SetBoneIndex(InBoneIndexMapping[BoneWeight.GetBoneIndex()]);
+				NewBoneWeights.Add(BoneWeight);
+			}
+
+			SkinWeights.Set(VertexID, FBoneWeights::Create(NewBoneWeights, Settings));
+		}
+	}
+	return true;
+}
+
+
 #undef LOCTEXT_NAMESPACE
diff --git a/Engine/Source/Runtime/SkeletalMeshDescription/Public/SkeletalMeshOperations.h b/Engine/Source/Runtime/SkeletalMeshDescription/Public/SkeletalMeshOperations.h
--- a/Engine/Source/Runtime/SkeletalMeshDescription/Public/SkeletalMeshOperations.h
+++ b/Engine/Source/Runtime/SkeletalMeshDescription/Public/SkeletalMeshOperations.h
@@ -51,4 +51,18 @@ public:
 		const FName InTargetProfile,
 		const TMap<int32, int32>* SourceBoneIndexToTargetBoneIndexMap
 		);
+
+	/** Remaps the bone indices on all skin weight attributes from one index to another. The array view should contain a full mapping of all the
+	 *  bone indices contained in the skin weights. The array is indexed by the current bone index and the value at that index is the new bone index.
+	 *  If the mapping is incomplete or if two entries map to the same bone, the result is undefined. No prior checking is performed.
+	 *  \param InMesh The mesh on which to modify all skin weight attributes to remap their bone indices.
+	 *  \param InBoneIndexMapping The mapping from one bone index to another. The old bone index is used to index into the array, the value at that
+	 *     position is the new bone index.
+	 *  \return true if the operation was successful. If the mapping array was incomplete then this will return false. If there are no skin weight
+	 *     attributes on the mesh, then the operation is still deemed successful.
+	 */  
+	static SKELETALMESHDESCRIPTION_API bool RemapBoneIndicesOnSkinWeightAttribute(
+		FMeshDescription& InMesh,
+		TConstArrayView<int32> InBoneIndexMapping
+		);
 };
