Geoscripting: Make it an option to rebind weights to the ref skeleton of the skeletal mesh asset.

#rnx

diff --git a/Engine/Plugins/Runtime/GeometryScripting/Source/GeometryScriptingCore/Private/MeshAssetFunctions.cpp b/Engine/Plugins/Runtime/GeometryScripting/Source/GeometryScriptingCore/Private/MeshAssetFunctions.cpp
--- a/Engine/Plugins/Runtime/GeometryScripting/Source/GeometryScriptingCore/Private/MeshAssetFunctions.cpp
+++ b/Engine/Plugins/Runtime/GeometryScripting/Source/GeometryScriptingCore/Private/MeshAssetFunctions.cpp
@@ -917,7 +917,7 @@ UDynamicMesh* UGeometryScriptLibrary_StaticMeshFunctions::CopyMeshToSkeletalMesh
 	// to root, since we can't verify that any current skin binding is valid.
 	TArray<int32> BoneRemapping;
 	const FDynamicMesh3& Mesh{FromDynamicMesh->GetMeshRef()}; 
-	if (Mesh.HasAttributes() && Mesh.Attributes()->HasBones())
+	if (Options.bRemapBoneIndicesToMatchAsset && Mesh.HasAttributes() && Mesh.Attributes()->HasBones())
 	{
 		const FDynamicMeshBoneNameAttribute* SrcBoneNames = Mesh.Attributes()->GetBoneNames();
 		TArray<FName> DstBoneNames = ToSkeletalMeshAsset->GetRefSkeleton().GetRawRefBoneNames();
@@ -948,23 +948,26 @@ UDynamicMesh* UGeometryScriptLibrary_StaticMeshFunctions::CopyMeshToSkeletalMesh
 	{
 		Converter.Convert(&ReadMesh, *MeshDescription, !Options.bEnableRecomputeTangents);
 	});
-	
-	if (!BoneRemapping.IsEmpty())
-	{
-		FSkeletalMeshOperations::RemapBoneIndicesOnSkinWeightAttribute(*MeshDescription, BoneRemapping);
-	}
-	else
+
+	if (Options.bRemapBoneIndicesToMatchAsset)
 	{
-		using namespace UE::AnimationCore;
-		FBoneWeight RootWeight(0, 1.0f);
-		FBoneWeights RootBinding = FBoneWeights::Create({RootWeight});
-		for (const FName AttributeName: MeshAttributes.GetSkinWeightProfileNames())
+		if (!BoneRemapping.IsEmpty())
 		{
-			FSkinWeightsVertexAttributesRef SkinWeights(MeshAttributes.GetVertexSkinWeights(AttributeName));
-
-			for (FVertexID VertexID: MeshDescription->Vertices().GetElementIDs())
+			FSkeletalMeshOperations::RemapBoneIndicesOnSkinWeightAttribute(*MeshDescription, BoneRemapping);
+		}
+		else
+		{
+			using namespace UE::AnimationCore;
+			FBoneWeight RootWeight(0, 1.0f);
+			FBoneWeights RootBinding = FBoneWeights::Create({RootWeight});
+			for (const FName AttributeName: MeshAttributes.GetSkinWeightProfileNames())
 			{
-				SkinWeights.Set(VertexID, RootBinding);
+				FSkinWeightsVertexAttributesRef SkinWeights(MeshAttributes.GetVertexSkinWeights(AttributeName));
+
+				for (FVertexID VertexID: MeshDescription->Vertices().GetElementIDs())
+				{
+					SkinWeights.Set(VertexID, RootBinding);
+				}
 			}
 		}
 	}
diff --git a/Engine/Plugins/Runtime/GeometryScripting/Source/GeometryScriptingCore/Public/GeometryScript/MeshAssetFunctions.h b/Engine/Plugins/Runtime/GeometryScripting/Source/GeometryScriptingCore/Public/GeometryScript/MeshAssetFunctions.h
--- a/Engine/Plugins/Runtime/GeometryScripting/Source/GeometryScriptingCore/Public/GeometryScript/MeshAssetFunctions.h
+++ b/Engine/Plugins/Runtime/GeometryScripting/Source/GeometryScriptingCore/Public/GeometryScript/MeshAssetFunctions.h
@@ -72,6 +72,12 @@ public:
 	UPROPERTY(EditAnywhere, BlueprintReadWrite, Category = Options)
 	bool bEnableRemoveDegenerates = false;
 
+	/** Remap the bone indices to match the asset. This requires the source mesh to have bone information present. If no bone information is present
+	 *  then all bone weights are mapped to the root.
+	 */
+	UPROPERTY(EditAnywhere, BlueprintReadWrite, Category = Options)
+	bool bRemapBoneIndicesToMatchAsset = false;
+
 	// Whether to use the build scale on the target asset. If enabled, the inverse scale will be applied when saving to the asset, and the BuildScale will be preserved. Otherwise, BuildScale will be set to 1.0 on the asset BuildSettings.
 	UPROPERTY(EditAnywhere, BlueprintReadWrite, Category = Options)
 	bool bUseBuildScale = true;
