add optional num parameter to the geometry script list wrapper type 'reset' method, so we can more efficiently reset when we know the desired list size

diff --git a/Engine/Plugins/Runtime/GeometryScripting/Source/GeometryScriptingCore/Private/GeometryScriptTypes.cpp b/Engine/Plugins/Runtime/GeometryScripting/Source/GeometryScriptingCore/Private/GeometryScriptTypes.cpp
--- a/Engine/Plugins/Runtime/GeometryScripting/Source/GeometryScriptingCore/Private/GeometryScriptTypes.cpp
+++ b/Engine/Plugins/Runtime/GeometryScripting/Source/GeometryScriptingCore/Private/GeometryScriptTypes.cpp
@@ -73,11 +73,11 @@ void FGeometryScriptSphereCovering::Reset()
 }
 
 
-void FGeometryScriptGeneralPolygonList::Reset()
+void FGeometryScriptGeneralPolygonList::Reset(int32 Num)
 {
 	if (!Polygons.IsValid())
 	{
 		Polygons = MakeShared<TArray<UE::Geometry::FGeneralPolygon2d>>();
 	}
-	Polygons->Reset();
+	Polygons->Reset(Num);
 }
\ No newline at end of file
diff --git a/Engine/Plugins/Runtime/GeometryScripting/Source/GeometryScriptingCore/Public/GeometryScript/GeometryScriptTypes.h b/Engine/Plugins/Runtime/GeometryScripting/Source/GeometryScriptingCore/Public/GeometryScript/GeometryScriptTypes.h
--- a/Engine/Plugins/Runtime/GeometryScripting/Source/GeometryScriptingCore/Public/GeometryScript/GeometryScriptTypes.h
+++ b/Engine/Plugins/Runtime/GeometryScripting/Source/GeometryScriptingCore/Public/GeometryScript/GeometryScriptTypes.h
@@ -320,13 +320,13 @@ struct GEOMETRYSCRIPTINGCORE_API FGeometryScriptIndexList
 public:
 	TSharedPtr<TArray<int>> List;
 
-	void Reset(EGeometryScriptIndexType TargetIndexType)
+	void Reset(EGeometryScriptIndexType TargetIndexType, int32 Num = 0)
 	{
 		if (List.IsValid() == false)
 		{
 			List = MakeShared<TArray<int>>();
 		}
-		List->Reset();
+		List->Reset(Num);
 		IndexType = TargetIndexType;
 	}
 
@@ -363,13 +363,13 @@ struct GEOMETRYSCRIPTINGCORE_API FGeometryScriptTriangleList
 public:
 	TSharedPtr<TArray<FIntVector>> List;
 
-	void Reset()
+	void Reset(int32 Num = 0)
 	{
 		if (List.IsValid() == false)
 		{
 			List = MakeShared<TArray<FIntVector>>();
 		}
-		List->Reset();
+		List->Reset(Num);
 	}
 
 	// Required by TStructOpsTypeTraits interface
@@ -400,13 +400,13 @@ struct GEOMETRYSCRIPTINGCORE_API FGeometryScriptScalarList
 public:
 	TSharedPtr<TArray<double>> List;
 
-	void Reset()
+	void Reset(int32 Num = 0)
 	{
 		if (List.IsValid() == false)
 		{
 			List = MakeShared<TArray<double>>();
 		}
-		List->Reset();
+		List->Reset(Num);
 	}
 
 	// Required by TStructOpsTypeTraits interface
@@ -438,13 +438,13 @@ struct GEOMETRYSCRIPTINGCORE_API FGeometryScriptVectorList
 public:
 	TSharedPtr<TArray<FVector>> List;
 
-	void Reset()
+	void Reset(int32 Num = 0)
 	{
 		if (List.IsValid() == false)
 		{
 			List = MakeShared<TArray<FVector>>();
 		}
-		List->Reset();
+		List->Reset(Num);
 	}
 
 	// Required by TStructOpsTypeTraits interface
@@ -475,13 +475,13 @@ struct GEOMETRYSCRIPTINGCORE_API FGeometryScriptUVList
 public:
 	TSharedPtr<TArray<FVector2D>> List;
 
-	void Reset()
+	void Reset(int32 Num = 0)
 	{
 		if (List.IsValid() == false)
 		{
 			List = MakeShared<TArray<FVector2D>>();
 		}
-		List->Reset();
+		List->Reset(Num);
 	}
 
 	// Required by TStructOpsTypeTraits interface
@@ -513,13 +513,13 @@ struct GEOMETRYSCRIPTINGCORE_API FGeometryScriptColorList
 public:
 	TSharedPtr<TArray<FLinearColor>> List;
 
-	void Reset()
+	void Reset(int32 Num = 0)
 	{
 		if (List.IsValid() == false)
 		{
 			List = MakeShared<TArray<FLinearColor>>();
 		}
-		List->Reset();
+		List->Reset(Num);
 	}
 
 	// Required by TStructOpsTypeTraits interface
@@ -553,13 +553,13 @@ public:
 	UPROPERTY(EditAnywhere, Category = "Options")
 	bool bClosedLoop = false;
 
-	void Reset()
+	void Reset(int32 Num = 0)
 	{
 		if (!Path.IsValid())
 		{
 			Path = MakeShared<TArray<FVector>>();
 		}
-		Path->Reset();
+		Path->Reset(Num);
 	}
 
 	// Required by TStructOpsTypeTraits interface
@@ -590,13 +590,13 @@ struct GEOMETRYSCRIPTINGCORE_API FGeometryScriptSimplePolygon
 public:
 	TSharedPtr<TArray<FVector2D>> Vertices;
 
-	void Reset()
+	void Reset(int32 Num = 0)
 	{
 		if (Vertices.IsValid() == false)
 		{
 			Vertices = MakeShared<TArray<FVector2D>>();
 		}
-		Vertices->Reset();
+		Vertices->Reset(Num);
 	}
 
 	// Required by TStructOpsTypeTraits interface
@@ -629,7 +629,7 @@ struct GEOMETRYSCRIPTINGCORE_API FGeometryScriptGeneralPolygonList
 public:
 	TSharedPtr<TArray<UE::Geometry::FGeneralPolygon2d>> Polygons;
 
-	void Reset();
+	void Reset(int32 Num = 0);
 
 	// Required by TStructOpsTypeTraits interface
 	bool operator==(const FGeometryScriptGeneralPolygonList& Other) const
