add geometry script vector math functions to transform, inverse transform, and plane project

diff --git a/Engine/Plugins/Runtime/GeometryScripting/Source/GeometryScriptingCore/Private/VectorMathFunctions.cpp b/Engine/Plugins/Runtime/GeometryScripting/Source/GeometryScriptingCore/Private/VectorMathFunctions.cpp
--- a/Engine/Plugins/Runtime/GeometryScripting/Source/GeometryScriptingCore/Private/VectorMathFunctions.cpp
+++ b/Engine/Plugins/Runtime/GeometryScripting/Source/GeometryScriptingCore/Private/VectorMathFunctions.cpp
@@ -145,6 +145,66 @@ void UGeometryScriptLibrary_VectorMathFunctions::VectorNormalizeInPlace(FGeometr
 	}
 }
 
+void UGeometryScriptLibrary_VectorMathFunctions::VectorTransformInPlace(UPARAM(ref) FGeometryScriptVectorList& VectorList, FTransform Transform, bool bAsPosition)
+{
+	if (!VectorList.List.IsValid())
+	{
+		UE_LOG(LogGeometry, Warning, TEXT("VectorTransformInPlace: VectorList is empty"));
+		return;
+	}
+	if (bAsPosition)
+	{
+		for (FVector& V : *VectorList.List)
+		{
+			V = Transform.TransformPosition(V);
+		}
+	}
+	else
+	{
+		for (FVector& V : *VectorList.List)
+		{
+			V = Transform.TransformVector(V);
+		}
+	}
+
+}
+
+void UGeometryScriptLibrary_VectorMathFunctions::VectorInverseTransformInPlace(UPARAM(ref) FGeometryScriptVectorList& VectorList, FTransform Transform, bool bAsPosition)
+{
+	if (!VectorList.List.IsValid())
+	{
+		UE_LOG(LogGeometry, Warning, TEXT("VectorInverseTransformInPlace: VectorList is empty"));
+		return;
+	}
+	if (bAsPosition)
+	{
+		for (FVector& V : *VectorList.List)
+		{
+			V = Transform.InverseTransformPosition(V);
+		}
+	}
+	else
+	{
+		for (FVector& V : *VectorList.List)
+		{
+			V = Transform.InverseTransformVector(V);
+		}
+	}
+}
+
+void UGeometryScriptLibrary_VectorMathFunctions::VectorPlaneProjectInPlace(UPARAM(ref) FGeometryScriptVectorList& VectorList, FPlane Plane)
+{
+	if (!VectorList.List.IsValid())
+	{
+		UE_LOG(LogGeometry, Warning, TEXT("VectorPlaneProjectInPlace: VectorList is empty"));
+		return;
+	}
+	for (FVector& V : *VectorList.List)
+	{
+		V = FVector::PointPlaneProject(V, Plane);
+	}
+}
+
 
 
 FGeometryScriptVectorList UGeometryScriptLibrary_VectorMathFunctions::VectorBlend(FGeometryScriptVectorList VectorListA, FGeometryScriptVectorList VectorListB, double ConstantA, double ConstantB)
diff --git a/Engine/Plugins/Runtime/GeometryScripting/Source/GeometryScriptingCore/Public/GeometryScript/VectorMathFunctions.h b/Engine/Plugins/Runtime/GeometryScripting/Source/GeometryScriptingCore/Public/GeometryScript/VectorMathFunctions.h
--- a/Engine/Plugins/Runtime/GeometryScripting/Source/GeometryScriptingCore/Public/GeometryScript/VectorMathFunctions.h
+++ b/Engine/Plugins/Runtime/GeometryScripting/Source/GeometryScriptingCore/Public/GeometryScript/VectorMathFunctions.h
@@ -46,6 +46,29 @@ public:
 	static void 
 	VectorNormalizeInPlace(UPARAM(ref) FGeometryScriptVectorList& VectorList, FVector SetOnFailure = FVector::ZeroVector);
 
+	/**
+	 * Transform each vector in VectorList, and store in VectorList.
+	 * @param bAsPosition Whether to treat input as positions or vectors (if vectors, will ignore the Transform's translation part)
+	 */
+	UFUNCTION(BlueprintCallable, Category = "GeometryScript|Math|VectorMath", meta=(ScriptMethod))
+	static void 
+	VectorTransformInPlace(UPARAM(ref) FGeometryScriptVectorList& VectorList, FTransform Transform, bool bAsPosition = true);
+
+	/**
+	 * Inverse transform each vector in VectorList, and store in VectorList.
+	 * @param bAsPosition Whether to treat input as positions or vectors (if vectors, will ignore the Transform's translation part)
+	 */
+	UFUNCTION(BlueprintCallable, Category = "GeometryScript|Math|VectorMath", meta=(ScriptMethod))
+	static void 
+	VectorInverseTransformInPlace(UPARAM(ref) FGeometryScriptVectorList& VectorList, FTransform Transform, bool bAsPosition = true);
+
+	/**
+	 * Project each vector in VectorList to the given Plane, and store in VectorList.
+	 */
+	UFUNCTION(BlueprintCallable, Category = "GeometryScript|Math|VectorMath", meta=(ScriptMethod))
+	static void 
+	VectorPlaneProjectInPlace(UPARAM(ref) FGeometryScriptVectorList& VectorList, FPlane Plane);
+
 	/**
 	 * Compute (ConstantA * A) + (ConstantB * B) for each pair of vectors in VectorListA and VectorListB and return in new VectorList.
 	 * By default (constants = 1) this just adds the two vectors. Set ConstantB = -1 to subtract B from A. 
