Add option to set material ID on appended mesh primitives in geometry script

diff --git a/Engine/Plugins/Runtime/GeometryScripting/Source/GeometryScriptingCore/Private/MeshPrimitiveFunctions.cpp b/Engine/Plugins/Runtime/GeometryScripting/Source/GeometryScriptingCore/Private/MeshPrimitiveFunctions.cpp
--- a/Engine/Plugins/Runtime/GeometryScripting/Source/GeometryScriptingCore/Private/MeshPrimitiveFunctions.cpp
+++ b/Engine/Plugins/Runtime/GeometryScripting/Source/GeometryScriptingCore/Private/MeshPrimitiveFunctions.cpp
@@ -69,6 +69,18 @@ static void ApplyPrimitiveOptionsToMesh(
 			}
 		}
 	}
+	// if we have a valid non-zero MaterialID, set it on all triangles (0 is default)
+	if (PrimitiveOptions.MaterialID > 0)
+	{
+		// Make sure material IDs are enabled
+		Mesh.EnableAttributes();
+		Mesh.Attributes()->EnableMaterialID();
+		FDynamicMeshMaterialAttribute* MaterialIDs = Mesh.Attributes()->GetMaterialID();
+		for (int32 TID : Mesh.TriangleIndicesItr())
+		{
+			MaterialIDs->SetValue(TID, PrimitiveOptions.MaterialID);
+		}
+	}
 }
 
 static void AppendPrimitive(
diff --git a/Engine/Plugins/Runtime/GeometryScripting/Source/GeometryScriptingCore/Public/GeometryScript/MeshPrimitiveFunctions.h b/Engine/Plugins/Runtime/GeometryScripting/Source/GeometryScriptingCore/Public/GeometryScript/MeshPrimitiveFunctions.h
--- a/Engine/Plugins/Runtime/GeometryScripting/Source/GeometryScriptingCore/Public/GeometryScript/MeshPrimitiveFunctions.h
+++ b/Engine/Plugins/Runtime/GeometryScripting/Source/GeometryScriptingCore/Public/GeometryScript/MeshPrimitiveFunctions.h
@@ -61,6 +61,10 @@ public:
 
 	UPROPERTY(EditAnywhere, BlueprintReadWrite, Category = Options)
 	EGeometryScriptPrimitiveUVMode UVMode = EGeometryScriptPrimitiveUVMode::Uniform;
+
+	// Material ID to set on primitive mesh triangles
+	UPROPERTY(EditAnywhere, BlueprintReadWrite, Category = Options)
+	int32 MaterialID = 0;
 };
 
 
