add geometry script support for selecting mesh boundary edges, selection boundary edges, and sharp edges (edges w/ dihedral angle > threshold)

diff --git a/Engine/Plugins/Runtime/GeometryScripting/Source/GeometryScriptingCore/Private/MeshSelectionFunctions.cpp b/Engine/Plugins/Runtime/GeometryScripting/Source/GeometryScriptingCore/Private/MeshSelectionFunctions.cpp
--- a/Engine/Plugins/Runtime/GeometryScripting/Source/GeometryScriptingCore/Private/MeshSelectionFunctions.cpp
+++ b/Engine/Plugins/Runtime/GeometryScripting/Source/GeometryScriptingCore/Private/MeshSelectionFunctions.cpp
@@ -1003,6 +1003,116 @@ UDynamicMesh* UGeometryScriptLibrary_MeshSelectionFunctions::SelectMeshElementsB
 	return TargetMesh;
 }
 
+UDynamicMesh* UGeometryScriptLibrary_MeshSelectionFunctions::SelectMeshSharpEdges(
+	UDynamicMesh* TargetMesh,
+	FGeometryScriptMeshSelection& SelectionOut,
+	double MinAngleDeg)
+{
+	if (TargetMesh == nullptr)
+	{
+		UE_LOG(LogGeometry, Warning, TEXT("SelectMeshEdgesByEdgeAngle: TargetMesh is Null"));
+		return TargetMesh;
+	}
+
+	FGeometrySelection GeoSelection;
+	GeoSelection.InitializeTypes(EGeometryElementType::Edge, EGeometryTopologyType::Triangle);
+	double CosThresh = FMath::Cos(FMathd::DegToRad * MinAngleDeg);
+	TargetMesh->ProcessMesh([&GeoSelection, &CosThresh](const FDynamicMesh3& Mesh)
+	{
+		for (int32 EID : Mesh.EdgeIndicesItr())
+		{
+			FIndex2i EdgeT = Mesh.GetEdgeT(EID);
+			if (EdgeT.B == INDEX_NONE)
+			{
+				continue;
+			}
+			FVector3d NormalA = Mesh.GetTriNormal(EdgeT.A);
+			FVector3d NormalB = Mesh.GetTriNormal(EdgeT.B);
+			if (NormalA.Dot(NormalB) <= CosThresh)
+			{
+				Mesh.EnumerateTriEdgeIDsFromEdgeID(EID, [&GeoSelection](FMeshTriEdgeID TriEdgeID)
+				{
+					GeoSelection.Selection.Add(TriEdgeID.Encoded());
+				});
+			}
+		}
+	});
+	SelectionOut.SetSelection(GeoSelection);
+	return TargetMesh;
+}
+
+UDynamicMesh* UGeometryScriptLibrary_MeshSelectionFunctions::SelectMeshBoundaryEdges(
+	UDynamicMesh* TargetMesh,
+	FGeometryScriptMeshSelection& SelectionOut)
+{
+	if (TargetMesh == nullptr)
+	{
+		UE_LOG(LogGeometry, Warning, TEXT("SelectMeshEdgesByEdgeAngle: TargetMesh is Null"));
+		return TargetMesh;
+	}
+
+	FGeometrySelection GeoSelection;
+	GeoSelection.InitializeTypes(EGeometryElementType::Edge, EGeometryTopologyType::Triangle);
+	TargetMesh->ProcessMesh([&GeoSelection](const FDynamicMesh3& Mesh)
+	{
+		for (int32 EID : Mesh.EdgeIndicesItr())
+		{
+			FIndex2i EdgeT = Mesh.GetEdgeT(EID);
+			if (EdgeT.B == INDEX_NONE)
+			{
+				GeoSelection.Selection.Add(Mesh.GetTriEdgeIDFromEdgeID(EID).Encoded());
+			}
+		}
+	});
+	SelectionOut.SetSelection(GeoSelection);
+	return TargetMesh;
+}
+
+UDynamicMesh* UGeometryScriptLibrary_MeshSelectionFunctions::SelectSelectionBoundaryEdges(
+	UDynamicMesh* TargetMesh,
+	const FGeometryScriptMeshSelection& RegionSelection,
+	FGeometryScriptMeshSelection& SelectionOut,
+	bool bExcludeMeshBoundaryEdges
+)
+{
+	if (TargetMesh == nullptr)
+	{
+		UE_LOG(LogGeometry, Warning, TEXT("SelectMeshEdgesByEdgeAngle: TargetMesh is Null"));
+		return TargetMesh;
+	}
+
+	FGeometrySelection GeoSelection;
+	GeoSelection.InitializeTypes(EGeometryElementType::Edge, EGeometryTopologyType::Triangle);
+	TargetMesh->ProcessMesh([&GeoSelection, &RegionSelection, bExcludeMeshBoundaryEdges](const FDynamicMesh3& Mesh)
+	{
+		TSet<int32> TriSel;
+		RegionSelection.ProcessByTriangleID(Mesh, [&TriSel](int32 TID) { TriSel.Add(TID); });
+		for (int32 TID : TriSel)
+		{
+			FIndex3i TriEdges = Mesh.GetTriEdges(TID);
+			FIndex3i NbrTris = Mesh.GetTriNeighbourTris(TID);
+			for (int32 SubIdx = 0; SubIdx < 3; ++SubIdx)
+			{
+				if (NbrTris[SubIdx] == INDEX_NONE)
+				{
+					if (!bExcludeMeshBoundaryEdges)
+					{
+						GeoSelection.Selection.Add(Mesh.GetTriEdgeIDFromEdgeID(TriEdges[SubIdx]).Encoded());
+					}
+				}
+				else if (!TriSel.Contains(NbrTris[SubIdx]))
+				{
+					Mesh.EnumerateTriEdgeIDsFromEdgeID(TriEdges[SubIdx], [&GeoSelection](FMeshTriEdgeID TriEdgeID)
+					{
+						GeoSelection.Selection.Add(TriEdgeID.Encoded());
+					});
+				}
+			}
+		}
+	});
+	SelectionOut.SetSelection(GeoSelection);
+	return TargetMesh;
+}
 
 UDynamicMesh* UGeometryScriptLibrary_MeshSelectionFunctions::SelectMeshElementsInsideMesh(
 	UDynamicMesh* TargetMesh,
diff --git a/Engine/Plugins/Runtime/GeometryScripting/Source/GeometryScriptingCore/Public/GeometryScript/MeshSelectionFunctions.h b/Engine/Plugins/Runtime/GeometryScripting/Source/GeometryScriptingCore/Public/GeometryScript/MeshSelectionFunctions.h
--- a/Engine/Plugins/Runtime/GeometryScripting/Source/GeometryScriptingCore/Public/GeometryScript/MeshSelectionFunctions.h
+++ b/Engine/Plugins/Runtime/GeometryScripting/Source/GeometryScriptingCore/Public/GeometryScript/MeshSelectionFunctions.h
@@ -263,6 +263,40 @@ public:
 		bool bInvert = false,
 		UPARAM(DisplayName = "Min Element Vertices") int MinNumTrianglePoints = 3);
 
+	/**
+	 * Create a new Selection, for the TargetMesh, of all 'sharp' edges where the edge's adjacent triangle normals differ by at least MinAngleDeg
+	 */
+	UFUNCTION(BlueprintCallable, Category = "GeometryScript|MeshSelection", meta = (ScriptMethod))
+	static UPARAM(DisplayName = "Target Mesh") UDynamicMesh*
+	SelectMeshSharpEdges(
+		UDynamicMesh* TargetMesh,
+		FGeometryScriptMeshSelection& Selection,
+		double MinAngleDeg = 20.0
+	);
+
+	/**
+	 * Create a new Selection, for the TargetMesh, of all mesh boundary edges
+	 */
+	UFUNCTION(BlueprintCallable, Category = "GeometryScript|MeshSelection", meta = (ScriptMethod))
+	static UPARAM(DisplayName = "Target Mesh") UDynamicMesh*
+	SelectMeshBoundaryEdges(
+		UDynamicMesh* TargetMesh,
+		FGeometryScriptMeshSelection& Selection
+	);
+
+	/**
+	 * Create a new BoundarySelection, for the TargetMesh, of the edges on the boundary of another Selection
+	 * @param bExcludeMeshBoundaryEdges If true, do not include Selection boundary edges if they are also mesh boundary edges
+	 */
+	UFUNCTION(BlueprintCallable, Category = "GeometryScript|MeshSelection", meta = (ScriptMethod))
+	static UPARAM(DisplayName = "Target Mesh") UDynamicMesh*
+	SelectSelectionBoundaryEdges(
+		UDynamicMesh* TargetMesh,
+		const FGeometryScriptMeshSelection& Selection,
+		FGeometryScriptMeshSelection& BoundarySelection,
+		bool bExcludeMeshBoundaryEdges = false
+	);
+
 
 	/**
 	* Create a new Mesh Selection of the SelectionType for the TargetMesh by finding all elements inside a second SelectionMesh
