add uv seam insertion/removal to geometry script

diff --git a/Engine/Plugins/Runtime/GeometryScripting/Source/GeometryScriptingCore/Private/MeshUVFunctions.cpp b/Engine/Plugins/Runtime/GeometryScripting/Source/GeometryScriptingCore/Private/MeshUVFunctions.cpp
--- a/Engine/Plugins/Runtime/GeometryScripting/Source/GeometryScriptingCore/Private/MeshUVFunctions.cpp
+++ b/Engine/Plugins/Runtime/GeometryScripting/Source/GeometryScriptingCore/Private/MeshUVFunctions.cpp
@@ -288,7 +288,7 @@ UDynamicMesh* UGeometryScriptLibrary_MeshUVFunctions::SetMeshUVElementPosition(
 
 
 void ApplyMeshUVEditorOperation(UDynamicMesh* TargetMesh, int32 UVSetIndex, bool& bHasUVSet, UGeometryScriptDebug* Debug,
-	TFunctionRef<void(FDynamicMesh3& Mesh, FDynamicMeshUVOverlay* UVOverlay, FDynamicMeshUVEditor& UVEditor)> EditFunc)
+	TFunctionRef<void(FDynamicMesh3& Mesh, FDynamicMeshUVOverlay* UVOverlay, FDynamicMeshUVEditor& UVEditor)> EditFunc, bool bDeferChangeNotifications = false)
 {
 	bHasUVSet = false;
 	TargetMesh->EditMesh([&](FDynamicMesh3& EditMesh)
@@ -305,7 +305,7 @@ void ApplyMeshUVEditorOperation(UDynamicMesh* TargetMesh, int32 UVSetIndex, bool
 		FDynamicMeshUVEditor Editor(&EditMesh, UVOverlay);
 		EditFunc(EditMesh, UVOverlay, Editor);
 
-	}, EDynamicMeshChangeType::GeneralEdit, EDynamicMeshAttributeChangeFlags::Unknown, false);
+	}, EDynamicMeshChangeType::GeneralEdit, EDynamicMeshAttributeChangeFlags::Unknown, bDeferChangeNotifications);
 }
 
 
@@ -347,6 +347,45 @@ void ApplyUVTransform(
 }
 
 
+UDynamicMesh* UGeometryScriptLibrary_MeshUVFunctions::SetUVSeamsAlongSelectedEdges(
+	UDynamicMesh* TargetMesh,
+	int UVSetIndex,
+	FGeometryScriptMeshSelection Selection,
+	bool bInsertSeams,
+	bool bDeferChangeNotifications,
+	UGeometryScriptDebug* Debug)
+{
+	if (TargetMesh == nullptr)
+	{
+		UE::Geometry::AppendError(Debug, EGeometryScriptErrorType::InvalidInputs, LOCTEXT("SetUVSeamsAlongSelectedEdges_InvalidInput", "SetUVSeamsAlongSelectedEdges: TargetMesh is Null"));
+		return TargetMesh;
+	}
+
+	bool bHasUVSet = false;
+	ApplyMeshUVEditorOperation(TargetMesh, UVSetIndex, bHasUVSet, Debug,
+		[&Selection, bInsertSeams](FDynamicMesh3& EditMesh, FDynamicMeshUVOverlay* UVOverlay, FDynamicMeshUVEditor& UVEditor)
+	{
+		TArray<int32> EdgeArr;
+		Selection.ConvertToMeshIndexArray(EditMesh, EdgeArr, EGeometryScriptIndexType::Edge);
+		TSet<int32> EdgeSet(EdgeArr);
+		if (bInsertSeams)
+		{
+			UVEditor.CreateSeamsAtEdges(EdgeSet);
+		}
+		else
+		{
+			UVEditor.RemoveSeamsAtEdges(EdgeSet);
+		}
+	}, bDeferChangeNotifications);
+	if (bHasUVSet == false)
+	{
+		UE::Geometry::AppendError(Debug, EGeometryScriptErrorType::InvalidInputs, LOCTEXT("SetUVSeamsAlongSelectedEdges_InvalidUVSet", "SetUVSeamsAlongSelectedEdges: UV Channel does not exist on TargetMesh"));
+	}
+
+	return TargetMesh;
+}
+
+
 UDynamicMesh* UGeometryScriptLibrary_MeshUVFunctions::TranslateMeshUVs(
 	UDynamicMesh* TargetMesh,
 	int UVSetIndex,
diff --git a/Engine/Plugins/Runtime/GeometryScripting/Source/GeometryScriptingCore/Public/GeometryScript/MeshUVFunctions.h b/Engine/Plugins/Runtime/GeometryScripting/Source/GeometryScriptingCore/Public/GeometryScript/MeshUVFunctions.h
--- a/Engine/Plugins/Runtime/GeometryScripting/Source/GeometryScriptingCore/Public/GeometryScript/MeshUVFunctions.h
+++ b/Engine/Plugins/Runtime/GeometryScripting/Source/GeometryScriptingCore/Public/GeometryScript/MeshUVFunctions.h
@@ -228,6 +228,25 @@ public:
 		FIntVector& TriangleUVElements,
 		bool& bHaveValidUVs);
 
+
+	/**
+	 * Convert Selection to an Edge selection, and set or remove UV seams along all of the selected edges
+	 * @param TargetMesh The mesh to update
+	 * @param UVChannel The UV Channel to update
+	 * @param Selection Which edges to operate on
+	 * @param bInsertSeams Whether to insert new seams. If false, removes existing seams instead.
+	 * @param bDeferChangeNotifications If true, no mesh change notification will be sent. Set to true if performing many changes in a loop.
+	 */
+	UFUNCTION(BlueprintCallable, Category = "GeometryScript|UVs", meta = (ScriptMethod, DisplayName = "Set UV Seams Along Selected Edges"))
+	static UPARAM(DisplayName = "Target Mesh") UDynamicMesh*
+	SetUVSeamsAlongSelectedEdges(
+		UDynamicMesh* TargetMesh,
+		UPARAM(DisplayName = "UV Channel") int UVChannel,
+		FGeometryScriptMeshSelection Selection,
+		bool bInsertSeams = true,
+		bool bDeferChangeNotifications = false,
+		UGeometryScriptDebug* Debug = nullptr);
+
 	/**
 	 * Returns the UV Position for a given UV Element ID in the specified UV Channel.
 	 * If the UV Set or Element ID does not exist, bIsValidElementID will be returned as false.
