GeometryScripting: Expose the UV Layout operator with all options, including multiple layout types and repack restrictions on scale and rotation, to geometry script as a new node.

diff --git a/Engine/Plugins/Runtime/GeometryScripting/Source/GeometryScriptingCore/Private/MeshUVFunctions.cpp b/Engine/Plugins/Runtime/GeometryScripting/Source/GeometryScriptingCore/Private/MeshUVFunctions.cpp
--- a/Engine/Plugins/Runtime/GeometryScripting/Source/GeometryScriptingCore/Private/MeshUVFunctions.cpp
+++ b/Engine/Plugins/Runtime/GeometryScripting/Source/GeometryScriptingCore/Private/MeshUVFunctions.cpp
@@ -10,6 +10,7 @@
 #include "Parameterization/PatchBasedMeshUVGenerator.h"
 #include "DynamicMesh/MeshNormals.h"
 #include "Parameterization/MeshLocalParam.h"
+#include "ParameterizationOps/UVLayoutOp.h"
 #include "TriangleTypes.h"
 #include "XAtlasWrapper.h"
 #include "ParameterizationOps/TexelDensityOp.h"
@@ -840,7 +841,88 @@ UDynamicMesh* UGeometryScriptLibrary_MeshUVFunctions::RepackMeshUVs(
 	return TargetMesh;
 }
 
+UDynamicMesh* UGeometryScriptLibrary_MeshUVFunctions::LayoutMeshUVs(
+	UDynamicMesh* TargetMesh,
+	int UVSetIndex,
+	FGeometryScriptLayoutUVsOptions LayoutOptions,
+	FGeometryScriptMeshSelection Selection,
+	UGeometryScriptDebug* Debug)
+{
+	if (TargetMesh == nullptr)
+	{
+		UE::Geometry::AppendError(Debug, EGeometryScriptErrorType::InvalidInputs, LOCTEXT("LayoutMeshUVs_InvalidInput", "LayoutMeshUVs: TargetMesh is Null"));
+		return TargetMesh;
+	}
+
+
+	TargetMesh->EditMesh([&](FDynamicMesh3& EditMesh)
+		{
+
+			TSharedPtr<FDynamicMesh3> SourceMesh = MakeShared<FDynamicMesh3>(MoveTemp(EditMesh));
+
+			FUVLayoutOp Op;
+
+			Op.OriginalMesh = SourceMesh;
+
+			switch (LayoutOptions.LayoutType)
+			{
+			case EGeometryScriptUVLayoutType::Normalize:
+				Op.UVLayoutMode = EUVLayoutOpLayoutModes::Normalize;
+				break;
+			case EGeometryScriptUVLayoutType::Repack:
+				Op.UVLayoutMode = EUVLayoutOpLayoutModes::RepackToUnitRect;
+				break;
+			case EGeometryScriptUVLayoutType::Stack:
+				Op.UVLayoutMode = EUVLayoutOpLayoutModes::StackInUnitRect;
+				break;
+			case EGeometryScriptUVLayoutType::Transform:
+				Op.UVLayoutMode = EUVLayoutOpLayoutModes::TransformOnly;
+				break;
+			default:
+				ensure(false);
+				break;
+			}
+
 
+			TArray<int32> TriangleSelection;
+			Selection.ConvertToMeshIndexArray(*SourceMesh, TriangleSelection, EGeometryScriptIndexType::Triangle);
+
+			Op.TextureResolution = LayoutOptions.TextureResolution;
+			Op.bPreserveScale = LayoutOptions.bPreserveScale;
+			Op.bPreserveRotation = LayoutOptions.bPreserveRotation;
+			Op.bAllowFlips = LayoutOptions.bAllowFlips;
+			Op.UVScaleFactor = LayoutOptions.Scale;
+			Op.UVTranslation = FVector2f(LayoutOptions.Translation);
+			Op.bMaintainOriginatingUDIM = LayoutOptions.bEnableUDIMLayout;					
+			Op.UVLayerIndex = UVSetIndex;
+			Op.TextureResolution = LayoutOptions.TextureResolution;
+			Op.SetTransform(FTransformSRT3d::Identity());
+			Op.bMaintainOriginatingUDIM = LayoutOptions.bEnableUDIMLayout;
+			if (TriangleSelection.Num() > 0)
+			{
+				Op.Selection = TSet<int32>(TriangleSelection);
+			}
+			if (LayoutOptions.UDIMResolutions.Num() > 0)
+			{
+				Op.TextureResolutionPerUDIM = LayoutOptions.UDIMResolutions;
+			}
+
+			Op.CalculateResult(nullptr);
+			if (Op.GetResultInfo().Result == EGeometryResultType::Success)
+			{
+				TUniquePtr<FDynamicMesh3> ResultMesh = Op.ExtractResult();
+				EditMesh = MoveTemp(*ResultMesh);
+			}
+			else
+			{
+				EditMesh = MoveTemp(*SourceMesh);
+				UE::Geometry::AppendError(Debug, EGeometryScriptErrorType::InvalidInputs, LOCTEXT("LayoutMeshUVs_ComputeError", "LayoutMeshUVs: Error computing result, returning input mesh"));
+			}
+
+		}, EDynamicMeshChangeType::AttributeEdit, EDynamicMeshAttributeChangeFlags::UVs, false);
+
+	return TargetMesh;
+}
 
 
 UDynamicMesh* UGeometryScriptLibrary_MeshUVFunctions::AutoGeneratePatchBuilderMeshUVs( 
diff --git a/Engine/Plugins/Runtime/GeometryScripting/Source/GeometryScriptingCore/Public/GeometryScript/MeshUVFunctions.h b/Engine/Plugins/Runtime/GeometryScripting/Source/GeometryScriptingCore/Public/GeometryScript/MeshUVFunctions.h
--- a/Engine/Plugins/Runtime/GeometryScripting/Source/GeometryScriptingCore/Public/GeometryScript/MeshUVFunctions.h
+++ b/Engine/Plugins/Runtime/GeometryScripting/Source/GeometryScriptingCore/Public/GeometryScript/MeshUVFunctions.h
@@ -23,6 +23,60 @@ public:
 	bool bOptimizeIslandRotation = true;
 };
 
+UENUM(BlueprintType)
+enum class EGeometryScriptUVLayoutType : uint8
+{
+	/** Apply Scale and Translation properties to all UV values */
+	Transform,
+	/** Uniformly scale and translate each UV island individually to pack it into the unit square, i.e. fit between 0 and 1 with overlap */
+	Stack,
+	/** Uniformly scale and translate UV islands collectively to pack them into the unit square, i.e. fit between 0 and 1 with no overlap */
+	Repack,
+	/** Scale and translate UV islands to normalize the UV islands' area to match an average texel density. */
+	Normalize
+};
+
+USTRUCT(BlueprintType)
+struct GEOMETRYSCRIPTINGCORE_API FGeometryScriptLayoutUVsOptions
+{
+	GENERATED_BODY()
+public:
+	/** Type of layout applied to input UVs */
+	UPROPERTY(EditAnywhere, BlueprintReadWrite, Category = Options)
+	EGeometryScriptUVLayoutType LayoutType = EGeometryScriptUVLayoutType::Repack;
+
+	/** Expected resolution of the output textures; this controls spacing left between UV islands to avoid interpolation artifacts */
+	UPROPERTY(EditAnywhere, BlueprintReadWrite, Category = Options)
+	int TextureResolution = 1024;
+
+	/** Uniform scale applied to UVs after packing */
+	UPROPERTY(EditAnywhere, BlueprintReadWrite, Category = Options)
+	float Scale = 1;
+
+	/** Translation applied to UVs after packing, and after scaling */
+	UPROPERTY(EditAnywhere, BlueprintReadWrite, Category = Options)
+	FVector2D Translation = FVector2D(0, 0);
+
+	/** Force the Repack layout type to preserve existing scaling of UV islands. Note, this might lead to the packing not fitting within a unit square, and therefore is disabled by default. */
+	UPROPERTY(EditAnywhere, BlueprintReadWrite, Category = Options)
+	bool bPreserveScale = false;
+
+	/** Force the Repack layout type to preserve existing rotation of UV islands. Note, this might lead to the packing not being as space efficient as possible, and therefore is disabled by default. */
+	UPROPERTY(EditAnywhere, BlueprintReadWrite, Category = Options)
+	bool bPreserveRotation = false;
+
+	/** Allow the Repack layout type to flip the orientation when rotating UV islands to save space. Note that this may cause problems for downstream operations, and therefore is disabled by default. */
+	UPROPERTY(EditAnywhere, BlueprintReadWrite, Category = Options)
+	bool bAllowFlips = false;
+
+	/** Enable UDIM aware layout and keep islands within their originating UDIM tiles when laying out.*/
+	UPROPERTY(EditAnywhere, BlueprintReadWrite, Category = Options)
+	bool bEnableUDIMLayout = false;
+
+	UPROPERTY(EditAnywhere, BlueprintReadWrite, Category = Options, meta = (DisplayName = "Per UDIM Texture Resolution"))
+	TMap<int32, int32> UDIMResolutions;
+};
+
 UENUM(BlueprintType)
 enum class EGeometryScriptUVFlattenMethod : uint8
 {
@@ -427,6 +481,18 @@ public:
 		FGeometryScriptRepackUVsOptions RepackOptions,
 		UGeometryScriptDebug* Debug = nullptr );
 
+	/**
+	* Packs the existing UV islands in the specified UV Channel into standard UV space based on the Repack Options.
+	*/
+	UFUNCTION(BlueprintCallable, Category = "GeometryScript|UVs", meta = (ScriptMethod))
+	static UPARAM(DisplayName = "Target Mesh") UDynamicMesh*
+	LayoutMeshUVs(
+		UDynamicMesh* TargetMesh,
+		UPARAM(DisplayName = "UV Channel") int UVSetIndex,
+		FGeometryScriptLayoutUVsOptions LayoutOptions,
+		FGeometryScriptMeshSelection Selection,
+		UGeometryScriptDebug* Debug = nullptr);
+
 	/**
 	* Computes new UVs for the specified UV Channel using PatchBuilder method in the Options, and optionally packs.
 	*/
