Geoscripting: Node for adding alt skin weights directly to the skeletal mesh asset.

diff --git a/Engine/Plugins/Runtime/GeometryScripting/Source/GeometryScriptingCore/Private/MeshAssetFunctions.cpp b/Engine/Plugins/Runtime/GeometryScripting/Source/GeometryScriptingCore/Private/MeshAssetFunctions.cpp
--- a/Engine/Plugins/Runtime/GeometryScripting/Source/GeometryScriptingCore/Private/MeshAssetFunctions.cpp
+++ b/Engine/Plugins/Runtime/GeometryScripting/Source/GeometryScriptingCore/Private/MeshAssetFunctions.cpp
@@ -27,6 +27,7 @@
 #include "AssetUtils/StaticMeshMaterialUtil.h"
 #include "ConversionUtils/SceneComponentToDynamicMesh.h"
 #include "DynamicMesh/DynamicBoneAttribute.h"
+#include "DynamicMesh/DynamicVertexSkinWeightsAttribute.h"
 #include "DynamicMesh/NonManifoldMappingSupport.h"
 
 #include UE_INLINE_GENERATED_CPP_BY_NAME(MeshAssetFunctions)
@@ -1191,4 +1192,174 @@ UDynamicMesh* UGeometryScriptLibrary_StaticMeshFunctions::CopyMorphTargetToSkele
 	return FromDynamicMesh;
 }
 
+
+UDynamicMesh* UGeometryScriptLibrary_StaticMeshFunctions::CopySkinWeightProfileToSkeletalMesh(
+		UDynamicMesh* FromDynamicMesh, 
+		USkeletalMesh* ToSkeletalMeshAsset,
+		FName TargetProfileName,
+		FName SourceProfileName,
+		FGeometryScriptCopySkinWeightProfileToAssetOptions Options,
+		FGeometryScriptMeshWriteLOD TargetLOD,
+		EGeometryScriptOutcomePins& Outcome,
+		UGeometryScriptDebug* Debug)
+{
+	Outcome = EGeometryScriptOutcomePins::Failure;
+
+	if (ToSkeletalMeshAsset == nullptr)
+	{
+		UE::Geometry::AppendError(Debug, EGeometryScriptErrorType::InvalidInputs, LOCTEXT("CopySkinWeightProfileToSkeletalMesh_InvalidInput1", "CopySkinWeightProfileToSkeletalMesh: ToSkeletalMeshAsset is Null"));
+		return FromDynamicMesh;
+	}
+	if (FromDynamicMesh == nullptr)
+	{
+		UE::Geometry::AppendError(Debug, EGeometryScriptErrorType::InvalidInputs, LOCTEXT("CopySkinWeightProfileToSkeletalMesh_InvalidInput2", "CopySkinWeightProfileToSkeletalMesh: FromDynamicMesh is Null"));
+		return FromDynamicMesh;
+	}
+
+#if WITH_EDITOR
+	if (ToSkeletalMeshAsset->GetPathName().StartsWith(TEXT("/Engine/")))
+	{
+		const FText Error = FText::Format(LOCTEXT("CopySkinWeightProfileToSkeletalMesh_BuiltInAsset", "CopySkinWeightProfileToSkeletalMesh: Cannot modify built-in engine asset: {0}"), FText::FromString(*ToSkeletalMeshAsset->GetPathName()));
+		UE::Geometry::AppendError(Debug, EGeometryScriptErrorType::InvalidInputs, Error);
+		return FromDynamicMesh;
+	}
+
+	TUniquePtr<FScopedTransaction> Transaction;
+	if (Options.bEmitTransaction)
+	{
+		Transaction = MakeUnique<FScopedTransaction>(LOCTEXT("UpdateSkeletalMesh", "Update Skeletal Mesh"));
+	}
+	
+	// flush any pending rendering commands, which might touch a component while we are rebuilding it's mesh
+	FlushRenderingCommands();
+
+	// make sure transactional flag is on for this asset
+	ToSkeletalMeshAsset->SetFlags(RF_Transactional);
+
+	(void)ToSkeletalMeshAsset->Modify();
+
+	// Ensure we have enough LODInfos to cover up to the requested LOD.
+	for (int32 LODIndex = ToSkeletalMeshAsset->GetLODNum(); LODIndex <= TargetLOD.LODIndex; LODIndex++)
+	{
+		FSkeletalMeshLODInfo& LODInfo = ToSkeletalMeshAsset->AddLODInfo();
+		
+		ToSkeletalMeshAsset->GetImportedModel()->LODModels.Add(new FSkeletalMeshLODModel);
+		LODInfo.ReductionSettings.BaseLOD = 0;
+	}
+
+	FMeshDescription* MeshDescription = ToSkeletalMeshAsset->GetMeshDescription(TargetLOD.LODIndex);
+	if (MeshDescription == nullptr)
+	{
+		if (Transaction)
+		{
+			Transaction->Cancel();
+		}
+		UE::Geometry::AppendError(Debug, EGeometryScriptErrorType::InvalidInputs, LOCTEXT("CopySkinWeightProfileToSkeletalMesh_TargetMeshDescription", "CopySkinWeightProfileToSkeletalMesh: Failed to generate the mesh data for the Target LOD Index"));
+		return FromDynamicMesh;
+	}
+
+	if (TargetProfileName.IsNone())
+	{
+		TargetProfileName = FSkeletalMeshAttributes::DefaultSkinWeightProfileName;
+	}
+	if (SourceProfileName.IsNone())
+	{
+		SourceProfileName = FSkeletalMeshAttributes::DefaultSkinWeightProfileName;
+	}
+	
+	// If the dynamic mesh has non-manifold information, use that to figure out what the original vertex count was.
+	// Otherwise, we assume that they have a 1:1 match.
+	const FDynamicMesh3& SourceMesh = FromDynamicMesh->GetMeshRef();
+	if (!SourceMesh.HasAttributes() || !SourceMesh.Attributes()->HasSkinWeightsAttribute(SourceProfileName))
+	{
+		if (Transaction)
+		{
+			Transaction->Cancel();
+		}
+		UE::Geometry::AppendError(Debug, EGeometryScriptErrorType::InvalidInputs, LOCTEXT("CopySkinWeightProfileToSkeletalMesh_InvalidSourceProfile", "CopySkinWeightProfileToSkeletalMesh: The requested skin weight profile does not exist on the source mesh."));
+		return FromDynamicMesh;
+	}
+	
+	const FNonManifoldMappingSupport NonManifoldMappingSupport(SourceMesh);
+	int32 SourceVertexCount;
+
+	if (NonManifoldMappingSupport.IsNonManifoldVertexInSource())
+	{
+		TSet<int32> UniqueVertices;
+		for (int32 SourceVID = 0; SourceVID < SourceMesh.VertexCount(); ++SourceVID)
+		{
+			UniqueVertices.Add(NonManifoldMappingSupport.GetOriginalNonManifoldVertexID(SourceVID));
+		}
+
+		SourceVertexCount = UniqueVertices.Num();
+	}
+	else
+	{
+		SourceVertexCount = SourceMesh.VertexCount();
+	}
+	if (MeshDescription->Vertices().Num() != SourceVertexCount)
+	{
+		if (Transaction)
+		{
+			Transaction->Cancel();
+		}
+		UE::Geometry::AppendError(Debug, EGeometryScriptErrorType::InvalidInputs, LOCTEXT("CopySkinWeightProfileToSkeletalMesh_InvalidAssetGeometry", "CopySkinWeightProfileToSkeletalMesh: Target skeletal mesh doesnt have the same number of vertices as the source mesh."));
+		return FromDynamicMesh;
+	}
+	
+	FSkeletalMeshAttributes MeshAttributes(*MeshDescription);
+	MeshAttributes.Register();
+
+	ToSkeletalMeshAsset->ModifyMeshDescription(TargetLOD.LODIndex);
+
+	if (TargetProfileName != FSkeletalMeshAttributes::DefaultSkinWeightProfileName)
+	{
+		if (!Options.bOverwriteExistingProfile && MeshAttributes.GetSkinWeightProfileNames().Contains(TargetProfileName))
+		{
+			if (Transaction)
+			{
+				Transaction->Cancel();
+			}
+			UE::Geometry::AppendError(Debug, EGeometryScriptErrorType::InvalidInputs, LOCTEXT("CopySkinWeightProfileToSkeletalMesh_CantOverrideProfile", "CopySkinWeightProfileToSkeletalMesh: Skin profile name already exists on the target mesh."));
+			return FromDynamicMesh;
+		}
+		
+		if (!MeshAttributes.RegisterSkinWeightAttribute(TargetProfileName))
+		{
+			if (Transaction)
+			{
+				Transaction->Cancel();
+			}
+			UE::Geometry::AppendError(Debug, EGeometryScriptErrorType::InvalidInputs, LOCTEXT("CopySkinWeightProfileToSkeletalMesh_InvalidProfileName", "CopySkinWeightProfileToSkeletalMesh: Cannot create target skin weight profile with the given profile name."));
+			return FromDynamicMesh;
+		}
+	}
+
+	const FDynamicMeshVertexSkinWeightsAttribute* SourceProfileAttribute = SourceMesh.Attributes()->GetSkinWeightsAttribute(SourceProfileName);
+	FSkinWeightsVertexAttributesRef TargetProfileAttribute = MeshAttributes.GetVertexSkinWeights(TargetProfileName);
+
+	for (int32 SourceVID = 0; SourceVID < SourceMesh.VertexCount(); ++SourceVID)
+	{
+		const int32 TargetVID = NonManifoldMappingSupport.GetOriginalNonManifoldVertexID(SourceVID);
+
+		UE::AnimationCore::FBoneWeights BoneWeights;
+		SourceProfileAttribute->GetValue(SourceVID, BoneWeights);
+		TargetProfileAttribute.Set(TargetVID, BoneWeights);
+	}
+
+	ToSkeletalMeshAsset->CommitMeshDescription(TargetLOD.LODIndex);
+
+	if (Options.bDeferMeshPostEditChange == false)
+	{
+		ToSkeletalMeshAsset->PostEditChange();
+	}
+
+	Outcome = EGeometryScriptOutcomePins::Success;
+#else
+	UE::Geometry::AppendError(Debug, EGeometryScriptErrorType::InvalidInputs, LOCTEXT("CopySkinWeightProfileToSkeletalMesh_EditorOnly", "CopySkinWeightProfileToSkeletalMesh: Not currently supported at Runtime"));
+#endif
+	
+	return FromDynamicMesh;
+}
+
 #undef LOCTEXT_NAMESPACE
diff --git a/Engine/Plugins/Runtime/GeometryScripting/Source/GeometryScriptingCore/Public/GeometryScript/MeshAssetFunctions.h b/Engine/Plugins/Runtime/GeometryScripting/Source/GeometryScriptingCore/Public/GeometryScript/MeshAssetFunctions.h
--- a/Engine/Plugins/Runtime/GeometryScripting/Source/GeometryScriptingCore/Public/GeometryScript/MeshAssetFunctions.h
+++ b/Engine/Plugins/Runtime/GeometryScripting/Source/GeometryScriptingCore/Public/GeometryScript/MeshAssetFunctions.h
@@ -18,7 +18,7 @@ USTRUCT(BlueprintType)
 struct GEOMETRYSCRIPTINGCORE_API FGeometryScriptCopyMeshFromAssetOptions
 {
 	GENERATED_BODY()
-public:
+	
 	// Whether to apply Build Settings during the mesh copy.
 	UPROPERTY(EditAnywhere, BlueprintReadWrite, Category = Options)
 	bool bApplyBuildSettings = true;
@@ -62,7 +62,7 @@ USTRUCT(BlueprintType)
 struct GEOMETRYSCRIPTINGCORE_API FGeometryScriptCopyMeshToAssetOptions
 {
 	GENERATED_BODY()
-public:
+	
 	UPROPERTY(EditAnywhere, BlueprintReadWrite, Category = Options)
 	bool bEnableRecomputeNormals = false;
 
@@ -122,7 +122,6 @@ USTRUCT(BlueprintType)
 struct GEOMETRYSCRIPTINGCORE_API FGeometryScriptCopyMorphTargetToAssetOptions
 {
 	GENERATED_BODY()
-public:
 
 	/** If true and the morph target with the given name exists, it will be overwritten. If false, will abort and print a 
 	 * console error. */
@@ -136,6 +135,25 @@ public:
 	bool bDeferMeshPostEditChange = false;
 };
 
+
+USTRUCT(BlueprintType)
+struct GEOMETRYSCRIPTINGCORE_API FGeometryScriptCopySkinWeightProfileToAssetOptions
+{
+	GENERATED_BODY()
+
+	/** If true and a skin weight profile with the given name exists, it will be overwritten. 
+	 *  If false, will abort and print a console error. */
+	UPROPERTY(EditAnywhere, BlueprintReadWrite, Category = Options)
+	bool bOverwriteExistingProfile = true;
+
+	UPROPERTY(EditAnywhere, BlueprintReadWrite, Category = Options)
+	bool bEmitTransaction = true;
+
+	UPROPERTY(EditAnywhere, BlueprintReadWrite, Category = Options)
+	bool bDeferMeshPostEditChange = false;
+};
+
+
 // Although the class name indicates StaticMeshFunctions, that was a naming mistake that is difficult
 // to correct. This class is intended to serve as a generic asset utils function library. The naming
 // issue is only visible at the C++ level. It is not visible in Python or BP.
@@ -359,6 +377,23 @@ public:
 		EGeometryScriptOutcomePins& Outcome,
 		UGeometryScriptDebug* Debug = nullptr);
 
+   /** 
+	* Add a Dynamic Mesh skin weight profile to a Skeletal Mesh Asset.
+	* 
+	* @param FromDynamicMesh the dynamic mesh representing the geometry of the morph target
+	* @param ToSkeletalMeshAsset the asset we are writing the morph target into
+	* @param TargetProfileName the name of the skin weight profile as it will appear in the UI. Leave blank for the default profile.
+	* @param SourceProfileName The name of the skin weight profile to copy from the dynamic mesh. Leave blank for the default profile.
+	*/
+	UFUNCTION(BlueprintCallable, Category = "GeometryScript|SkeletalMesh", meta = (ExpandEnumAsExecs = "Outcome"))
+	static UPARAM(DisplayName = "Dynamic Mesh") UDynamicMesh* 
+	CopySkinWeightProfileToSkeletalMesh(
+		UDynamicMesh* FromDynamicMesh, 
+		USkeletalMesh* ToSkeletalMeshAsset,
+		FName TargetProfileName,
+		FName SourceProfileName,
+		FGeometryScriptCopySkinWeightProfileToAssetOptions Options,
+		FGeometryScriptMeshWriteLOD TargetLOD,
+		EGeometryScriptOutcomePins& Outcome,
+		UGeometryScriptDebug* Debug = nullptr);
 };
-
-
