Geoscripting: Reject empty morph targets right at creation.

#rnx

diff --git a/Engine/Plugins/Runtime/GeometryScripting/Source/GeometryScriptingCore/Private/MeshAssetFunctions.cpp b/Engine/Plugins/Runtime/GeometryScripting/Source/GeometryScriptingCore/Private/MeshAssetFunctions.cpp
--- a/Engine/Plugins/Runtime/GeometryScripting/Source/GeometryScriptingCore/Private/MeshAssetFunctions.cpp
+++ b/Engine/Plugins/Runtime/GeometryScripting/Source/GeometryScriptingCore/Private/MeshAssetFunctions.cpp
@@ -1082,9 +1082,10 @@ UDynamicMesh* UGeometryScriptLibrary_StaticMeshFunctions::CopyMorphTargetToSkele
 	// flush any pending rendering commands, which might touch a component while we are rebuilding it's mesh
 	FlushRenderingCommands();
 
+	TUniquePtr<FScopedTransaction> Transaction;
 	if (Options.bEmitTransaction)
 	{
-		GEditor->BeginTransaction(LOCTEXT("UpdateSkeletalMesh", "Update Skeletal Mesh"));
+		Transaction = MakeUnique<FScopedTransaction>(LOCTEXT("UpdateSkeletalMesh", "Update Skeletal Mesh"));
 	}
 	
 	// make sure transactional flag is on for this asset
@@ -1105,6 +1106,10 @@ UDynamicMesh* UGeometryScriptLibrary_StaticMeshFunctions::CopyMorphTargetToSkele
 
 	if (MeshDescription == nullptr)
 	{
+		if (Transaction)
+		{
+			Transaction->Cancel();
+		}
 		UE::Geometry::AppendError(Debug, EGeometryScriptErrorType::InvalidInputs, LOCTEXT("CopyMorphTargetToSkeletalMesh_TargetMeshDescription", "CopyMorphTargetToSkeletalMesh: Failed to generate the mesh data for the Target LOD Index"));
 		return FromDynamicMesh;
 	}
@@ -1131,6 +1136,10 @@ UDynamicMesh* UGeometryScriptLibrary_StaticMeshFunctions::CopyMorphTargetToSkele
 	}
 	if (MeshDescription->Vertices().Num() != SourceVertexCount)
 	{
+		if (Transaction)
+		{
+			Transaction->Cancel();
+		}
 		UE::Geometry::AppendError(Debug, EGeometryScriptErrorType::InvalidInputs, LOCTEXT("CopyMorphTargetToSkeletalMesh_InvalidMorphTargetGeometry", "CopyMorphTargetToSkeletalMesh: Morph target mesh doesnt have the same number of vertices as the skeletal mesh."));
 		return FromDynamicMesh;
 	}
@@ -1144,6 +1153,10 @@ UDynamicMesh* UGeometryScriptLibrary_StaticMeshFunctions::CopyMorphTargetToSkele
 	{
 		if (!Options.bOverwriteExistingTarget) // only throw error if we dont want to overwrite the existing target 
 		{
+			if (Transaction)
+			{
+				Transaction->Cancel();
+			}
 			UE::Geometry::AppendError(Debug, EGeometryScriptErrorType::InvalidInputs, LOCTEXT("CopyMorphTargetToSkeletalMesh_InvalidMorphTargetName1", "CopyMorphTargetToSkeletalMesh: Morph target name already exists"));
 			return FromDynamicMesh;
 		}
@@ -1152,15 +1165,23 @@ UDynamicMesh* UGeometryScriptLibrary_StaticMeshFunctions::CopyMorphTargetToSkele
 	{
 		if (!MeshAttributes.RegisterMorphTargetAttribute(MorphTargetName, false))
 		{
+			if (Transaction)
+			{
+				Transaction->Cancel();
+			}
 			UE::Geometry::AppendError(Debug, EGeometryScriptErrorType::InvalidInputs, LOCTEXT("CopyMorphTargetToSkeletalMesh_InvalidMorphTargetName2", "CopyMorphTargetToSkeletalMesh: Morph target name is invalid."));
 			return FromDynamicMesh;
 		}
 	}
 
+	const FSkeletalMeshLODInfo* LODInfo = ToSkeletalMeshAsset->GetLODInfo(TargetLOD.LODIndex);
+	const float MorphThresholdSquared = LODInfo->BuildSettings.MorphThresholdPosition * LODInfo->BuildSettings.MorphThresholdPosition;
 	
 	TVertexAttributesRef<FVector3f> PositionDelta = MeshAttributes.GetVertexMorphPositionDelta(MorphTargetName);
 	TVertexAttributesRef<FVector3f> VertexPositions = MeshAttributes.GetVertexPositions();
 
+	bool bMorphTargetIsEmpty = true;
+	
 	for (int32 SourceVID = 0; SourceVID < SourceMesh.VertexCount(); ++SourceVID)
 	{
 		const int32 TargetVID = NonManifoldMappingSupport.GetOriginalNonManifoldVertexID(SourceVID);
@@ -1169,7 +1190,22 @@ UDynamicMesh* UGeometryScriptLibrary_StaticMeshFunctions::CopyMorphTargetToSkele
 		const FVector3f V1 = VertexPositions[TargetVID];
 
 		const FVector3f Delta = FVector3f{V0} - V1;
-		PositionDelta.Set(TargetVID, Delta);
+		if (Delta.SquaredLength() > MorphThresholdSquared)
+		{
+			bMorphTargetIsEmpty = false;
+			PositionDelta.Set(TargetVID, Delta);
+		}
+	}
+
+	if (!bMorphTargetIsEmpty)
+	{
+		MeshAttributes.UnregisterMorphTargetAttribute(MorphTargetName);
+		if (Transaction)
+		{
+			Transaction->Cancel();
+		}
+		UE::Geometry::AppendWarning(Debug, EGeometryScriptErrorType::OperationFailed, LOCTEXT("CopyMorphTargetToSkeletalMesh_EmptyMorphTarget", "CopyMorphTargetToSkeletalMesh: Morph target is empty since it does not differ from the base mesh's vertex position."));
+		return FromDynamicMesh;
 	}
 
 	ToSkeletalMeshAsset->CommitMeshDescription(TargetLOD.LODIndex);
@@ -1179,11 +1215,6 @@ UDynamicMesh* UGeometryScriptLibrary_StaticMeshFunctions::CopyMorphTargetToSkele
 		ToSkeletalMeshAsset->PostEditChange();
 	}
 
-	if (Options.bEmitTransaction)
-	{
-		GEditor->EndTransaction();
-	}
-
 	Outcome = EGeometryScriptOutcomePins::Success;
 #else
 	UE::Geometry::AppendError(Debug, EGeometryScriptErrorType::InvalidInputs, LOCTEXT("CopyMorphTargetToSkeletalMesh_EditorOnly", "CopyMorphTargetToSkeletalMesh: Not currently supported at Runtime"));
@@ -1267,8 +1298,6 @@ UDynamicMesh* UGeometryScriptLibrary_StaticMeshFunctions::CopySkinWeightProfileT
 		SourceProfileName = FSkeletalMeshAttributes::DefaultSkinWeightProfileName;
 	}
 	
-	// If the dynamic mesh has non-manifold information, use that to figure out what the original vertex count was.
-	// Otherwise, we assume that they have a 1:1 match.
 	const FDynamicMesh3& SourceMesh = FromDynamicMesh->GetMeshRef();
 	if (!SourceMesh.HasAttributes() || !SourceMesh.Attributes()->HasSkinWeightsAttribute(SourceProfileName))
 	{
@@ -1283,6 +1312,8 @@ UDynamicMesh* UGeometryScriptLibrary_StaticMeshFunctions::CopySkinWeightProfileT
 	const FNonManifoldMappingSupport NonManifoldMappingSupport(SourceMesh);
 	int32 SourceVertexCount;
 
+	// If the dynamic mesh has non-manifold information, use that to figure out what the original vertex count was.
+	// Otherwise, we assume that they have a 1:1 match.
 	if (NonManifoldMappingSupport.IsNonManifoldVertexInSource())
 	{
 		TSet<int32> UniqueVertices;
