fix crash in GetMaterialIDsOfTriangles if passed null triangle list
+ fix behavior for invalid tris to match documentation in comment
+ fix parameter names with bad spacing (IDList -> ID List)

diff --git a/Engine/Plugins/Runtime/GeometryScripting/Source/GeometryScriptingCore/Private/MeshMaterialFunctions.cpp b/Engine/Plugins/Runtime/GeometryScripting/Source/GeometryScriptingCore/Private/MeshMaterialFunctions.cpp
--- a/Engine/Plugins/Runtime/GeometryScripting/Source/GeometryScriptingCore/Private/MeshMaterialFunctions.cpp
+++ b/Engine/Plugins/Runtime/GeometryScripting/Source/GeometryScriptingCore/Private/MeshMaterialFunctions.cpp
@@ -297,7 +297,6 @@ UDynamicMesh* UGeometryScriptLibrary_MeshMaterialFunctions::GetMaterialIDsOfTria
     UGeometryScriptDebug* Debug)
 {
 	MaterialIDList.Reset(EGeometryScriptIndexType::MaterialID);
-	MaterialIDList.List->SetNumZeroed(TriangleIDList.List->Num());
 	
 	if (TargetMesh == nullptr)
 	{
@@ -321,17 +320,24 @@ UDynamicMesh* UGeometryScriptLibrary_MeshMaterialFunctions::GetMaterialIDsOfTria
 	
 	bool bHasMaterials = false;
 	bool bAllValidTriangles = true;
-	SimpleMeshMaterialQuery<int32>(TargetMesh, bHasMaterials, 0, [&](const FDynamicMesh3& EditMesh, const FDynamicMeshMaterialAttribute& MaterialIDs) {
+	SimpleMeshMaterialQuery<int32>(TargetMesh, bHasMaterials, 0, 
+		[&MaterialIDList, &bHasMaterials, &bAllValidTriangles, &TriangleIDList] 
+		(const FDynamicMesh3& EditMesh, const FDynamicMeshMaterialAttribute& MaterialIDs) 
+	{
 		TArray<int>& MaterialIDArray = *MaterialIDList.List;
+		MaterialIDArray.SetNum(TriangleIDList.List->Num());
 		for (int32 i = 0; i < TriangleIDList.List->Num(); i++)
 		{
 			int32 TriangleID = (*TriangleIDList.List)[i];
 			if (EditMesh.IsTriangle(TriangleID) == false)
 			{
 				bAllValidTriangles = false;
-				return 0;
+				MaterialIDArray[i] = -1;
+			}
+			else
+			{
+				MaterialIDArray[i] = MaterialIDs.GetValue(TriangleID);
 			}
-			MaterialIDArray[i] = MaterialIDs.GetValue(TriangleID);
 		}
 		return 0;
 	});
diff --git a/Engine/Plugins/Runtime/GeometryScripting/Source/GeometryScriptingCore/Public/GeometryScript/MeshMaterialFunctions.h b/Engine/Plugins/Runtime/GeometryScripting/Source/GeometryScriptingCore/Public/GeometryScript/MeshMaterialFunctions.h
--- a/Engine/Plugins/Runtime/GeometryScripting/Source/GeometryScriptingCore/Public/GeometryScript/MeshMaterialFunctions.h
+++ b/Engine/Plugins/Runtime/GeometryScripting/Source/GeometryScriptingCore/Public/GeometryScript/MeshMaterialFunctions.h
@@ -109,8 +109,8 @@ public:
 	static UPARAM(DisplayName = "Target Mesh") UDynamicMesh* 
 	GetMaterialIDsOfTriangles(
 		UDynamicMesh* TargetMesh,
-		FGeometryScriptIndexList TriangleIDList,
-		FGeometryScriptIndexList& MaterialIDList,
+		UPARAM(DisplayName = "Triangle ID List") FGeometryScriptIndexList TriangleIDList,
+		UPARAM(DisplayName = "Material ID List") FGeometryScriptIndexList& MaterialIDList,
 		UGeometryScriptDebug* Debug = nullptr);
 	
 	/**
@@ -130,7 +130,7 @@ public:
 	GetTrianglesByMaterialID( 
 		UDynamicMesh* TargetMesh, 
 		int MaterialID,
-		FGeometryScriptIndexList& TriangleIDList,
+		UPARAM(DisplayName = "Triangle ID List") FGeometryScriptIndexList& TriangleIDList,
 		UGeometryScriptDebug* Debug = nullptr);
 
 	/**
@@ -155,7 +155,7 @@ public:
 	static UPARAM(DisplayName = "Target Mesh") UDynamicMesh*
 	SetAllTriangleMaterialIDs(
 		UDynamicMesh* TargetMesh,
-		FGeometryScriptIndexList TriangleMaterialIDList,
+		UPARAM(DisplayName = "Triangle Material ID List") FGeometryScriptIndexList TriangleMaterialIDList,
 		bool bDeferChangeNotifications = false,
 		UGeometryScriptDebug* Debug = nullptr);
 
@@ -169,7 +169,7 @@ public:
 	static UPARAM(DisplayName = "Target Mesh") UDynamicMesh* 
 	SetMaterialIDOnTriangles(
 		UDynamicMesh* TargetMesh,
-		FGeometryScriptIndexList TriangleIDList,
+		UPARAM(DisplayName = "Triangle ID List") FGeometryScriptIndexList TriangleIDList,
 		int MaterialID,
         bool bDeferChangeNotifications = false,
 		UGeometryScriptDebug* Debug = nullptr);
