Add missing null check to the StaticMeshHasCustomizedCollision Geometry Script method

diff --git a/Engine/Plugins/Runtime/GeometryScripting/Source/GeometryScriptingCore/Private/CollisionFunctions.cpp b/Engine/Plugins/Runtime/GeometryScripting/Source/GeometryScriptingCore/Private/CollisionFunctions.cpp
--- a/Engine/Plugins/Runtime/GeometryScripting/Source/GeometryScriptingCore/Private/CollisionFunctions.cpp
+++ b/Engine/Plugins/Runtime/GeometryScripting/Source/GeometryScriptingCore/Private/CollisionFunctions.cpp
@@ -413,6 +413,11 @@ void UGeometryScriptLibrary_CollisionFunctions::SetStaticMeshCollisionFromCompon
 bool UGeometryScriptLibrary_CollisionFunctions::StaticMeshHasCustomizedCollision(UStaticMesh* StaticMeshAsset)
 {
 #if WITH_EDITORONLY_DATA
+	if (!StaticMeshAsset)
+	{
+		UE_LOG(LogGeometry, Warning, TEXT("StaticMeshHasCustomizedCollision: StaticMeshAsset is Null"));
+		return false;
+	}
 	return StaticMeshAsset->bCustomizedCollision;
 #else
 	return false;
