GeoScripting: Print the calling BP function when emitting warnings/errors in geoscripting calls.
- This helps in understanding where the error/warning is coming from.

diff --git a/Engine/Plugins/Runtime/GeometryScripting/Source/GeometryScriptingCore/Private/GeometryScriptTypes.cpp b/Engine/Plugins/Runtime/GeometryScripting/Source/GeometryScriptingCore/Private/GeometryScriptTypes.cpp
--- a/Engine/Plugins/Runtime/GeometryScripting/Source/GeometryScriptingCore/Private/GeometryScriptTypes.cpp
+++ b/Engine/Plugins/Runtime/GeometryScripting/Source/GeometryScriptingCore/Private/GeometryScriptTypes.cpp
@@ -7,17 +7,49 @@
 #include "CompGeom/ConvexDecomposition3.h"
 
 
+static FString FindCurrentBPFunction()
+{
+#if WITH_EDITOR
+	const FBlueprintContextTracker* ContextTracker = FBlueprintContextTracker::TryGet();
+	if (ContextTracker && !ContextTracker->GetCurrentScriptStack().IsEmpty())
+	{
+		TStringBuilder<256> StringBuilder;
+		ContextTracker->GetCurrentScriptStack().Last()->GetStackDescription(StringBuilder);
+		return StringBuilder.ToString();
+	}
+#endif
+	return {};
+}
+
 FGeometryScriptDebugMessage UE::Geometry::MakeScriptError(EGeometryScriptErrorType ErrorTypeIn, const FText& MessageIn)
 {
-	UE_LOG(LogGeometry, Warning, TEXT("GeometryScriptError: %s"), *MessageIn.ToString() );
+	FString CurrentBPFunction = FindCurrentBPFunction();
+	
+	if (CurrentBPFunction.IsEmpty())
+	{
+		UE_LOG(LogGeometry, Error, TEXT("%s"), *MessageIn.ToString() );
+	}
+	else
+	{
+		UE_LOG(LogGeometry, Error, TEXT("%s [Called from: %s]"), *MessageIn.ToString(), *CurrentBPFunction);
+	}
 
 	return FGeometryScriptDebugMessage{ EGeometryScriptDebugMessageType::ErrorMessage, ErrorTypeIn, MessageIn };
 }
 
 FGeometryScriptDebugMessage UE::Geometry::MakeScriptWarning(EGeometryScriptErrorType WarningTypeIn, const FText& MessageIn)
 {
-	UE_LOG(LogGeometry, Warning, TEXT("GeometryScriptWarning: %s"), *MessageIn.ToString() );
-
+	FString CurrentBPFunction = FindCurrentBPFunction();
+	
+	if (CurrentBPFunction.IsEmpty())
+	{
+		UE_LOG(LogGeometry, Warning, TEXT("%s"), *MessageIn.ToString() );
+	}
+	else
+	{
+		UE_LOG(LogGeometry, Warning, TEXT("%s [Called from: %s]"), *MessageIn.ToString(), *CurrentBPFunction);
+	}
+	
 	return FGeometryScriptDebugMessage{ EGeometryScriptDebugMessageType::WarningMessage, WarningTypeIn, MessageIn };
 }
 
