fix AppendMesh nodes inconsistently handling transforms with negative scales

diff --git a/Engine/Plugins/Runtime/GeometryScripting/Source/GeometryScriptingCore/Private/MeshBasicEditFunctions.cpp b/Engine/Plugins/Runtime/GeometryScripting/Source/GeometryScriptingCore/Private/MeshBasicEditFunctions.cpp
--- a/Engine/Plugins/Runtime/GeometryScripting/Source/GeometryScriptingCore/Private/MeshBasicEditFunctions.cpp
+++ b/Engine/Plugins/Runtime/GeometryScripting/Source/GeometryScriptingCore/Private/MeshBasicEditFunctions.cpp
@@ -458,7 +458,8 @@ namespace UE::Private::AppendHelpers
 				}
 				Editor.AppendMesh(UseOtherMesh, TmpMappings,
 					[&XForm](int, const FVector3d& Position) { return XForm.TransformPosition(Position); },
-					[&XForm](int, const FVector3d& Normal) { return XForm.TransformNormal(Normal); });
+					[&XForm](int, const FVector3d& Normal) { return XForm.TransformNormal(Normal); },
+					XForm.GetDeterminant() < 0);
 				if (AppendMaterialRemap)
 				{
 					ApplyMaterialRemap(TmpMappings, *AppendMaterialRemap, AppendToMesh);
@@ -534,7 +535,8 @@ namespace UE::Private::AppendHelpers
 
 					Editor.AppendMesh(UseOtherMesh, TmpMappings,
 						[&TransformSequence](int, const FVector3d& Position) { return TransformSequence.TransformPosition(Position); },
-						[&TransformSequence](int, const FVector3d& Normal) { return TransformSequence.TransformNormal(Normal); });
+						[&TransformSequence](int, const FVector3d& Normal) { return TransformSequence.TransformNormal(Normal); },
+						TransformSequence.WillInvert());
 					if (AppendMaterialRemap)
 					{
 						ApplyMaterialRemap(TmpMappings, *AppendMaterialRemap, AppendToMesh);
diff --git a/Engine/Source/Runtime/GeometryCore/Private/DynamicMeshEditor.cpp b/Engine/Source/Runtime/GeometryCore/Private/DynamicMeshEditor.cpp
--- a/Engine/Source/Runtime/GeometryCore/Private/DynamicMeshEditor.cpp
+++ b/Engine/Source/Runtime/GeometryCore/Private/DynamicMeshEditor.cpp
@@ -1899,7 +1899,8 @@ int FDynamicMeshEditor::FindOrCreateDuplicateGroup(int TriangleID, FMeshIndexMap
 void FDynamicMeshEditor::AppendMesh(const FDynamicMesh3* AppendMesh,
 	FMeshIndexMappings& IndexMapsOut, 
 	TFunction<FVector3d(int, const FVector3d&)> PositionTransform,
-	TFunction<FVector3d(int, const FVector3d&)> NormalTransform)
+	TFunction<FVector3d(int, const FVector3d&)> NormalTransform,
+	bool bReverseOrientation)
 {
 	// todo: handle this case by making a copy?
 	check(AppendMesh != Mesh);
@@ -2111,6 +2112,17 @@ void FDynamicMeshEditor::AppendMesh(const FDynamicMesh3* AppendMesh,
 			}
 		}
 	}
+
+	// Flip tri orientations if requested, after all appends -- note we do this after mapping across all the attributes, as that logic assumes triangle vertices have matching vertex order
+	if (bReverseOrientation)
+	{
+		const TMap<int32, int32>& ReverseTriMap = TriangleMap.GetReverseMap();
+		for (const TPair<int32, int32> KV : ReverseTriMap)
+		{
+			int32 NewTID = KV.Key;
+			Mesh->ReverseTriOrientation(NewTID);
+		}
+	}
 }
 
 
diff --git a/Engine/Source/Runtime/GeometryCore/Public/DynamicMeshEditor.h b/Engine/Source/Runtime/GeometryCore/Public/DynamicMeshEditor.h
--- a/Engine/Source/Runtime/GeometryCore/Public/DynamicMeshEditor.h
+++ b/Engine/Source/Runtime/GeometryCore/Public/DynamicMeshEditor.h
@@ -525,7 +525,8 @@ public:
 	 */
 	GEOMETRYCORE_API void AppendMesh(const FDynamicMesh3* AppendMesh, FMeshIndexMappings& IndexMapsOut, 
 		TFunction<FVector3d(int, const FVector3d&)> PositionTransform = nullptr,
-		TFunction<FVector3d(int, const FVector3d&)> NormalTransform = nullptr);
+		TFunction<FVector3d(int, const FVector3d&)> NormalTransform = nullptr,
+		bool bReverseOrientation = false);
 
 	/**
 	 * Append input mesh to our internal Mesh. If the internal Mesh has attributes enabled,
