Add world context object version of the geometry script "Create New Volume From Mesh" method

diff --git a/Engine/Plugins/Runtime/GeometryScripting/Source/GeometryScriptingEditor/Private/CreateNewAssetUtilityFunctions.cpp b/Engine/Plugins/Runtime/GeometryScripting/Source/GeometryScriptingEditor/Private/CreateNewAssetUtilityFunctions.cpp
--- a/Engine/Plugins/Runtime/GeometryScripting/Source/GeometryScriptingEditor/Private/CreateNewAssetUtilityFunctions.cpp
+++ b/Engine/Plugins/Runtime/GeometryScripting/Source/GeometryScriptingEditor/Private/CreateNewAssetUtilityFunctions.cpp
@@ -8,6 +8,7 @@
 #include "DynamicMesh/MeshBones.h"
 #include "UDynamicMesh.h"
 
+#include "Engine/Engine.h" // for GEngine
 #include "Engine/StaticMesh.h"
 #include "Engine/SkeletalMesh.h"
 #include "Engine/StaticMeshActor.h"
@@ -163,7 +164,26 @@ AVolume* UGeometryScriptLibrary_CreateNewAssetFunctions::CreateNewVolumeFromMesh
 }
 
 
-
+AVolume* UGeometryScriptLibrary_CreateNewAssetFunctions::CreateNewVolumeFromMesh_WorldContext(
+	UObject* WorldContextObject,
+	UDynamicMesh* FromDynamicMesh,
+	FTransform ActorTransform,
+	FString BaseActorName,
+	FGeometryScriptCreateNewVolumeFromMeshOptions Options,
+	EGeometryScriptOutcomePins& Outcome,
+	UGeometryScriptDebug* Debug)
+{
+	if (!GEngine)
+	{
+		UE::Geometry::AppendError(Debug, EGeometryScriptErrorType::InvalidInputs, LOCTEXT("CreateNewVolumeFromMesh_NoEngine", "CreateNewVolumeFromMesh: Engine is not available. Cannot get world context object"));
+		return nullptr;
+	}
+	if (UWorld* World = GEngine->GetWorldFromContextObject(WorldContextObject, EGetWorldErrorMode::LogAndReturnNull))
+	{
+		return CreateNewVolumeFromMesh(FromDynamicMesh, World, ActorTransform, BaseActorName, Options, Outcome, Debug);
+	}
+	return nullptr;
+}
 
 
 
diff --git a/Engine/Plugins/Runtime/GeometryScripting/Source/GeometryScriptingEditor/Public/GeometryScript/CreateNewAssetUtilityFunctions.h b/Engine/Plugins/Runtime/GeometryScripting/Source/GeometryScriptingEditor/Public/GeometryScript/CreateNewAssetUtilityFunctions.h
--- a/Engine/Plugins/Runtime/GeometryScripting/Source/GeometryScriptingEditor/Public/GeometryScript/CreateNewAssetUtilityFunctions.h
+++ b/Engine/Plugins/Runtime/GeometryScripting/Source/GeometryScriptingEditor/Public/GeometryScript/CreateNewAssetUtilityFunctions.h
@@ -131,7 +131,24 @@ public:
 		UGeometryScriptDebug* Debug = nullptr);
 
 
-	UFUNCTION(BlueprintCallable, Category = "GeometryScript|AssetManagement", meta = (ExpandEnumAsExecs = "Outcome"))
+	/**
+	 * Create a new Volume from a Dynamic Mesh, in the same world as the calling blueprint.
+	 */
+	UFUNCTION(BlueprintCallable, Category = "GeometryScript|AssetManagement", meta = (ExpandEnumAsExecs = "Outcome", DisplayName = "Create New Volume From Mesh", WorldContext = "WorldContextObject"))
+	static UPARAM(DisplayName = "Volume Actor") AVolume*
+	CreateNewVolumeFromMesh_WorldContext(
+		UObject* WorldContextObject,
+		UDynamicMesh* FromDynamicMesh,
+		FTransform ActorTransform,
+		FString BaseActorName,
+		FGeometryScriptCreateNewVolumeFromMeshOptions Options,
+		EGeometryScriptOutcomePins& Outcome,
+		UGeometryScriptDebug* Debug = nullptr);
+
+	/**
+	 * Create a new Volume from a Dynamic Mesh, in the specified world.
+	 */
+	UFUNCTION(BlueprintCallable, Category = "GeometryScript|AssetManagement", meta = (ExpandEnumAsExecs = "Outcome", DisplayName = "Create New Volume From Mesh In Target World"))
 	static UPARAM(DisplayName = "Volume Actor") AVolume* 
 	CreateNewVolumeFromMesh(
 		UDynamicMesh* FromDynamicMesh, 
