clean up geometry script methods display names

diff --git a/Engine/Plugins/Runtime/GeometryScripting/Source/GeometryScriptingCore/Public/GeometryScript/MeshSubdivideFunctions.h b/Engine/Plugins/Runtime/GeometryScripting/Source/GeometryScriptingCore/Public/GeometryScript/MeshSubdivideFunctions.h
--- a/Engine/Plugins/Runtime/GeometryScripting/Source/GeometryScriptingCore/Public/GeometryScript/MeshSubdivideFunctions.h
+++ b/Engine/Plugins/Runtime/GeometryScripting/Source/GeometryScriptingCore/Public/GeometryScript/MeshSubdivideFunctions.h
@@ -58,7 +58,7 @@ public:
 	/**
 	 * Apply PN Tessellation to the Target Mesh as controlled by the Tessellation Level and the Options.
 	 */
-	UFUNCTION(BlueprintCallable, Category = "GeometryScript|Subdivide", meta=(ScriptMethod))
+	UFUNCTION(BlueprintCallable, Category = "GeometryScript|Subdivide", meta=(ScriptMethod, DisplayName = "Apply PN Tessellation"))
 	static UPARAM(DisplayName = "Target Mesh") UDynamicMesh* 
 	ApplyPNTessellation(
 		UDynamicMesh* TargetMesh,
diff --git a/Engine/Plugins/Runtime/GeometryScripting/Source/GeometryScriptingEditor/Public/GeometryScript/OpenSubdivUtilityFunctions.h b/Engine/Plugins/Runtime/GeometryScripting/Source/GeometryScriptingEditor/Public/GeometryScript/OpenSubdivUtilityFunctions.h
--- a/Engine/Plugins/Runtime/GeometryScripting/Source/GeometryScriptingEditor/Public/GeometryScript/OpenSubdivUtilityFunctions.h
+++ b/Engine/Plugins/Runtime/GeometryScripting/Source/GeometryScriptingEditor/Public/GeometryScript/OpenSubdivUtilityFunctions.h
@@ -14,7 +14,7 @@ class GEOMETRYSCRIPTINGEDITOR_API UGeometryScriptLibrary_OpenSubdivFunctions : p
 public:
 
 
-	UFUNCTION(BlueprintCallable, Category = "GeometryScript|OpenSubdiv", meta=(ScriptMethod, DisplayName = "ApplyPolyGroupCatmullClarkSubD"))
+	UFUNCTION(BlueprintCallable, Category = "GeometryScript|OpenSubdiv", meta=(ScriptMethod, DisplayName = "Apply PolyGroup Catmull Clark SubD"))
 	static UPARAM(DisplayName = "Target Mesh") UDynamicMesh* 
 	ApplyPolygroupCatmullClarkSubD(
 		UDynamicMesh* FromDynamicMesh, 
@@ -23,7 +23,7 @@ public:
 		UGeometryScriptDebug* Debug = nullptr);
 
 
-	UFUNCTION(BlueprintCallable, Category = "GeometryScript|OpenSubdiv", meta=(ScriptMethod))
+	UFUNCTION(BlueprintCallable, Category = "GeometryScript|OpenSubdiv", meta=(ScriptMethod, DisplayName = "Apply Triangle Loop SubD"))
 	static UPARAM(DisplayName = "Target Mesh") UDynamicMesh* 
 	ApplyTriangleLoopSubD(
 		UDynamicMesh* FromDynamicMesh, 
