add directional scaling option to convex decomposition tolerance

diff --git a/Engine/Plugins/Runtime/GeometryScripting/Source/GeometryScriptingCore/Private/CollisionFunctions.cpp b/Engine/Plugins/Runtime/GeometryScripting/Source/GeometryScriptingCore/Private/CollisionFunctions.cpp
--- a/Engine/Plugins/Runtime/GeometryScripting/Source/GeometryScriptingCore/Private/CollisionFunctions.cpp
+++ b/Engine/Plugins/Runtime/GeometryScripting/Source/GeometryScriptingCore/Private/CollisionFunctions.cpp
@@ -261,6 +261,51 @@ static UE::Geometry::FNegativeSpaceSampleSettings::ESampleMethod ConvertNegative
 	return UE::Geometry::FNegativeSpaceSampleSettings::ESampleMethod::Uniform;
 }
 
+static void ApplyNegativeSpaceToleranceDirectionalScales(UE::Geometry::FNegativeSpaceSampleSettings& NegativeSpaceSettings, const TArray<FNegativeSpaceDirectionalToleranceScale>& ToleranceDirectionalScales)
+{
+	if (ToleranceDirectionalScales.IsEmpty())
+	{
+		return;
+	}
+
+	// Copy the input scalars to transform them for use by the custom scale function
+	TArray<FNegativeSpaceDirectionalToleranceScale> TransformedScalars = ToleranceDirectionalScales;
+	double MinScale = 1.0;
+
+	for (FNegativeSpaceDirectionalToleranceScale& Scalar : TransformedScalars)
+	{
+		Scalar.Direction.Normalize();
+		Scalar.AngleRange = FMath::Max(FMathd::ZeroTolerance, FMath::Cos(FMathd::DegToRad * Scalar.AngleRange)); // transform to dot product threshold
+		Scalar.ScaleFactor = FMath::Clamp(Scalar.ScaleFactor, .1, 1.0); // clamp to supported range
+		MinScale = FMath::Min(Scalar.ScaleFactor, MinScale);
+	}
+	TransformedScalars.Sort([](const FNegativeSpaceDirectionalToleranceScale& A, const FNegativeSpaceDirectionalToleranceScale& B)
+		{
+			return A.ScaleFactor < B.ScaleFactor;
+		});
+	if (MinScale < 1.0)
+	{
+		NegativeSpaceSettings.MinCustomMarginScale = MinScale;
+		NegativeSpaceSettings.CustomRadiusMarginScale =
+			[Scalars = MoveTemp(TransformedScalars)](const FDynamicMesh3& Mesh, int32 NearTID, FVector3d Pos)
+			{
+				FDistPoint3Triangle3d Query = TMeshQueries<FDynamicMesh3>::TriangleDistance(Mesh, NearTID, Pos);
+				FVector3d ToSurfaceDir = Query.ClosestTrianglePoint - Pos;
+				if (ToSurfaceDir.Normalize())
+				{
+					for (const FNegativeSpaceDirectionalToleranceScale& Scalar : Scalars)
+					{
+						if (Scalar.Direction.Dot(ToSurfaceDir) > Scalar.AngleRange)
+						{
+							return Scalar.ScaleFactor;
+						}
+					}
+				}
+				return 1.0;
+			};
+	}
+}
+
 // local helper to convert negative space settings from the blueprint-accessible struct to the geometrycore equivalent
 static FNegativeSpaceSampleSettings ConvertNegativeSpaceOptions(const FComputeNegativeSpaceOptions& NegativeSpaceOptions)
 {
@@ -273,6 +318,7 @@ static FNegativeSpaceSampleSettings ConvertNegativeSpaceOptions(const FComputeNe
 	NegativeSpaceSettings.bRequireSearchSampleCoverage = NegativeSpaceOptions.bRequireSearchSampleCoverage;
 	NegativeSpaceSettings.bOnlyConnectedToHull = NegativeSpaceOptions.bOnlyConnectedToHull;
 	NegativeSpaceSettings.MaxVoxelsPerDim = NegativeSpaceOptions.MaxVoxelsPerDim;
+	ApplyNegativeSpaceToleranceDirectionalScales(NegativeSpaceSettings, NegativeSpaceOptions.ToleranceDirectionalScales);
 	NegativeSpaceSettings.Sanitize();
 	return NegativeSpaceSettings;
 }
@@ -1451,6 +1497,7 @@ FGeometryScriptSimpleCollision UGeometryScriptLibrary_CollisionFunctions::Comput
 		NegativeSpaceSettings.ReduceRadiusMargin = NegativeSpaceOptions.Tolerance;
 		NegativeSpaceSettings.MinRadius = NegativeSpaceOptions.MinRadius;
 		NegativeSpaceSettings.MinSpacing = 0;
+		UELocal::ApplyNegativeSpaceToleranceDirectionalScales(NegativeSpaceSettings, NegativeSpaceOptions.ToleranceDirectionalScales);
 
 		if (!NegativeSpaceOptions.UnreachablePlanes.IsEmpty())
 		{
@@ -1476,7 +1523,7 @@ FGeometryScriptSimpleCollision UGeometryScriptLibrary_CollisionFunctions::Comput
 		constexpr int32 MaxAllowedSplits = 1000000; // more parts than any expected / reasonable decomposition
 		for (int32 Split = 0; ; Split++)
 		{
-			int32 NumSplit = ConvexDecomposition.SplitWorst(false, -1, true, NegativeSpaceSettings.ReduceRadiusMargin * .5);
+			int32 NumSplit = ConvexDecomposition.SplitWorst(false, -1, true, NegativeSpaceSettings.ReduceRadiusMargin * NegativeSpaceSettings.MinCustomMarginScale * .5);
 
 			if (NumSplit == 0)
 			{
diff --git a/Engine/Plugins/Runtime/GeometryScripting/Source/GeometryScriptingCore/Public/GeometryScript/CollisionFunctions.h b/Engine/Plugins/Runtime/GeometryScripting/Source/GeometryScriptingCore/Public/GeometryScript/CollisionFunctions.h
--- a/Engine/Plugins/Runtime/GeometryScripting/Source/GeometryScriptingCore/Public/GeometryScript/CollisionFunctions.h
+++ b/Engine/Plugins/Runtime/GeometryScripting/Source/GeometryScriptingCore/Public/GeometryScript/CollisionFunctions.h
@@ -128,12 +128,31 @@ enum class ENegativeSpaceSampleMethod : uint8
 	NavigableVoxelSearch
 };
 
+// Defines direction-specific custom scaling to apply to the FComputeNegativeSpaceOptions::NegativeSpaceTolerance value, based on the direction to the surface
+// For example, this can be used to more closely fit collision shapes on ground surfaces
+USTRUCT(BlueprintType)
+struct GEOMETRYSCRIPTINGCORE_API FNegativeSpaceDirectionalToleranceScale
+{
+	GENERATED_BODY()
+
+	// Direction to apply custom tolerance scale
+	UPROPERTY(EditAnywhere, BlueprintReadWrite, Category = Options)
+	FVector Direction = FVector(0,0,-1);
+
+	// Angle range around direction within which the scale will be applied, in degrees
+	UPROPERTY(EditAnywhere, BlueprintReadWrite, Category = Options, meta = (ClampMin = 0, ClampMax = 180, Units = deg))
+	double AngleRange = 45.0;
+
+	// Scale factor to apply to tolerance. Should be in the (0, 1] range
+	UPROPERTY(EditAnywhere, BlueprintReadWrite, Category = Options, meta = (ClampMin = .1, ClampMax = 1, UIMin = .25))
+	double ScaleFactor = 1.0;
+};
+
 // Options controlling how to sample the negative space of shapes, e.g. to define a region that must be avoided when merging collision shapes
 USTRUCT(BlueprintType)
 struct GEOMETRYSCRIPTINGCORE_API FComputeNegativeSpaceOptions
 {
 	GENERATED_BODY()
-public:
 
 	/** Method to use to find and sample negative space */
 	UPROPERTY(EditAnywhere, BlueprintReadWrite, Category = NegativeSpace)
@@ -163,6 +182,10 @@ public:
 	UPROPERTY(EditAnywhere, BlueprintReadWrite, Category = NegativeSpace, meta = (ClampMin = .01, UIMin = .1, Units = cm))
 	double NegativeSpaceTolerance = 2.0;
 
+	/** Optionally scale the negative space tolerances in custom directions */
+	UPROPERTY(EditAnywhere, BlueprintReadWrite, Category = NegativeSpace, AdvancedDisplay)
+	TArray<FNegativeSpaceDirectionalToleranceScale> ToleranceDirectionalScales;
+
 	/** Spheres smaller than this are not included in the negative space */
 	UPROPERTY(EditAnywhere, BlueprintReadWrite, Category = NegativeSpace, meta = (ClampMin = 0, Units = cm))
 	double MinRadius = 10.0;
@@ -199,6 +222,10 @@ public:
 	/// Use this for example to specify a ground plane, if a mesh will always be placed on ground and need not be navigable from below.
 	UPROPERTY(EditAnywhere, BlueprintReadWrite, Category = Options)
 	TArray<FPlane> UnreachablePlanes;
+
+	/** Optionally scale the tolerance in custom directions */
+	UPROPERTY(EditAnywhere, BlueprintReadWrite, Category = NegativeSpace, AdvancedDisplay)
+	TArray<FNegativeSpaceDirectionalToleranceScale> ToleranceDirectionalScales;
 };
 
 
diff --git a/Engine/Source/Runtime/GeometryCore/Private/CompGeom/ConvexDecomposition3.cpp b/Engine/Source/Runtime/GeometryCore/Private/CompGeom/ConvexDecomposition3.cpp
--- a/Engine/Source/Runtime/GeometryCore/Private/CompGeom/ConvexDecomposition3.cpp
+++ b/Engine/Source/Runtime/GeometryCore/Private/CompGeom/ConvexDecomposition3.cpp
@@ -300,7 +300,7 @@ bool FSphereCovering::AddNegativeSpace(const TFastWindingTree<FDynamicMesh3>& Sp
 		FAxisAlignedBox3d HullBox = HullAABB.GetBoundingBox();
 		HullBox.Expand(SampleSettings.VoxelExpandBoundsFactor * SampleSettings.GetAppliedScaleFactor());
 		FMarchingCubes MarchingCubes;
-		const double TargetCubeSize = SampleSettings.ReduceRadiusMargin * SampleSettings.MarchingCubesGridScale;
+		const double TargetCubeSize = SampleSettings.ReduceRadiusMargin * SampleSettings.MarchingCubesGridScale * SampleSettings.MinCustomMarginScale;
 		MarchingCubes.CubeSize = FMath::Clamp(TargetCubeSize, HullBox.MaxDim() / (double)SampleSettings.MaxVoxelsPerDim, HullBox.MinDim() * .5);
 		MarchingCubes.Bounds = HullBox;
 		MarchingCubes.RootMode = ERootfindingModes::Bisection;
@@ -346,7 +346,20 @@ bool FSphereCovering::AddNegativeSpace(const TFastWindingTree<FDynamicMesh3>& Sp
 				int32 NearTri = Spatial.GetTree()->FindNearestTriangle(Pt, NearDistSq, SampleSettings.ReduceRadiusMargin + SampleSettings.MinRadius);
 				if (NearTri != INDEX_NONE) // If we're closer than MinRadius + ReduceRadiusMargin to the surface, we can't place a MinRadius sphere here
 				{
-					return -1.0;
+					// re-test the distance w/ scaled margin, if applicable
+					if (SampleSettings.CustomRadiusMarginScale && SampleSettings.MinCustomMarginScale < 1.0)
+					{
+						double MarginScale = FMath::Clamp(SampleSettings.CustomRadiusMarginScale(*Spatial.GetTree()->GetMesh(), NearTri, Pt), SampleSettings.MinCustomMarginScale, 1.0);
+						double ScaledRadiusWithMargin = SampleSettings.MinRadius + SampleSettings.ReduceRadiusMargin * MarginScale;
+						if (NearDistSq < ScaledRadiusWithMargin * ScaledRadiusWithMargin)
+						{
+							return -1.0;
+						}
+					}
+					else
+					{
+						return -1.0;
+					}
 				}
 				if (!SampleSettings.bOnlyConnectedToHull && Spatial.FastWindingNumber(Pt) * WindingSign > .5) // If we're not requiring 'only connected to hull', then remove the internal spaces
 				{
@@ -385,7 +398,7 @@ bool FSphereCovering::AddNegativeSpace(const TFastWindingTree<FDynamicMesh3>& Sp
 		if (SampleSettings.bOnlyConnectedToHull)
 		{
 			TArray<FVector> Seeds;
-			double MinSeedSpacing = FMath::Max(MarchingCubes.CubeSize, SampleSettings.ReduceRadiusMargin);
+			double MinSeedSpacing = FMath::Max(MarchingCubes.CubeSize, SampleSettings.ReduceRadiusMargin * SampleSettings.MinCustomMarginScale);
 			double AddPtLenSq = 4 * MinSeedSpacing * MinSeedSpacing;
 			double SubDivLenSq = 16 * MinSeedSpacing * MinSeedSpacing;
 
@@ -463,7 +476,14 @@ bool FSphereCovering::AddNegativeSpace(const TFastWindingTree<FDynamicMesh3>& Sp
 
 			double NearDistSq;
 			int NearTID = Spatial.GetTree()->FindNearestTriangle(Pos, NearDistSq);
-			double R = FMath::Sqrt(NearDistSq) - SampleSettings.ReduceRadiusMargin;
+			double NearDist = FMath::Sqrt(NearDistSq);
+			double UseRadiusMargin = SampleSettings.ReduceRadiusMargin;
+			if (SampleSettings.CustomRadiusMarginScale && SampleSettings.MinCustomMarginScale < 1.0)
+			{
+				double MarginScale = FMath::Clamp(SampleSettings.CustomRadiusMarginScale(*Mesh, NearTID, Pos), SampleSettings.MinCustomMarginScale, 1.0);
+				UseRadiusMargin = SampleSettings.ReduceRadiusMargin * MarginScale;
+			}
+			double R = NearDist - UseRadiusMargin;
 
 			// Walk away from the closest point until we're far enough away to create a sample
 			// (give up if we haven't found a valid sample after a few steps, or if we stepped inside the shape)
@@ -485,7 +505,7 @@ bool FSphereCovering::AddNegativeSpace(const TFastWindingTree<FDynamicMesh3>& Sp
 					if (Spatial.FastWindingNumber(Pos) * WindingSign <= .5)
 					{
 						NearTID = Spatial.GetTree()->FindNearestTriangle(Pos, NearDistSq);
-						R = FMath::Sqrt(NearDistSq) - SampleSettings.ReduceRadiusMargin;
+						R = FMath::Sqrt(NearDistSq) - UseRadiusMargin;
 						if (R >= SampleSettings.MinRadius)
 						{
 							break;
diff --git a/Engine/Source/Runtime/GeometryCore/Public/CompGeom/ConvexDecomposition3.h b/Engine/Source/Runtime/GeometryCore/Public/CompGeom/ConvexDecomposition3.h
--- a/Engine/Source/Runtime/GeometryCore/Public/CompGeom/ConvexDecomposition3.h
+++ b/Engine/Source/Runtime/GeometryCore/Public/CompGeom/ConvexDecomposition3.h
@@ -110,6 +110,15 @@ struct FNegativeSpaceSampleSettings
 	// Note: Assumed to be in world space
 	TFunction<double(FVector Pos)> OptionalObstacleSDF;
 
+	// Optional function to define a bias (scalar multiplier) for the ReduceRadiusMargin value, to support scaling the tolerance based on location or direction to surfaces.
+	// The custom method takes the reference mesh, nearest triangle ID on that mesh, and the query position,
+	// and returns a scale factor between 0 and 1 to apply to the ReduceRadiusMargin
+	TFunction<double(const FDynamicMesh3& Mesh, int32 NearTID, FVector3d QueryPos)> CustomRadiusMarginScale;
+
+	// Margin scales should be between this value and 1.0. When the MinCustomMarginScale is 1.0, no custom scaling is applied.
+	// Should be greater than 0. Note that the sampling will be more expensive to compute for smaller minimum scale values.
+	double MinCustomMarginScale = 1.0;
+
 	// @return the scale factor that has been applied by Rescale()
 	double GetAppliedScaleFactor() const
 	{
@@ -140,6 +149,7 @@ struct FNegativeSpaceSampleSettings
 		ReduceRadiusMargin = FMath::Max(0.0, ReduceRadiusMargin);
 		MinRadius = FMath::Max(0.0, MinRadius);
 		MaxVoxelsPerDim = FMath::Clamp(MaxVoxelsPerDim, 4, 4096);
+		MinCustomMarginScale = FMath::Clamp(MinCustomMarginScale, .1, 1.0);
 	}
 
 	// helper to evaluate OptionalObstacleSDF for local positions
