Add more edge selection methods to geometry script

diff --git a/Engine/Plugins/Runtime/GeometryScripting/Source/GeometryScriptingCore/Private/MeshSelectionFunctions.cpp b/Engine/Plugins/Runtime/GeometryScripting/Source/GeometryScriptingCore/Private/MeshSelectionFunctions.cpp
--- a/Engine/Plugins/Runtime/GeometryScripting/Source/GeometryScriptingCore/Private/MeshSelectionFunctions.cpp
+++ b/Engine/Plugins/Runtime/GeometryScripting/Source/GeometryScriptingCore/Private/MeshSelectionFunctions.cpp
@@ -183,7 +183,32 @@ static void SelectMeshElementsWithContainmentTest(
 
 	SelectionOut.SetSelection(MoveTemp(GeoSelection));
 }
+
+static void SelectEdgesWithFilter(const FDynamicMesh3& ReadMesh, FGeometryScriptMeshSelection& SelectionOut, TFunctionRef<bool(int32)> EdgeFilter, bool bExcludeMeshBoundaryEdges = false)
+{
+	FGeometrySelection GeoSelection;
+	GeoSelection.InitializeTypes(EGeometryElementType::Edge, EGeometryTopologyType::Triangle);
+	for (int32 EID : ReadMesh.EdgeIndicesItr())
+	{
+		if (bExcludeMeshBoundaryEdges)
+		{
+			FIndex2i EdgeT = ReadMesh.GetEdgeT(EID);
+			if (EdgeT.B == INDEX_NONE)
+			{
+				continue;
+			}
+		}
+		if (EdgeFilter(EID))
+		{
+			ReadMesh.EnumerateTriEdgeIDsFromEdgeID(EID, [&GeoSelection](FMeshTriEdgeID TriEdgeID)
+			{
+				GeoSelection.Selection.Add(TriEdgeID.Encoded());
+			});
+		}
+	}
+	SelectionOut.SetSelection(MoveTemp(GeoSelection));
 }
+} // namespace UE::MeshSelectionLocals
 
 
 void UGeometryScriptLibrary_MeshSelectionFunctions::GetMeshSelectionInfo(
@@ -1010,34 +1035,57 @@ UDynamicMesh* UGeometryScriptLibrary_MeshSelectionFunctions::SelectMeshSharpEdge
 {
 	if (TargetMesh == nullptr)
 	{
-		UE_LOG(LogGeometry, Warning, TEXT("SelectMeshEdgesByEdgeAngle: TargetMesh is Null"));
+		UE_LOG(LogGeometry, Warning, TEXT("SelectMeshSharpEdges: TargetMesh is Null"));
 		return TargetMesh;
 	}
 
-	FGeometrySelection GeoSelection;
-	GeoSelection.InitializeTypes(EGeometryElementType::Edge, EGeometryTopologyType::Triangle);
 	double CosThresh = FMath::Cos(FMathd::DegToRad * MinAngleDeg);
-	TargetMesh->ProcessMesh([&GeoSelection, &CosThresh](const FDynamicMesh3& Mesh)
+	TargetMesh->ProcessMesh([&SelectionOut, CosThresh](const FDynamicMesh3& Mesh)
 	{
-		for (int32 EID : Mesh.EdgeIndicesItr())
+		UE::MeshSelectionLocals::SelectEdgesWithFilter(Mesh, SelectionOut, [CosThresh, &Mesh](int32 EID)
 		{
 			FIndex2i EdgeT = Mesh.GetEdgeT(EID);
 			if (EdgeT.B == INDEX_NONE)
 			{
-				continue;
+				return false;
 			}
 			FVector3d NormalA = Mesh.GetTriNormal(EdgeT.A);
 			FVector3d NormalB = Mesh.GetTriNormal(EdgeT.B);
-			if (NormalA.Dot(NormalB) <= CosThresh)
+			return NormalA.Dot(NormalB) <= CosThresh;
+		});
+	});
+	return TargetMesh;
+}
+
+UDynamicMesh* UGeometryScriptLibrary_MeshSelectionFunctions::SelectMeshSplitNormalEdges(
+	UDynamicMesh* TargetMesh,
+	FGeometryScriptMeshSelection& SelectionOut)
+{
+	if (TargetMesh == nullptr)
+	{
+		UE_LOG(LogGeometry, Warning, TEXT("SelectMeshSplitNormalEdges: TargetMesh is Null"));
+		return TargetMesh;
+	}
+
+	TargetMesh->ProcessMesh([&SelectionOut](const FDynamicMesh3& Mesh)
+	{
+		if (!Mesh.HasAttributes() || !Mesh.Attributes()->PrimaryNormals())
+		{
+			UE_LOG(LogGeometry, Warning, TEXT("SelectMeshSplitNormalEdges: TargetMesh does has no normals attribute"));
+			return;
+		}
+		const FDynamicMeshNormalOverlay* Normals = Mesh.Attributes()->PrimaryNormals();
+		UE::MeshSelectionLocals::SelectEdgesWithFilter(Mesh, SelectionOut, [&Mesh, Normals](int32 EID)
+		{
+			FIndex2i EdgeT = Mesh.GetEdgeT(EID);
+			// boundary edges are not split normal edges, but the overlay will consider them as seams, so handle them here
+			if (EdgeT.B == INDEX_NONE)
 			{
-				Mesh.EnumerateTriEdgeIDsFromEdgeID(EID, [&GeoSelection](FMeshTriEdgeID TriEdgeID)
-				{
-					GeoSelection.Selection.Add(TriEdgeID.Encoded());
-				});
+				return false;
 			}
-		}
+			return Normals->IsSeamEdge(EID);
+		});
 	});
-	SelectionOut.SetSelection(GeoSelection);
 	return TargetMesh;
 }
 
@@ -1047,7 +1095,7 @@ UDynamicMesh* UGeometryScriptLibrary_MeshSelectionFunctions::SelectMeshBoundaryE
 {
 	if (TargetMesh == nullptr)
 	{
-		UE_LOG(LogGeometry, Warning, TEXT("SelectMeshEdgesByEdgeAngle: TargetMesh is Null"));
+		UE_LOG(LogGeometry, Warning, TEXT("SelectMeshBoundaryEdges: TargetMesh is Null"));
 		return TargetMesh;
 	}
 
@@ -1064,7 +1112,7 @@ UDynamicMesh* UGeometryScriptLibrary_MeshSelectionFunctions::SelectMeshBoundaryE
 			}
 		}
 	});
-	SelectionOut.SetSelection(GeoSelection);
+	SelectionOut.SetSelection(MoveTemp(GeoSelection));
 	return TargetMesh;
 }
 
@@ -1077,7 +1125,7 @@ UDynamicMesh* UGeometryScriptLibrary_MeshSelectionFunctions::SelectSelectionBoun
 {
 	if (TargetMesh == nullptr)
 	{
-		UE_LOG(LogGeometry, Warning, TEXT("SelectMeshEdgesByEdgeAngle: TargetMesh is Null"));
+		UE_LOG(LogGeometry, Warning, TEXT("SelectSelectionBoundaryEdges: TargetMesh is Null"));
 		return TargetMesh;
 	}
 
@@ -1110,7 +1158,84 @@ UDynamicMesh* UGeometryScriptLibrary_MeshSelectionFunctions::SelectSelectionBoun
 			}
 		}
 	});
-	SelectionOut.SetSelection(GeoSelection);
+	SelectionOut.SetSelection(MoveTemp(GeoSelection));
+	return TargetMesh;
+}
+
+UDynamicMesh* UGeometryScriptLibrary_MeshSelectionFunctions::SelectMeshUVSeamEdges(
+	UDynamicMesh* TargetMesh, FGeometryScriptMeshSelection& SelectionOut,
+	int32 UVChannel, bool& bHaveValidUVs,
+	bool bExcludeMeshBoundaryEdges
+)
+{
+	bHaveValidUVs = false;
+	SelectionOut.ClearSelection();
+	if (TargetMesh == nullptr)
+	{
+		UE_LOG(LogGeometry, Warning, TEXT("SelectMeshUVSeamEdges: TargetMesh is Null"));
+		return TargetMesh;
+	}
+
+	TargetMesh->ProcessMesh([&SelectionOut, UVChannel, &bHaveValidUVs, bExcludeMeshBoundaryEdges](const FDynamicMesh3& Mesh)
+	{
+		if (!Mesh.HasAttributes() || !Mesh.Attributes()->GetUVLayer(UVChannel))
+		{
+			return;
+		}
+		bHaveValidUVs = true;
+		const FDynamicMeshUVOverlay* UVLayer = Mesh.Attributes()->GetUVLayer(UVChannel);
+		UE::MeshSelectionLocals::SelectEdgesWithFilter(Mesh, SelectionOut, [UVLayer](int32 EID)
+		{
+			return UVLayer->IsSeamEdge(EID);
+		}, bExcludeMeshBoundaryEdges);
+	});
+	return TargetMesh;
+}
+
+UDynamicMesh* UGeometryScriptLibrary_MeshSelectionFunctions::SelectMeshPolyGroupBoundaryEdges(
+	UDynamicMesh* TargetMesh, FGeometryScriptMeshSelection& SelectionOut,
+	FGeometryScriptGroupLayer GroupLayer,
+	bool bExcludeMeshBoundaryEdges
+)
+{
+	SelectionOut.ClearSelection();
+	if (TargetMesh == nullptr)
+	{
+		UE_LOG(LogGeometry, Warning, TEXT("SelectMeshPolyGroupBoundaryEdges: TargetMesh is Null"));
+		return TargetMesh;
+	}
+
+	TargetMesh->ProcessMesh([&SelectionOut, GroupLayer, bExcludeMeshBoundaryEdges](const FDynamicMesh3& Mesh)
+	{
+		if (GroupLayer.bDefaultLayer == true)
+		{
+			if (!Mesh.HasTriangleGroups())
+			{
+				UE_LOG(LogGeometry, Warning, TEXT("SelectMeshPolyGroupBoundaryEdges: Mesh does not have PolyGroups enabled"));
+				return;
+			}
+			UE::MeshSelectionLocals::SelectEdgesWithFilter(Mesh, SelectionOut, [&Mesh](int32 EID)
+			{
+				FIndex2i EdgeT = Mesh.GetEdgeT(EID);
+				return EdgeT.B == INDEX_NONE || Mesh.GetTriangleGroup(EdgeT.A) != Mesh.GetTriangleGroup(EdgeT.B);
+			}, bExcludeMeshBoundaryEdges);
+		}
+		else
+		{
+			if (!Mesh.HasAttributes() || Mesh.Attributes()->NumPolygroupLayers() <= GroupLayer.ExtendedLayerIndex)
+			{
+				UE_LOG(LogGeometry, Warning, TEXT("SelectMeshPolyGroupBoundaryEdges: Requested Polygroup Layer (%d) not found"), GroupLayer.ExtendedLayerIndex);
+				return;
+			}
+
+			const FDynamicMeshPolygroupAttribute* PolygroupAttr = Mesh.Attributes()->GetPolygroupLayer(GroupLayer.ExtendedLayerIndex);
+			UE::MeshSelectionLocals::SelectEdgesWithFilter(Mesh, SelectionOut, [&Mesh, PolygroupAttr](int32 EID)
+			{
+				FIndex2i EdgeT = Mesh.GetEdgeT(EID);
+				return EdgeT.B == INDEX_NONE || PolygroupAttr->GetValue(EdgeT.A) != PolygroupAttr->GetValue(EdgeT.B);
+			}, bExcludeMeshBoundaryEdges);
+		}
+	});
 	return TargetMesh;
 }
 
diff --git a/Engine/Plugins/Runtime/GeometryScripting/Source/GeometryScriptingCore/Public/GeometryScript/MeshSelectionFunctions.h b/Engine/Plugins/Runtime/GeometryScripting/Source/GeometryScriptingCore/Public/GeometryScript/MeshSelectionFunctions.h
--- a/Engine/Plugins/Runtime/GeometryScripting/Source/GeometryScriptingCore/Public/GeometryScript/MeshSelectionFunctions.h
+++ b/Engine/Plugins/Runtime/GeometryScripting/Source/GeometryScriptingCore/Public/GeometryScript/MeshSelectionFunctions.h
@@ -274,6 +274,16 @@ public:
 		double MinAngleDeg = 20.0
 	);
 
+	/**
+	 * Create a new Selection, for the TargetMesh, of all edges seams in the mesh normal topology -- i.e., where normals can be different across the edge
+	 */
+	UFUNCTION(BlueprintCallable, Category = "GeometryScript|MeshSelection", meta = (ScriptMethod))
+	static UPARAM(DisplayName = "Target Mesh") UDynamicMesh*
+	SelectMeshSplitNormalEdges(
+		UDynamicMesh* TargetMesh,
+		FGeometryScriptMeshSelection& Selection
+	);
+
 	/**
 	 * Create a new Selection, for the TargetMesh, of all mesh boundary edges
 	 */
@@ -297,6 +307,34 @@ public:
 		bool bExcludeMeshBoundaryEdges = false
 	);
 
+	
+	/**
+	 * Create a new Mesh Selection of UV seam edges for the TargetMesh
+	 * @param UVChannel Seams will be found with respect to this UV Channel
+	 * @param bHaveValidUVs If false, the requested UV Channel was not present on the TargetMesh
+	 * @param bExcludeMeshBoundaryEdges If true, do not include mesh boundary edges in the selection; otherwise, they will be included.
+	 */
+	UFUNCTION(BlueprintCallable, Category = "GeometryScript|MeshSelection", meta = (ScriptMethod, DisplayName = "Select Mesh UV Seam Edges"))
+	static UPARAM(DisplayName = "Target Mesh") UDynamicMesh*
+	SelectMeshUVSeamEdges(
+		UDynamicMesh* TargetMesh, FGeometryScriptMeshSelection& Selection,
+		UPARAM(DisplayName = "UV Channel") int32 UVChannel, bool& bHaveValidUVs,
+		bool bExcludeMeshBoundaryEdges = false
+	);
+	
+	/**
+	 * Create a new Mesh Selection of PolyGroup boundary edges for the TargetMesh
+	 * @param GroupLayer PolyGroup layer to reference
+	 * @param bExcludeMeshBoundaryEdges If true, do not include mesh boundary edges in the selection; otherwise, they will be included.
+	 */
+	UFUNCTION(BlueprintCallable, Category = "GeometryScript|MeshSelection", meta = (ScriptMethod, DisplayName = "Select Mesh PolyGroup Boundary Edges"))
+	static UPARAM(DisplayName = "Target Mesh") UDynamicMesh*
+	SelectMeshPolyGroupBoundaryEdges(
+		UDynamicMesh* TargetMesh, FGeometryScriptMeshSelection& Selection,
+		FGeometryScriptGroupLayer GroupLayer,
+		bool bExcludeMeshBoundaryEdges = false
+	);
+
 
 	/**
 	* Create a new Mesh Selection of the SelectionType for the TargetMesh by finding all elements inside a second SelectionMesh
