CIS fix for text macro namespace conflict with the  ComputeUniformRandomPointSampling key

diff --git a/Engine/Plugins/Runtime/GeometryScripting/Source/GeometryScriptingCore/Private/MeshSamplingFunctions.cpp b/Engine/Plugins/Runtime/GeometryScripting/Source/GeometryScriptingCore/Private/MeshSamplingFunctions.cpp
--- a/Engine/Plugins/Runtime/GeometryScripting/Source/GeometryScriptingCore/Private/MeshSamplingFunctions.cpp
+++ b/Engine/Plugins/Runtime/GeometryScripting/Source/GeometryScriptingCore/Private/MeshSamplingFunctions.cpp
@@ -27,7 +27,7 @@ const UDynamicMesh* UGeometryScriptLibrary_MeshSamplingFunctions::ComputeUniform
 {
 	if (TargetMesh == nullptr)
 	{
-		UE::Geometry::AppendError(Debug, EGeometryScriptErrorType::InvalidInputs, LOCTEXT("ComputeUniformRandomPointSampling", "ComputePointSampling: TargetMesh is Null"));
+		UE::Geometry::AppendError(Debug, EGeometryScriptErrorType::InvalidInputs, LOCTEXT("ComputeUniformRandomPointSampling_InvalidInput", "ComputeUniformRandomPointSampling: TargetMesh is Null"));
 		return TargetMesh;
 	}
 
@@ -41,7 +41,7 @@ const UDynamicMesh* UGeometryScriptLibrary_MeshSamplingFunctions::ComputeUniform
 		if (!MeshUniformSampling.IsValid())
 		{
 			// cannot sample mesh (e.g., if there are no triangles)
-			UE::Geometry::AppendWarning(Debug, EGeometryScriptErrorType::InvalidInputs, LOCTEXT("ComputeUniformRandomPointSampling", "ComputeUniformRandomPointSampling: TargetMesh does not have triangles to sample."));
+			UE::Geometry::AppendWarning(Debug, EGeometryScriptErrorType::InvalidInputs, LOCTEXT("ComputeUniformRandomPointSampling_EmptyMesh", "ComputeUniformRandomPointSampling: TargetMesh does not have triangles to sample."));
 			return;
 		}
 		FRandomStream RandomStream(RandomSeed);
