in append bounding box, don't append if the bounding box is not valid

diff --git a/Engine/Plugins/Runtime/GeometryScripting/Source/GeometryScriptingCore/Private/MeshPrimitiveFunctions.cpp b/Engine/Plugins/Runtime/GeometryScripting/Source/GeometryScriptingCore/Private/MeshPrimitiveFunctions.cpp
--- a/Engine/Plugins/Runtime/GeometryScripting/Source/GeometryScriptingCore/Private/MeshPrimitiveFunctions.cpp
+++ b/Engine/Plugins/Runtime/GeometryScripting/Source/GeometryScriptingCore/Private/MeshPrimitiveFunctions.cpp
@@ -275,6 +275,12 @@ namespace UE::Private::AppendMeshHelper
 			return TargetMesh;
 		}
 
+		if (!Box.IsValid)
+		{
+			UE::Geometry::AppendError(Debug, EGeometryScriptErrorType::InvalidInputs, LOCTEXT("PrimitiveFunctions_AppendBoundingBox_InvalidBox", "AppendBoundingBox: Bounding Box is empty, and cannot be appended."));
+			return TargetMesh;
+		}
+
 		// todo: if Steps X/Y/Z are zero, can use trivial box generator
 
 		FGridBoxMeshGenerator GridBoxGenerator;
