Add additional mesh polygroup query functions for polygroup bounds, uv bounds, and uv centroid

diff --git a/Engine/Plugins/Runtime/GeometryScripting/Source/GeometryScriptingCore/Private/MeshPolygroupFunctions.cpp b/Engine/Plugins/Runtime/GeometryScripting/Source/GeometryScriptingCore/Private/MeshPolygroupFunctions.cpp
--- a/Engine/Plugins/Runtime/GeometryScripting/Source/GeometryScriptingCore/Private/MeshPolygroupFunctions.cpp
+++ b/Engine/Plugins/Runtime/GeometryScripting/Source/GeometryScriptingCore/Private/MeshPolygroupFunctions.cpp
@@ -38,6 +38,47 @@ ReturnType SimpleMeshPolygroupQuery(UDynamicMesh* Mesh, FGeometryScriptGroupLaye
 	return RetVal;
 }
 
+// Version of SimpleMeshPolygroupQuery that avoids initial full mesh pass to initialize FPolygroupSet, instead passing a simple lambda to access triangle polygroups
+template<typename ReturnType>
+ReturnType SimpleMeshPolygroupQuery(const UDynamicMesh* Mesh, FGeometryScriptGroupLayer ScriptGroupLayer, bool& bHasGroups, ReturnType DefaultValue,
+	TFunctionRef<ReturnType(const FDynamicMesh3& Mesh, TFunctionRef<int32(int32)> GetPolygroup)> QueryFunc)
+{
+	bHasGroups = false;
+	ReturnType RetVal = DefaultValue;
+	if (Mesh)
+	{
+		Mesh->ProcessMesh([ScriptGroupLayer, &bHasGroups, &RetVal, QueryFunc](const FDynamicMesh3& ReadMesh)
+		{
+			if (ScriptGroupLayer.bDefaultLayer)
+			{
+				if (ReadMesh.HasTriangleGroups())
+				{
+					bHasGroups = true;
+					RetVal = QueryFunc(ReadMesh, [&ReadMesh](int32 TID) -> int32 
+					{
+						return ReadMesh.GetTriangleGroup(TID); 
+					});
+				}
+			}
+			else
+			{
+				if (ReadMesh.HasAttributes())
+				{
+					if (const FDynamicMeshPolygroupAttribute* GroupLayer = ReadMesh.Attributes()->GetPolygroupLayer(ScriptGroupLayer.ExtendedLayerIndex))
+					{
+						bHasGroups = true;
+						RetVal = QueryFunc(ReadMesh, [GroupLayer](int32 TID) -> int32 
+						{
+							return GroupLayer->GetValue(TID); 
+						});
+					}
+				}
+			}
+		});
+	}
+	return RetVal;
+}
+
 
 
 
@@ -443,6 +484,126 @@ UDynamicMesh* UGeometryScriptLibrary_MeshPolygroupFunctions::GetPolygroupIDsInMe
 }
 
 
+const UDynamicMesh* UGeometryScriptLibrary_MeshPolygroupFunctions::GetPolyGroupBoundingBox(
+	const UDynamicMesh* TargetMesh,
+	FGeometryScriptGroupLayer GroupLayer,
+	int32 GroupID,
+	FBox& Bounds)
+{
+	Bounds.Init();
+	bool bIsValidPolygroupLayer = false;
+	SimpleMeshPolygroupQuery<bool>(TargetMesh, GroupLayer, bIsValidPolygroupLayer, 0,
+		[GroupID, &Bounds](const FDynamicMesh3& Mesh, TFunctionRef<int32(int32)> GetGroup) -> bool
+		{
+			for (int32 TID : Mesh.TriangleIndicesItr())
+			{
+				if (GroupID == GetGroup(TID))
+				{
+					FVector3d A, B, C;
+					Mesh.GetTriVertices(TID, A, B, C);
+					Bounds += A;
+					Bounds += B;
+					Bounds += C;
+				}
+			}
+			return (bool)Bounds.IsValid;
+		});
+	return TargetMesh;
+}
+
+const UDynamicMesh* UGeometryScriptLibrary_MeshPolygroupFunctions::GetPolyGroupUVBoundingBox(
+	const UDynamicMesh* TargetMesh,
+	FGeometryScriptGroupLayer GroupLayer,
+	int32 GroupID,
+	int32 UVChannel,
+	FBox2D& Bounds)
+{
+	Bounds.Init();
+	bool bIsValidPolygroupLayer = false;
+	SimpleMeshPolygroupQuery<bool>(TargetMesh, GroupLayer, bIsValidPolygroupLayer, 0,
+		[GroupID, UVChannel, &Bounds](const FDynamicMesh3& Mesh, TFunctionRef<int32(int32)> GetGroup) -> bool
+		{
+			if (!Mesh.HasAttributes())
+			{
+				return false;
+			}
+			const FDynamicMeshUVOverlay* UVLayer = Mesh.Attributes()->GetUVLayer(UVChannel);
+			if (!UVLayer)
+			{
+				return false;
+			}
+
+			for (int32 TID : Mesh.TriangleIndicesItr())
+			{
+				if (GroupID == GetGroup(TID))
+				{
+					FIndex3i Els;
+					if (UVLayer->GetTriangleIfValid(TID, Els))
+					{
+						for (int32 SubIdx = 0; SubIdx < 3; ++SubIdx)
+						{
+							Bounds += (FVector2D)UVLayer->GetElement(Els[SubIdx]);
+						}
+					}
+				}
+			}
+			return Bounds.bIsValid;
+		});
+	return TargetMesh;
+}
+
+const UDynamicMesh* UGeometryScriptLibrary_MeshPolygroupFunctions::GetPolyGroupUVCentroid(
+	const UDynamicMesh* TargetMesh,
+	FGeometryScriptGroupLayer GroupLayer,
+	int32 GroupID,
+	int32 UVChannel,
+	FVector2D& Centroid,
+	bool& bIsValid)
+{
+	Centroid = FVector2D::Zero();
+	bool bIsValidPolygroupLayer = false;
+	bIsValid = SimpleMeshPolygroupQuery<bool>(TargetMesh, GroupLayer, bIsValidPolygroupLayer, 0,
+		[GroupID, UVChannel, &Centroid](const FDynamicMesh3& Mesh, TFunctionRef<int32(int32)> GetGroup) -> bool
+		{
+			if (!Mesh.HasAttributes())
+			{
+				return false;
+			}
+			const FDynamicMeshUVOverlay* UVLayer = Mesh.Attributes()->GetUVLayer(UVChannel);
+			if (!UVLayer)
+			{
+				return false;
+			}
+
+			double WtSum = 0.0;
+			FVector2D Sum = FVector2D::Zero();
+			for (int32 TID : Mesh.TriangleIndicesItr())
+			{
+				if (GroupID == GetGroup(TID))
+				{
+					FIndex3i Els;
+					if (UVLayer->GetTriangleIfValid(TID, Els))
+					{
+						FVector2f A = UVLayer->GetElement(Els.A), B = UVLayer->GetElement(Els.B), C = UVLayer->GetElement(Els.C);
+						FVector2D TriCenter = FVector2D(A + B + C) * (1.0 / 3.0);
+						double Wt = (double)VectorUtil::Area(A, B, C);
+						Sum += TriCenter * Wt;
+						WtSum += Wt;
+					}
+				}
+			}
+			if (WtSum > 0)
+			{
+				Centroid = Sum / WtSum;
+				return true;
+			}
+			else
+			{
+				return false;
+			}
+		});
+	return TargetMesh;
+}
 
 UDynamicMesh* UGeometryScriptLibrary_MeshPolygroupFunctions::GetTrianglesInPolygroup(
 	UDynamicMesh* TargetMesh,
diff --git a/Engine/Plugins/Runtime/GeometryScripting/Source/GeometryScriptingCore/Public/GeometryScript/MeshPolygroupFunctions.h b/Engine/Plugins/Runtime/GeometryScripting/Source/GeometryScriptingCore/Public/GeometryScript/MeshPolygroupFunctions.h
--- a/Engine/Plugins/Runtime/GeometryScripting/Source/GeometryScriptingCore/Public/GeometryScript/MeshPolygroupFunctions.h
+++ b/Engine/Plugins/Runtime/GeometryScripting/Source/GeometryScriptingCore/Public/GeometryScript/MeshPolygroupFunctions.h
@@ -165,6 +165,54 @@ public:
 		FGeometryScriptGroupLayer GroupLayer, 
 		UPARAM(ref, DisplayName="PolyGroup IDs Out") FGeometryScriptIndexList& PolygroupIDsOut );
 
+	/**
+	 * Compute the bounds of a PolyGroup
+	 * 
+	 * @param GroupLayer The PolyGroup layer to query
+	 * @param GroupID The PolyGroup to query
+	 * @param Bounds The bounds of the PolyGroups's triangles. Can be invalid if the PolyGroup was not found.
+	 */
+	UFUNCTION(BlueprintCallable, Category = "GeometryScript|PolyGroups", meta = (ScriptMethod, DisplayName = "Get PolyGroup Bounding Box"))
+	static UPARAM(DisplayName = "Target Mesh") const UDynamicMesh*
+	GetPolyGroupBoundingBox(const UDynamicMesh* TargetMesh,
+		FGeometryScriptGroupLayer GroupLayer,
+		int32 GroupID,
+		FBox& Bounds);
+
+	/**
+	 * Compute the UV bounds of a PolyGroup
+	 * 
+	 * @param GroupLayer The PolyGroup layer to query
+	 * @param GroupID The PolyGroup to query
+	 * @param UVChannel The UV channel to query
+	 * @param UVBounds The bounds of the PolyGroups's UV triangles. Can be invalid if the PolyGroup was not found or did not have UVs in the requested channel.
+	 */
+	UFUNCTION(BlueprintCallable, Category = "GeometryScript|PolyGroups", meta = (ScriptMethod, DisplayName = "Get PolyGroup UV Bounding Box"))
+	static UPARAM(DisplayName = "Target Mesh") const UDynamicMesh*
+	GetPolyGroupUVBoundingBox(const UDynamicMesh* TargetMesh,
+		FGeometryScriptGroupLayer GroupLayer,
+		int32 GroupID,
+		UPARAM(DisplayName = "UV Channel") int32 UVChannel,
+		UPARAM(DisplayName = "UV Bounds") FBox2D& UVBounds);
+
+	/**
+	 * Compute the UV centroid of a PolyGroup
+	 * 
+	 * @param GroupLayer The PolyGroup layer to query
+	 * @param GroupID The PolyGroup to query
+	 * @param UVChannel The UV channel to query
+	 * @param Centroid The centroid of the polygroup's UV triangles
+	 * @param bIsValid True if a valid centroid was found, false otherwise. Can be false if the PolyGroup was not found or did not have UVs in the requested channel.
+	 */
+	UFUNCTION(BlueprintCallable, Category = "GeometryScript|PolyGroups", meta = (ScriptMethod, DisplayName = "Get PolyGroup UV Centroid"))
+	static UPARAM(DisplayName = "Target Mesh") const UDynamicMesh*
+	GetPolyGroupUVCentroid(const UDynamicMesh* TargetMesh,
+		FGeometryScriptGroupLayer GroupLayer,
+		int32 GroupID,
+		UPARAM(DisplayName = "UV Channel") int32 UVChannel,
+		FVector2D& Centroid,
+		bool& bIsValid);
+
 	/**
 	 * Create list of all triangles with the given PolyGroup ID in the given GroupLayer (not necessarily a single connected-component)
 	 */
