Geoscripting: copy morph target normals if requested

diff --git a/Engine/Plugins/Runtime/GeometryScripting/Source/GeometryScriptingCore/Private/MeshAssetFunctions.cpp b/Engine/Plugins/Runtime/GeometryScripting/Source/GeometryScriptingCore/Private/MeshAssetFunctions.cpp
--- a/Engine/Plugins/Runtime/GeometryScripting/Source/GeometryScriptingCore/Private/MeshAssetFunctions.cpp
+++ b/Engine/Plugins/Runtime/GeometryScripting/Source/GeometryScriptingCore/Private/MeshAssetFunctions.cpp
@@ -10,11 +10,9 @@
 #include "Engine/StaticMesh.h"
 #include "Engine/SkinnedAssetCommon.h"
 #include "Engine/SkeletalMesh.h"
-#include "Engine/StaticMeshActor.h"
 #include "StaticMeshAttributes.h"
 #include "StaticMeshOperations.h"
 #include "RenderingThread.h"
-#include "PhysicsEngine/BodySetup.h"
 #include "Rendering/SkeletalMeshModel.h"
 #include "Rendering/SkeletalMeshRenderData.h"
 #include "SkeletalMeshAttributes.h"
@@ -23,12 +21,12 @@
 #include "MeshDescriptionToDynamicMesh.h"
 #include "DynamicMeshToMeshDescription.h"
 #include "SkeletalMeshOperations.h"
-#include "StaticMeshLODResourcesToDynamicMesh.h"
 #include "AssetUtils/StaticMeshMaterialUtil.h"
 #include "ConversionUtils/SceneComponentToDynamicMesh.h"
 #include "DynamicMesh/DynamicBoneAttribute.h"
 #include "DynamicMesh/DynamicVertexSkinWeightsAttribute.h"
 #include "DynamicMesh/NonManifoldMappingSupport.h"
+#include "DynamicMesh/MeshNormals.h"
 
 #include UE_INLINE_GENERATED_CPP_BY_NAME(MeshAssetFunctions)
 
@@ -1150,6 +1148,14 @@ UDynamicMesh* UGeometryScriptLibrary_StaticMeshFunctions::CopyMorphTargetToSkele
 		UE::Geometry::AppendError(Debug, EGeometryScriptErrorType::InvalidInputs, LOCTEXT("CopyMorphTargetToSkeletalMesh_InvalidMorphTargetGeometry", "CopyMorphTargetToSkeletalMesh: Morph target mesh doesnt have the same number of vertices as the skeletal mesh."));
 		return FromDynamicMesh;
 	}
+
+	FMeshNormals Normals;
+	const bool bSourceHasNormals = SourceMesh.HasVertexNormals();
+	if (Options.bCopyNormals && !bSourceHasNormals)
+	{
+		Normals = FMeshNormals(&SourceMesh);
+		Normals.ComputeVertexNormals();
+	}
 	
 	FSkeletalMeshAttributes MeshAttributes(*MeshDescription);
 	MeshAttributes.Register();
@@ -1168,39 +1174,55 @@ UDynamicMesh* UGeometryScriptLibrary_StaticMeshFunctions::CopyMorphTargetToSkele
 			return FromDynamicMesh;
 		}
 	}
-	else
+
+	// call RegisterMorphTargetAttribute to make sure normals are registered / unregistered when needed
+	if (!MeshAttributes.RegisterMorphTargetAttribute(MorphTargetName, Options.bCopyNormals))
 	{
-		if (!MeshAttributes.RegisterMorphTargetAttribute(MorphTargetName, false))
+		if (Transaction)
 		{
-			if (Transaction)
-			{
-				Transaction->Cancel();
-			}
-			UE::Geometry::AppendError(Debug, EGeometryScriptErrorType::InvalidInputs, LOCTEXT("CopyMorphTargetToSkeletalMesh_InvalidMorphTargetName2", "CopyMorphTargetToSkeletalMesh: Morph target name is invalid."));
-			return FromDynamicMesh;
+			Transaction->Cancel();
 		}
+		UE::Geometry::AppendError(Debug, EGeometryScriptErrorType::InvalidInputs, LOCTEXT("CopyMorphTargetToSkeletalMesh_InvalidMorphTargetName2", "CopyMorphTargetToSkeletalMesh: Morph target name is invalid."));
+		return FromDynamicMesh;
 	}
 
 	const FSkeletalMeshLODInfo* LODInfo = ToSkeletalMeshAsset->GetLODInfo(TargetLOD.LODIndex);
 	const float MorphThresholdSquared = LODInfo->BuildSettings.MorphThresholdPosition * LODInfo->BuildSettings.MorphThresholdPosition;
 	
-	TVertexAttributesRef<FVector3f> PositionDelta = MeshAttributes.GetVertexMorphPositionDelta(MorphTargetName);
-	TVertexAttributesRef<FVector3f> VertexPositions = MeshAttributes.GetVertexPositions();
+	TVertexAttributesRef<FVector3f> MorphPositionDelta = MeshAttributes.GetVertexMorphPositionDelta(MorphTargetName);
+	TVertexAttributesConstRef<FVector3f> VertexPositions = MeshAttributes.GetVertexPositions();
 
+	TVertexInstanceAttributesRef<FVector3f> MorphNormalDelta = MeshAttributes.GetVertexInstanceMorphNormalDelta(MorphTargetName);
+	TVertexInstanceAttributesConstRef<FVector3f> VertexNormals = MeshAttributes.GetVertexInstanceNormals();
+	
 	bool bMorphTargetIsEmpty = true;
 	
 	for (int32 SourceVID = 0; SourceVID < SourceMesh.VertexCount(); ++SourceVID)
 	{
 		const int32 TargetVID = NonManifoldMappingSupport.GetOriginalNonManifoldVertexID(SourceVID);
 		
-		const FVector3d V0 = FromDynamicMesh->GetMeshRef().GetVertex(SourceVID);
+		const FVector3d V0 = SourceMesh.GetVertex(SourceVID);
 		const FVector3f V1 = VertexPositions[TargetVID];
 
 		const FVector3f Delta = FVector3f{V0} - V1;
 		if (Delta.SquaredLength() > MorphThresholdSquared)
 		{
 			bMorphTargetIsEmpty = false;
-			PositionDelta.Set(TargetVID, Delta);
+			MorphPositionDelta.Set(TargetVID, Delta);
+
+			if (Options.bCopyNormals)
+			{
+				const FVector3f N0 = bSourceHasNormals ? SourceMesh.GetVertexNormal(SourceVID) : FVector3f{Normals[SourceVID]};
+				for (const FVertexInstanceID VertexInstanceID : MeshDescription->GetVertexVertexInstanceIDs(TargetVID))
+				{
+					const FVector3f N1 = VertexNormals.Get(VertexInstanceID);
+					const FVector3f NDelta(N0-N1);
+					if (Delta.SquaredLength() > MorphThresholdSquared)
+					{
+						MorphNormalDelta.Set(VertexInstanceID, NDelta);
+					}
+				}
+			}
 		}
 	}
 
diff --git a/Engine/Plugins/Runtime/GeometryScripting/Source/GeometryScriptingCore/Public/GeometryScript/MeshAssetFunctions.h b/Engine/Plugins/Runtime/GeometryScripting/Source/GeometryScriptingCore/Public/GeometryScript/MeshAssetFunctions.h
--- a/Engine/Plugins/Runtime/GeometryScripting/Source/GeometryScriptingCore/Public/GeometryScript/MeshAssetFunctions.h
+++ b/Engine/Plugins/Runtime/GeometryScripting/Source/GeometryScriptingCore/Public/GeometryScript/MeshAssetFunctions.h
@@ -133,6 +133,9 @@ struct GEOMETRYSCRIPTINGCORE_API FGeometryScriptCopyMorphTargetToAssetOptions
 
 	UPROPERTY(EditAnywhere, BlueprintReadWrite, Category = Options)
 	bool bDeferMeshPostEditChange = false;
+
+	UPROPERTY(EditAnywhere, BlueprintReadWrite, Category = Options)
+	bool bCopyNormals = false;
 };
 
 
diff --git a/Engine/Source/Runtime/SkeletalMeshDescription/Private/SkeletalMeshAttributes.cpp b/Engine/Source/Runtime/SkeletalMeshDescription/Private/SkeletalMeshAttributes.cpp
--- a/Engine/Source/Runtime/SkeletalMeshDescription/Private/SkeletalMeshAttributes.cpp
+++ b/Engine/Source/Runtime/SkeletalMeshDescription/Private/SkeletalMeshAttributes.cpp
@@ -179,14 +179,30 @@ bool FSkeletalMeshAttributes::RegisterMorphTargetAttribute(
 		return false;
 	}
 
-	bool bSuccess = MeshDescription.VertexAttributes().RegisterAttribute<FVector3f>(AttributeName, 1, FVector3f::ZeroVector, EMeshAttributeFlags::None).IsValid();
+	// Register position attribute if it doesn't already exist.
+	TAttributesSet<FVertexID>& VertexAttributes = MeshDescription.VertexAttributes();
+	bool bSuccess = VertexAttributes.HasAttribute(AttributeName); 
+	if (!bSuccess)
+	{
+		bSuccess = VertexAttributes.RegisterAttribute<FVector3f>(AttributeName, 1, FVector3f::ZeroVector, EMeshAttributeFlags::None).IsValid();		
+	}
+
+	TAttributesSet<FVertexInstanceID>& VertexInstanceAttributes = MeshDescription.VertexInstanceAttributes();
 	if (bSuccess && bIncludeNormals)
 	{
 		// Register normal attribute, if requested, if it fails, then we have to unregister the vertex attribute too.
-		bSuccess = MeshDescription.VertexInstanceAttributes().RegisterAttribute<FVector3f>(AttributeName, 1, FVector3f::ZeroVector, EMeshAttributeFlags::None).IsValid();
+		bSuccess = VertexInstanceAttributes.RegisterAttribute<FVector3f>(AttributeName, 1, FVector3f::ZeroVector, EMeshAttributeFlags::None).IsValid();
 		if (!bSuccess)
 		{
-			MeshDescription.VertexAttributes().UnregisterAttribute(AttributeName);
+			VertexAttributes.UnregisterAttribute(AttributeName);
+		}
+	}
+	else
+	{
+		// Unregister normal attribute if it fails of if it's not needed anymore.
+		if (VertexInstanceAttributes.HasAttribute(AttributeName))
+		{
+			VertexInstanceAttributes.UnregisterAttribute(AttributeName);
 		}
 	}
 
