Geometry Scripting: Fix normal copy to use the normal overlay since that's what the geo scripting nodes normal nodes operate on.

diff --git a/Engine/Plugins/Runtime/GeometryScripting/Source/GeometryScriptingCore/Private/MeshAssetFunctions.cpp b/Engine/Plugins/Runtime/GeometryScripting/Source/GeometryScriptingCore/Private/MeshAssetFunctions.cpp
--- a/Engine/Plugins/Runtime/GeometryScripting/Source/GeometryScriptingCore/Private/MeshAssetFunctions.cpp
+++ b/Engine/Plugins/Runtime/GeometryScripting/Source/GeometryScriptingCore/Private/MeshAssetFunctions.cpp
@@ -1149,9 +1149,9 @@ UDynamicMesh* UGeometryScriptLibrary_StaticMeshFunctions::CopyMorphTargetToSkele
 		return FromDynamicMesh;
 	}
 
+	const FDynamicMeshNormalOverlay* SourceNormals = SourceMesh.Attributes()->PrimaryNormals();
 	FMeshNormals Normals;
-	const bool bSourceHasNormals = SourceMesh.HasVertexNormals();
-	if (Options.bCopyNormals && !bSourceHasNormals)
+	if (Options.bCopyNormals && !SourceNormals)
 	{
 		Normals = FMeshNormals(&SourceMesh);
 		Normals.ComputeVertexNormals();
@@ -1212,15 +1212,12 @@ UDynamicMesh* UGeometryScriptLibrary_StaticMeshFunctions::CopyMorphTargetToSkele
 
 			if (Options.bCopyNormals)
 			{
-				const FVector3f N0 = bSourceHasNormals ? SourceMesh.GetVertexNormal(SourceVID) : FVector3f{Normals[SourceVID]};
+				const FVector3f N0 = SourceNormals ? SourceNormals->GetElement(SourceVID) : FVector3f{Normals[SourceVID]};
 				for (const FVertexInstanceID VertexInstanceID : MeshDescription->GetVertexVertexInstanceIDs(TargetVID))
 				{
 					const FVector3f N1 = VertexNormals.Get(VertexInstanceID);
 					const FVector3f NDelta(N0-N1);
-					if (Delta.SquaredLength() > MorphThresholdSquared)
-					{
-						MorphNormalDelta.Set(VertexInstanceID, NDelta);
-					}
+					MorphNormalDelta.Set(VertexInstanceID, NDelta);
 				}
 			}
 		}
