Geometry scripting: Preserve existing mesh description attributes when registering new ones for the morph/skinweight copy nodes.

diff --git a/Engine/Plugins/Runtime/GeometryScripting/Source/GeometryScriptingCore/Private/MeshAssetFunctions.cpp b/Engine/Plugins/Runtime/GeometryScripting/Source/GeometryScriptingCore/Private/MeshAssetFunctions.cpp
--- a/Engine/Plugins/Runtime/GeometryScripting/Source/GeometryScriptingCore/Private/MeshAssetFunctions.cpp
+++ b/Engine/Plugins/Runtime/GeometryScripting/Source/GeometryScriptingCore/Private/MeshAssetFunctions.cpp
@@ -1158,7 +1158,8 @@ UDynamicMesh* UGeometryScriptLibrary_StaticMeshFunctions::CopyMorphTargetToSkele
 	}
 	
 	FSkeletalMeshAttributes MeshAttributes(*MeshDescription);
-	MeshAttributes.Register();
+	constexpr bool bKeepExistingAttributes = true;
+	MeshAttributes.Register(bKeepExistingAttributes);
 
 	ToSkeletalMeshAsset->ModifyMeshDescription(TargetLOD.LODIndex);
 
@@ -1173,6 +1174,9 @@ UDynamicMesh* UGeometryScriptLibrary_StaticMeshFunctions::CopyMorphTargetToSkele
 			UE::Geometry::AppendError(Debug, EGeometryScriptErrorType::InvalidInputs, LOCTEXT("CopyMorphTargetToSkeletalMesh_InvalidMorphTargetName1", "CopyMorphTargetToSkeletalMesh: Morph target name already exists"));
 			return FromDynamicMesh;
 		}
+
+		// Remove existing attribute so that we start with a clean slate.
+		MeshAttributes.UnregisterMorphTargetAttribute(MorphTargetName);
 	}
 
 	// call RegisterMorphTargetAttribute to make sure normals are registered / unregistered when needed
@@ -1196,7 +1200,7 @@ UDynamicMesh* UGeometryScriptLibrary_StaticMeshFunctions::CopyMorphTargetToSkele
 	TVertexInstanceAttributesConstRef<FVector3f> VertexNormals = MeshAttributes.GetVertexInstanceNormals();
 	
 	bool bMorphTargetIsEmpty = true;
-	
+
 	for (int32 SourceVID = 0; SourceVID < SourceMesh.VertexCount(); ++SourceVID)
 	{
 		const int32 TargetVID = NonManifoldMappingSupport.GetOriginalNonManifoldVertexID(SourceVID);
@@ -1365,7 +1369,8 @@ UDynamicMesh* UGeometryScriptLibrary_StaticMeshFunctions::CopySkinWeightProfileT
 	}
 	
 	FSkeletalMeshAttributes MeshAttributes(*MeshDescription);
-	MeshAttributes.Register();
+	constexpr bool bKeepExistingAttributes = true;
+	MeshAttributes.Register(bKeepExistingAttributes);
 
 	ToSkeletalMeshAsset->ModifyMeshDescription(TargetLOD.LODIndex);
 
diff --git a/Engine/Source/Runtime/Engine/Private/Rendering/SkeletalMeshLODImporterData.cpp b/Engine/Source/Runtime/Engine/Private/Rendering/SkeletalMeshLODImporterData.cpp
--- a/Engine/Source/Runtime/Engine/Private/Rendering/SkeletalMeshLODImporterData.cpp
+++ b/Engine/Source/Runtime/Engine/Private/Rendering/SkeletalMeshLODImporterData.cpp
@@ -2125,6 +2125,7 @@ bool FSkeletalMeshImportData::GetMeshDescription(const USkeletalMesh* InSkeletal
 	}
 	
 	VertexInstanceUVs.SetNumChannels(NumTexCoords);
+	OutMeshDescription.SetNumUVChannels(NumTexCoords);
 	
 	// Avoid repeated allocations and reserve the target buffers right off the bat.
 	OutMeshDescription.ReserveNewPolygonGroups(Materials.Num());
diff --git a/Engine/Source/Runtime/Engine/Private/SkeletalMeshLODModel.cpp b/Engine/Source/Runtime/Engine/Private/SkeletalMeshLODModel.cpp
--- a/Engine/Source/Runtime/Engine/Private/SkeletalMeshLODModel.cpp
+++ b/Engine/Source/Runtime/Engine/Private/SkeletalMeshLODModel.cpp
@@ -1446,6 +1446,7 @@ void FSkeletalMeshLODModel::GetMeshDescription(const USkeletalMesh *InSkeletalMe
 
 	// Ensure we have enough channels to store all the defined UV coordinates.
 	VertexInstanceUVs.SetNumChannels(static_cast<int32>(NumTexCoords));
+	OutMeshDescription.SetNumUVChannels(static_cast<int32>(NumTexCoords));
 	
 	const TArray<FSkeletalMaterial>& Materials = InSkeletalMesh->GetMaterials();
 	const bool bHasVertexColors = EnumHasAllFlags(InSkeletalMesh->GetVertexBufferFlags(), ESkeletalMeshVertexFlags::HasVertexColors);
diff --git a/Engine/Source/Runtime/SkeletalMeshDescription/Private/SkeletalMeshAttributes.cpp b/Engine/Source/Runtime/SkeletalMeshDescription/Private/SkeletalMeshAttributes.cpp
--- a/Engine/Source/Runtime/SkeletalMeshDescription/Private/SkeletalMeshAttributes.cpp
+++ b/Engine/Source/Runtime/SkeletalMeshDescription/Private/SkeletalMeshAttributes.cpp
@@ -71,11 +71,20 @@ void FSkeletalMeshAttributes::Register(bool bKeepExistingAttribute)
 		BoneElementsShared = BoneElements = MeshDescription.GetElements().Emplace(BonesElementName).Get();
 	}
 
-	BoneAttributes().RegisterAttribute<FName>(MeshAttribute::Bone::Name, 1, NAME_None, EMeshAttributeFlags::Mandatory);
-	
-	BoneAttributes().RegisterAttribute<int32>(MeshAttribute::Bone::ParentIndex, 1, INDEX_NONE, EMeshAttributeFlags::Mandatory);
+	if (!BoneAttributes().HasAttribute(MeshAttribute::Bone::Name) || !bKeepExistingAttribute)
+	{
+		BoneAttributes().RegisterAttribute<FName>(MeshAttribute::Bone::Name, 1, NAME_None, EMeshAttributeFlags::Mandatory);
+	}
 	
-	BoneAttributes().RegisterAttribute<FTransform>(MeshAttribute::Bone::Pose, 1, FTransform::Identity, EMeshAttributeFlags::Mandatory);
+	if (!BoneAttributes().HasAttribute(MeshAttribute::Bone::ParentIndex) || !bKeepExistingAttribute)
+	{
+		BoneAttributes().RegisterAttribute<int32>(MeshAttribute::Bone::ParentIndex, 1, INDEX_NONE, EMeshAttributeFlags::Mandatory);
+	}
+
+	if (!BoneAttributes().HasAttribute(MeshAttribute::Bone::Pose) || !bKeepExistingAttribute)
+	{
+		BoneAttributes().RegisterAttribute<FTransform>(MeshAttribute::Bone::Pose, 1, FTransform::Identity, EMeshAttributeFlags::Mandatory);
+	}
 
 	// Call super class
 	FStaticMeshAttributes::Register(bKeepExistingAttribute);
