add missing null pointer check to geometry script set simple collision on static mesh function

diff --git a/Engine/Plugins/Runtime/GeometryScripting/Source/GeometryScriptingCore/Private/CollisionFunctions.cpp b/Engine/Plugins/Runtime/GeometryScripting/Source/GeometryScriptingCore/Private/CollisionFunctions.cpp
--- a/Engine/Plugins/Runtime/GeometryScripting/Source/GeometryScriptingCore/Private/CollisionFunctions.cpp
+++ b/Engine/Plugins/Runtime/GeometryScripting/Source/GeometryScriptingCore/Private/CollisionFunctions.cpp
@@ -157,6 +157,10 @@ void ComputeCollisionFromMesh(
 
 static void UpdateStaticMeshCollision(UStaticMesh* StaticMeshAsset, TFunctionRef<bool(UStaticMesh*, UBodySetup*)> ApplyUpdate, bool bEmitTransaction, bool bMarkCollisionAsCustomized = true)
 {
+	if (!ensure(StaticMeshAsset))
+	{
+		return;
+	}
 #if WITH_EDITOR
 	if (bEmitTransaction && GEditor)
 	{
@@ -702,6 +706,11 @@ void UGeometryScriptLibrary_CollisionFunctions::SetSimpleCollisionOfStaticMesh(
 	FGeometryScriptSetStaticMeshCollisionOptions StaticMeshCollisionOptions,
 	UGeometryScriptDebug* Debug)
 {
+	if (StaticMesh == nullptr)
+	{
+		UE::Geometry::AppendError(Debug, EGeometryScriptErrorType::InvalidInputs, LOCTEXT("SetSimpleCollisionOfStaticMesh_InvalidStaticMesh", "SetSimpleCollisionOfStaticMesh: Input Mesh is Null"));
+		return;
+	}
 	UELocal::SetStaticMeshSimpleCollision(StaticMesh, SimpleCollision.AggGeom, Options.bEmitTransaction, StaticMeshCollisionOptions.bMarkAsCustomized);
 }
 
