fix split mesh by polygroup and split mesh by material ID geometry script methods not ordering the returned IDs in the same order as the split meshes
+ support reporting the ordering of IDs in the geometrycore split mesh function

diff --git a/Engine/Plugins/Experimental/MeshModelingToolsetExp/Source/GeometryProcessingAdapters/Private/CombineMeshInstancesImpl.cpp b/Engine/Plugins/Experimental/MeshModelingToolsetExp/Source/GeometryProcessingAdapters/Private/CombineMeshInstancesImpl.cpp
--- a/Engine/Plugins/Experimental/MeshModelingToolsetExp/Source/GeometryProcessingAdapters/Private/CombineMeshInstancesImpl.cpp
+++ b/Engine/Plugins/Experimental/MeshModelingToolsetExp/Source/GeometryProcessingAdapters/Private/CombineMeshInstancesImpl.cpp
@@ -3734,7 +3734,7 @@ void BuildCombinedMesh(
 			FDynamicMesh3& LODMesh = MeshLODs[LODIndex].Mesh;
 			FDynamicMeshPolygroupAttribute* SubsetIDAttrib = MeshLODs[LODIndex].SubsetIDs;
 			TArray<FDynamicMesh3> SplitMeshes;
-			FDynamicMeshEditor::SplitMesh(&LODMesh, SplitMeshes, [&](int32 tid) { return OrderedSubsetIDs.IndexOfByKey(SubsetIDAttrib->GetValue(tid)); });
+			FDynamicMeshEditor::SplitMesh(&LODMesh, SplitMeshes, [&](int32 tid) { return OrderedSubsetIDs.IndexOfByKey(SubsetIDAttrib->GetValue(tid)); }, INDEX_NONE, nullptr, true);
 
 			// code below assumes this. If it's not the case, then we have some more complex processing to figure out...
 			check(SplitMeshes.Num() == NumSubsets);		
@@ -4262,7 +4262,7 @@ void BuildCombinedMeshFromPrecomputedMeshes(
 			FDynamicMesh3& LODMesh = MeshLODs[LODIndex].Mesh;
 			FDynamicMeshPolygroupAttribute* SubsetIDAttrib = MeshLODs[LODIndex].SubsetIDs;
 			TArray<FDynamicMesh3> SplitMeshes;
-			FDynamicMeshEditor::SplitMesh(&LODMesh, SplitMeshes, [&](int32 tid) { return OrderedSubsetIDs.IndexOfByKey(SubsetIDAttrib->GetValue(tid)); });
+			FDynamicMeshEditor::SplitMesh(&LODMesh, SplitMeshes, [&](int32 tid) { return OrderedSubsetIDs.IndexOfByKey(SubsetIDAttrib->GetValue(tid)); }, INDEX_NONE, nullptr, true);
 
 			// code below assumes this. If it's not the case, then we have some more complex processing to figure out...
 			check(SplitMeshes.Num() == NumSubsets);
diff --git a/Engine/Plugins/Runtime/GeometryScripting/Source/GeometryScriptingCore/Private/MeshDecompositionFunctions.cpp b/Engine/Plugins/Runtime/GeometryScripting/Source/GeometryScriptingCore/Private/MeshDecompositionFunctions.cpp
--- a/Engine/Plugins/Runtime/GeometryScripting/Source/GeometryScriptingCore/Private/MeshDecompositionFunctions.cpp
+++ b/Engine/Plugins/Runtime/GeometryScripting/Source/GeometryScriptingCore/Private/MeshDecompositionFunctions.cpp
@@ -116,18 +116,10 @@ UDynamicMesh* UGeometryScriptLibrary_MeshDecompositionFunctions::SplitMeshByMate
 			return;
 		}
 
-		for (int32 tid : EditMesh.TriangleIndicesItr())
+		FDynamicMeshEditor::SplitMesh(&EditMesh, SplitMeshes, [&MaterialIDs](int32 TID)
 		{
-			int32 MaterialID = MaterialIDs->GetValue(tid);
-			ComponentMaterialIDs.AddUnique(MaterialID);
-		}
-		ComponentMaterialIDs.Sort();
-
-		FDynamicMeshEditor::SplitMesh(&EditMesh, SplitMeshes, [&](int32 tid) {
-			int32 MaterialID = MaterialIDs->GetValue(tid);
-			int32 Index = ComponentMaterialIDs.IndexOfByKey(MaterialID);
-			return (Index == INDEX_NONE) ? 0 : Index;
-		});
+			return MaterialIDs->GetValue(TID);
+		}, INDEX_NONE, &ComponentMaterialIDs, true);
 
 	});
 
@@ -168,24 +160,10 @@ UDynamicMesh* UGeometryScriptLibrary_MeshDecompositionFunctions::SplitMeshByPoly
 		}
 		SplitGroups = MakeUnique<FPolygroupSet>(&EditMesh, InputGroupLayer);
 
-		for (int32 tid : EditMesh.TriangleIndicesItr())
+		FDynamicMeshEditor::SplitMesh(&EditMesh, SplitMeshes, [&SplitGroups](int32 TID)
 		{
-			int32 GroupID = SplitGroups->GetGroup(tid);
-			ComponentPolygroups.AddUnique(GroupID);
-		}
-		ComponentPolygroups.Sort();
-
-		TMap<int32, int32> PolygroupMap;
-		for (int32 k = 0; k < ComponentPolygroups.Num(); ++k)
-		{
-			PolygroupMap.Add(ComponentPolygroups[k], k);
-		}
-
-		FDynamicMeshEditor::SplitMesh(&EditMesh, SplitMeshes, [&](int32 tid) {
-			int32 GroupID = SplitGroups->GetGroup(tid);
-			const int32* FoundIndex = PolygroupMap.Find(GroupID);
-			return (FoundIndex == nullptr) ? 0 : *FoundIndex;
-		});
+			return SplitGroups->GetGroup(TID);
+		}, INDEX_NONE, &ComponentPolygroups, true);
 	});
 
 	BuildNewDynamicMeshes(TargetMesh, MeshPool, SplitMeshes, ComponentMeshes);
diff --git a/Engine/Source/Runtime/GeometryCore/Private/DynamicMeshEditor.cpp b/Engine/Source/Runtime/GeometryCore/Private/DynamicMeshEditor.cpp
--- a/Engine/Source/Runtime/GeometryCore/Private/DynamicMeshEditor.cpp
+++ b/Engine/Source/Runtime/GeometryCore/Private/DynamicMeshEditor.cpp
@@ -2550,7 +2550,7 @@ void FDynamicMeshEditor::AppendTriangles(const FDynamicMesh3* SourceMesh, const
 }
 
 
-bool FDynamicMeshEditor::SplitMesh(const FDynamicMesh3* SourceMesh, TArray<FDynamicMesh3>& SplitMeshes, TFunctionRef<int(int)> TriIDToMeshID, int DeleteMeshID)
+bool FDynamicMeshEditor::SplitMesh(const FDynamicMesh3* SourceMesh, TArray<FDynamicMesh3>& SplitMeshes, TFunctionRef<int(int)> TriIDToMeshID, int DeleteMeshID, TArray<int32>* MeshIDsOut, bool bSortByID)
 {
 	using namespace UE::DynamicMeshEditorInternals;
 
@@ -2576,6 +2576,26 @@ bool FDynamicMeshEditor::SplitMesh(const FDynamicMesh3* SourceMesh, TArray<FDyna
 		return false; // nothing to do, so don't bother filling the split meshes array
 	}
 
+	if (bSortByID)
+	{
+		TArray<int32> SortedMeshIDs;
+		MeshIDToIndex.GenerateKeyArray(SortedMeshIDs);
+		SortedMeshIDs.Sort();
+		for (int32 Idx = 0; Idx < SortedMeshIDs.Num(); ++Idx)
+		{
+			MeshIDToIndex[SortedMeshIDs[Idx]] = Idx;
+		}
+	}
+
+	if (MeshIDsOut)
+	{
+		MeshIDsOut->SetNum(NumMeshes);
+		for (TPair<int32, int32> MeshIDIndexPair : MeshIDToIndex)
+		{
+			(*MeshIDsOut)[MeshIDIndexPair.Value] = MeshIDIndexPair.Key;
+		}
+	}
+
 	SplitMeshes.Reset();
 	SplitMeshes.SetNum(NumMeshes);
 	// enable matching attributes
diff --git a/Engine/Source/Runtime/GeometryCore/Public/DynamicMeshEditor.h b/Engine/Source/Runtime/GeometryCore/Public/DynamicMeshEditor.h
--- a/Engine/Source/Runtime/GeometryCore/Public/DynamicMeshEditor.h
+++ b/Engine/Source/Runtime/GeometryCore/Public/DynamicMeshEditor.h
@@ -627,12 +627,15 @@ public:
 	 * Create multiple meshes out of the source mesh by splitting triangles out.
 	 * Static because it creates multiple output meshes, so doesn't quite fit in the FDynamicMeshEditor model of operating on a single mesh
 	 *
-	 * @param SourceMesh
-	 * @param SplitMeshes
-	 * @param TriIDToMeshID
+	 * @param SourceMesh The mesh to split
+	 * @param SplitMeshes An array with one mesh per unique Mesh ID found in SourceMesh (optionally excluding the DeleteMeshID)
+	 * @param TriIDToMeshID A function mapping Triangle IDs to Mesh IDs
+	 * @param DeleteMeshID If non-negative, triangles with this Mesh ID will be removed rather than being placed in a split mesh
+	 * @param MeshIDPerSplitMesh If non-null, will be filled with the Mesh ID for each of the SplitMeshes generated
+	 * @param bSortByMeshID If true, SplitMeshes will be sorted in order of increasing Mesh ID
 	 * @return true if needed split, false if there were not multiple mesh ids so no split was needed
 	 */
-	static GEOMETRYCORE_API bool SplitMesh(const FDynamicMesh3* SourceMesh, TArray<FDynamicMesh3>& SplitMeshes, TFunctionRef<int(int)> TriIDToMeshID, int DeleteMeshID = -1);
+	static GEOMETRYCORE_API bool SplitMesh(const FDynamicMesh3* SourceMesh, TArray<FDynamicMesh3>& SplitMeshes, TFunctionRef<int(int)> TriIDToMeshID, int DeleteMeshID = -1, TArray<int32>* MeshIDPerSplitMesh = nullptr, bool bSortByMeshID = false);
 
 	/**
 	 * Merge any seams in the given 2D attribute Overlay along the given mesh edge IDs
