Add method to clear a UV channel

diff --git a/Engine/Plugins/Runtime/GeometryScripting/Source/GeometryScriptingCore/Private/MeshUVFunctions.cpp b/Engine/Plugins/Runtime/GeometryScripting/Source/GeometryScriptingCore/Private/MeshUVFunctions.cpp
--- a/Engine/Plugins/Runtime/GeometryScripting/Source/GeometryScriptingCore/Private/MeshUVFunctions.cpp
+++ b/Engine/Plugins/Runtime/GeometryScripting/Source/GeometryScriptingCore/Private/MeshUVFunctions.cpp
@@ -57,6 +57,35 @@ UDynamicMesh* UGeometryScriptLibrary_MeshUVFunctions::SetNumUVSets(
 }
 
 
+UDynamicMesh* UGeometryScriptLibrary_MeshUVFunctions::ClearUVChannel(
+	UDynamicMesh* TargetMesh,
+	int UVChannel,
+	UGeometryScriptDebug* Debug)
+{
+	if (TargetMesh == nullptr)
+	{
+		UE::Geometry::AppendError(Debug, EGeometryScriptErrorType::InvalidInputs, LOCTEXT("ClearUVChannel_InvalidInput", "ClearUVChannel: TargetMesh is Null"));
+		return TargetMesh;
+	}
+	TargetMesh->EditMesh([UVChannel, Debug](FDynamicMesh3& EditMesh)
+	{
+		FDynamicMeshUVOverlay* UVOverlay = nullptr;
+		if (!EditMesh.HasAttributes() || UVChannel < 0 || UVChannel >= EditMesh.Attributes()->NumUVLayers())
+		{
+			UE::Geometry::AppendError(Debug, EGeometryScriptErrorType::InvalidInputs, LOCTEXT("ClearUVChannel_NoChannel", "ClearUVChannel: Mesh did not have the requested UV channel, so it could not be cleared."));
+			return;
+		}
+		else
+		{
+			UVOverlay = EditMesh.Attributes()->GetUVLayer(UVChannel);
+			UVOverlay->ClearElements();
+		}
+	}, EDynamicMeshChangeType::AttributeEdit, EDynamicMeshAttributeChangeFlags::UVs, false);
+
+	return TargetMesh;
+}
+
+
 UDynamicMesh* UGeometryScriptLibrary_MeshUVFunctions::CopyUVSet(
 	UDynamicMesh* TargetMesh,
 	int FromUVSet,
diff --git a/Engine/Plugins/Runtime/GeometryScripting/Source/GeometryScriptingCore/Public/GeometryScript/MeshUVFunctions.h b/Engine/Plugins/Runtime/GeometryScripting/Source/GeometryScriptingCore/Public/GeometryScript/MeshUVFunctions.h
--- a/Engine/Plugins/Runtime/GeometryScripting/Source/GeometryScriptingCore/Public/GeometryScript/MeshUVFunctions.h
+++ b/Engine/Plugins/Runtime/GeometryScripting/Source/GeometryScriptingCore/Public/GeometryScript/MeshUVFunctions.h
@@ -244,6 +244,13 @@ public:
 		UPARAM(DisplayName = "Num UV Channels") int NumUVSets,
 		UGeometryScriptDebug* Debug = nullptr );
 
+	UFUNCTION(BlueprintCallable, Category = "GeometryScript|UVs", meta = (ScriptMethod), DisplayName = "Clear UV Channel")
+	static UPARAM(DisplayName = "Target Mesh") UDynamicMesh*
+	ClearUVChannel(
+		UDynamicMesh* TargetMesh,
+		UPARAM(DisplayName = "UV Channel") int UVChannel,
+		UGeometryScriptDebug* Debug = nullptr);
+
 	/**
 	* Copy the data in one UV Channel to another UV Channel on the same Target Mesh.
 	*/
