add selection-based version of geometry script flip normals

diff --git a/Engine/Plugins/Runtime/GeometryScripting/Source/GeometryScriptingCore/Private/MeshNormalsFunctions.cpp b/Engine/Plugins/Runtime/GeometryScripting/Source/GeometryScriptingCore/Private/MeshNormalsFunctions.cpp
--- a/Engine/Plugins/Runtime/GeometryScripting/Source/GeometryScriptingCore/Private/MeshNormalsFunctions.cpp
+++ b/Engine/Plugins/Runtime/GeometryScripting/Source/GeometryScriptingCore/Private/MeshNormalsFunctions.cpp
@@ -46,6 +46,79 @@ UDynamicMesh* UGeometryScriptLibrary_MeshNormalsFunctions::FlipNormals( UDynamic
 	return TargetMesh;
 }
 
+UDynamicMesh* UGeometryScriptLibrary_MeshNormalsFunctions::FlipTriangleSelectionNormals(
+	UDynamicMesh* TargetMesh,
+	FGeometryScriptMeshSelection Selection, 
+	bool bFlipTriangleOrientation,
+	bool bFlipNormalDirection, 
+	UGeometryScriptDebug* Debug)
+{
+	if (TargetMesh == nullptr)
+	{
+		UE::Geometry::AppendError(Debug, EGeometryScriptErrorType::InvalidInputs, LOCTEXT("FlipTriangleSelectionNormals_InvalidInput", "FlipTriangleSelectionNormals: TargetMesh is Null"));
+		return TargetMesh;
+	}
+
+	// Return early if no flipping was requested
+	if (!bFlipTriangleOrientation && !bFlipNormalDirection)
+	{
+		return TargetMesh;
+	}
+
+	TargetMesh->EditMesh([&Selection, bFlipTriangleOrientation, bFlipNormalDirection](FDynamicMesh3& EditMesh)
+	{
+		TSet<int32> FlippedVIDs, FlippedEIDs;
+		FDynamicMeshNormalOverlay* Normals = EditMesh.HasAttributes() ? EditMesh.Attributes()->PrimaryNormals() : nullptr;
+
+
+		Selection.ProcessByTriangleID(EditMesh, 
+		[&EditMesh, &FlippedVIDs, &FlippedEIDs, Normals, bFlipTriangleOrientation, bFlipNormalDirection](int32 TriangleID)
+		{
+			if (bFlipTriangleOrientation)
+			{
+				EditMesh.ReverseTriOrientation(TriangleID);
+			}
+
+			if (bFlipNormalDirection)
+			{
+				if (EditMesh.HasVertexNormals())
+				{
+					FIndex3i Tri = EditMesh.GetTriangle(TriangleID);
+					for (int32 SubIdx = 0; SubIdx < 3; ++SubIdx)
+					{
+						int32 VertexID = Tri[SubIdx];
+						bool bWasInSet;
+						FlippedVIDs.Add(VertexID, &bWasInSet);
+						if (!bWasInSet)
+						{
+							EditMesh.SetVertexNormal(VertexID, -EditMesh.GetVertexNormal(VertexID));
+						}
+					}
+				}
+				if (Normals)
+				{
+					FIndex3i TriEls;
+					if (Normals->GetTriangleIfValid(TriangleID, TriEls))
+					{
+						for (int32 SubIdx = 0; SubIdx < 3; ++SubIdx)
+						{
+							int32 ElementID = TriEls[SubIdx];
+							bool bWasInSet;
+							FlippedEIDs.Add(ElementID, &bWasInSet);
+							if (!bWasInSet)
+							{
+								Normals->SetElement(ElementID, -Normals->GetElement(ElementID));
+							}
+						}
+					}
+				}
+			}
+		}, false /*bProcessAllTrisIfSelectionEmpty*/);
+	}, EDynamicMeshChangeType::GeneralEdit, EDynamicMeshAttributeChangeFlags::Unknown, false);
+
+	return TargetMesh;
+}
+
 
 
 UDynamicMesh* UGeometryScriptLibrary_MeshNormalsFunctions::AutoRepairNormals(UDynamicMesh* TargetMesh, UGeometryScriptDebug* Debug)
diff --git a/Engine/Plugins/Runtime/GeometryScripting/Source/GeometryScriptingCore/Public/GeometryScript/MeshNormalsFunctions.h b/Engine/Plugins/Runtime/GeometryScripting/Source/GeometryScriptingCore/Public/GeometryScript/MeshNormalsFunctions.h
--- a/Engine/Plugins/Runtime/GeometryScripting/Source/GeometryScriptingCore/Public/GeometryScript/MeshNormalsFunctions.h
+++ b/Engine/Plugins/Runtime/GeometryScripting/Source/GeometryScriptingCore/Public/GeometryScript/MeshNormalsFunctions.h
@@ -80,6 +80,21 @@ public:
 		UDynamicMesh* TargetMesh, 
 		UGeometryScriptDebug* Debug = nullptr);
 
+	/**
+	 * Flip/Invert the normal vectors of the triangles in a given selection of TargetMesh. For edge or vertex selections, normals for the triangles touching the selected edges or vertices will be flipped.
+	 * 
+	 * @param bFlipTriangleOrientation If true, triangles are flipped by reversing the mesh triangle orientations, ie triangle (a,b,c) becomes (b,a,c)
+	 * @param bFlipNormalDirection If true, normals of vertices on the selected triangles are flipped by multiplying them by -1
+	 */
+	UFUNCTION(BlueprintCallable, Category = "GeometryScript|Normals", meta = (ScriptMethod))
+	static UPARAM(DisplayName = "Target Mesh") UDynamicMesh*
+	FlipTriangleSelectionNormals(
+		UDynamicMesh* TargetMesh,
+		FGeometryScriptMeshSelection Selection,
+		bool bFlipTriangleOrientation = true,
+		bool bFlipNormalDirection = true,
+		UGeometryScriptDebug* Debug = nullptr);
+
 	/**
 	 * Attempt to repair inconsistent normals in TargetMesh. Currently this is done in two passes. In the first pass, triangles with
 	 * reversed orientation from their neighours are incrementally flipped until each connected component has a consistent orientation,
