Fix the 'preserve vertex positions' flag in the geometry script simplify nodes (was not connected up)

diff --git a/Engine/Plugins/Runtime/GeometryScripting/Source/GeometryScriptingCore/Private/MeshSimplifyFunctions.cpp b/Engine/Plugins/Runtime/GeometryScripting/Source/GeometryScriptingCore/Private/MeshSimplifyFunctions.cpp
--- a/Engine/Plugins/Runtime/GeometryScripting/Source/GeometryScriptingCore/Private/MeshSimplifyFunctions.cpp
+++ b/Engine/Plugins/Runtime/GeometryScripting/Source/GeometryScriptingCore/Private/MeshSimplifyFunctions.cpp
@@ -142,6 +142,11 @@ void DoSimplifyMesh(
 		Options.bAllowSeamSplits, Options.bAllowSeamSmoothing, Options.bAllowSeamCollapse);
 	Simplifier.SetExternalConstraints(MoveTemp(Constraints));
 
+	if (Options.bPreserveVertexPositions)
+	{
+		Simplifier.CollapseMode = SimplificationType::ESimplificationCollapseModes::MinimalExistingVertexError;
+	}
+
 	if ( ProjectionTarget != nullptr && GeometricTolerance > 0)
 	{
 		Simplifier.GeometricErrorConstraint = SimplificationType::EGeometricErrorCriteria::PredictedPointToProjectionTarget;
