Add geometry script methods to merge mesh vertices -- either a pair of specific vertices by index, or closest pairs in two selections of vertices

diff --git a/Engine/Plugins/Runtime/GeometryScripting/Source/GeometryScriptingCore/Private/MeshBasicEditFunctions.cpp b/Engine/Plugins/Runtime/GeometryScripting/Source/GeometryScriptingCore/Private/MeshBasicEditFunctions.cpp
--- a/Engine/Plugins/Runtime/GeometryScripting/Source/GeometryScriptingCore/Private/MeshBasicEditFunctions.cpp
+++ b/Engine/Plugins/Runtime/GeometryScripting/Source/GeometryScriptingCore/Private/MeshBasicEditFunctions.cpp
@@ -366,6 +366,163 @@ UDynamicMesh* UGeometryScriptLibrary_MeshBasicEditFunctions::DeleteSelectedTrian
 }
 
 
+UDynamicMesh* UGeometryScriptLibrary_MeshBasicEditFunctions::MergeMeshVertexPair(
+	UDynamicMesh* TargetMesh,
+	int32 VertexKeep,
+	int32 VertexDiscard,
+	FGeometryScriptMergeVertexOptions Options,
+	bool& bSuccess,
+	double InterpParam,
+	bool bDeferChangeNotifications,
+	UGeometryScriptDebug* Debug)
+{
+	bSuccess = false;
+
+	if (!TargetMesh)
+	{
+		UE::Geometry::AppendError(Debug, EGeometryScriptErrorType::InvalidInputs, LOCTEXT("MergeMeshVertexPair_InvalidInput", "MergeMeshVertexPair: TargetMesh is Null"));
+		return TargetMesh;
+	}
+
+	TargetMesh->EditMesh([VertexKeep, VertexDiscard, InterpParam, Options, &bSuccess](FDynamicMesh3& Mesh)
+	{
+		if (Options.bOnlyBoundary)
+		{
+			if (!Mesh.IsBoundaryVertex(VertexKeep) || !Mesh.IsBoundaryVertex(VertexDiscard))
+			{
+				return;
+			}
+		}
+
+		FDynamicMesh3::FMergeVerticesOptions MergeOptions;
+		MergeOptions.bAllowNonBoundaryBowtieCreation = Options.bAllowNonBoundaryBowties;
+		FDynamicMesh3::FMergeVerticesInfo MergeInfo;
+		EMeshResult Result = Mesh.MergeVertices(VertexKeep, VertexDiscard, InterpParam, MergeOptions, MergeInfo);
+		bSuccess = Result == EMeshResult::Ok;
+	}, EDynamicMeshChangeType::GeneralEdit, EDynamicMeshAttributeChangeFlags::Unknown, bDeferChangeNotifications);
+
+	return TargetMesh;
+}
+
+
+UDynamicMesh* UGeometryScriptLibrary_MeshBasicEditFunctions::MergeMeshVerticesInSelections(
+	UDynamicMesh* TargetMesh,
+	FGeometryScriptMeshSelection SelectionKeep,
+	FGeometryScriptMeshSelection SelectionDiscard,
+	FGeometryScriptMergeVertexOptions Options,
+	int& NumMerged,
+	double InterpParam,
+	double DistanceThreshold,
+	bool bDeferChangeNotifications,
+	UGeometryScriptDebug* Debug)
+{
+	NumMerged = 0;
+
+	if (!TargetMesh)
+	{
+		UE::Geometry::AppendError(Debug, EGeometryScriptErrorType::InvalidInputs, LOCTEXT("MergeMeshVerticesInSelections_InvalidInput", "MergeMeshVerticesInSelections: TargetMesh is Null"));
+		return TargetMesh;
+	}
+
+	TargetMesh->EditMesh([&SelectionKeep, &SelectionDiscard, InterpParam, Options, &NumMerged, DistanceThreshold](FDynamicMesh3& Mesh)
+	{
+		TSet<int32> VertSets[2];
+		SelectionKeep.ProcessByVertexID(Mesh, [&Mesh, &VertSets, Options](int32 VID)
+		{
+			if (Options.bOnlyBoundary && !Mesh.IsBoundaryVertex(VID))
+			{
+				return;
+			}
+			VertSets[0].Add(VID);
+		});
+		SelectionDiscard.ProcessByVertexID(Mesh, [&Mesh, &SelectionKeep, &VertSets, Options](int32 VID)
+		{
+			if (Options.bOnlyBoundary && !Mesh.IsBoundaryVertex(VID))
+			{
+				return;
+			}
+			if (VertSets[0].Contains(VID))
+			{
+				return;
+			}
+			VertSets[1].Add(VID);
+		});
+
+		constexpr int32 MaxNumPasses = 2;
+		for (int32 Pass = 0; Pass < MaxNumPasses; ++Pass)
+		{
+			// Number of closest vertex matches that we've discarded due to finding a better reverse match
+			// Finding these discarded matches indicates that there may be additional matches to make if we run a second matching pass
+			int32 DiscardedMatches = 0;
+			// Mapping from DiscardVID to the best KeepVID match found so far
+			TMap<int32, TPair<int32, double>> BestMatches;
+			double ThreshSq = DistanceThreshold * DistanceThreshold;
+			for (int32 KeepVID : VertSets[0])
+			{
+				int32 BestMatch = INDEX_NONE;
+				double BestDistSq = ThreshSq;
+				FVector3d KeepPos = Mesh.GetVertex(KeepVID);
+
+				for (int32 DiscardVID : VertSets[1])
+				{
+					FVector3d DiscardPos = Mesh.GetVertex(DiscardVID);
+					double DistSq = FVector3d::DistSquared(KeepPos, DiscardPos);
+					if (DistSq < BestDistSq)
+					{
+						BestDistSq = DistSq;
+						BestMatch = DiscardVID;
+					}
+				}
+
+				if (BestMatch != INDEX_NONE)
+				{
+					TPair<int32, double>* PrevMatch = BestMatches.Find(BestMatch);
+					if (PrevMatch)
+					{
+						// Two 'keep' vertices have the same 'discard' vertex as their best match
+						// Use the best of these two matches, and track that we discarded the other match
+						DiscardedMatches++;
+						if (PrevMatch->Value > BestDistSq)
+						{
+							PrevMatch->Key = KeepVID;
+							PrevMatch->Value = BestDistSq;
+						}
+					}
+					else
+					{
+						BestMatches.Add(BestMatch, TPair<int32, double>(KeepVID, BestDistSq));
+					}
+				}
+			}
+
+			const bool bLastPass = Pass + 1 == MaxNumPasses || !DiscardedMatches;
+			for (TPair<int32, TPair<int32, double>> Match : BestMatches)
+			{
+				int32 KeepVertex = Match.Value.Key;
+				int32 DiscardVertex = Match.Key;
+				FDynamicMesh3::FMergeVerticesOptions MergeOptions;
+				MergeOptions.bAllowNonBoundaryBowtieCreation = Options.bAllowNonBoundaryBowties;
+				FDynamicMesh3::FMergeVerticesInfo MergeInfo;
+				EMeshResult Result = Mesh.MergeVertices(KeepVertex, DiscardVertex, InterpParam, MergeInfo);
+				NumMerged += (int32)(Result == EMeshResult::Ok);
+				if (!bLastPass)
+				{
+					VertSets[0].Remove(KeepVertex);
+					VertSets[1].Remove(DiscardVertex);
+				}
+			}
+
+			if (!DiscardedMatches)
+			{
+				break;
+			}
+		}
+	}, EDynamicMeshChangeType::GeneralEdit, EDynamicMeshAttributeChangeFlags::Unknown, bDeferChangeNotifications);
+
+	return TargetMesh;
+}
+
+
 void FGeometryScriptAppendMeshOptions::UpdateAttributesForCombineMode(FDynamicMesh3& Target, const FDynamicMesh3& Source)
 {
 	if (CombineMode == EGeometryScriptCombineAttributesMode::EnableAllMatching)
diff --git a/Engine/Plugins/Runtime/GeometryScripting/Source/GeometryScriptingCore/Public/GeometryScript/MeshBasicEditFunctions.h b/Engine/Plugins/Runtime/GeometryScripting/Source/GeometryScriptingCore/Public/GeometryScript/MeshBasicEditFunctions.h
--- a/Engine/Plugins/Runtime/GeometryScripting/Source/GeometryScriptingCore/Public/GeometryScript/MeshBasicEditFunctions.h
+++ b/Engine/Plugins/Runtime/GeometryScripting/Source/GeometryScriptingCore/Public/GeometryScript/MeshBasicEditFunctions.h
@@ -6,6 +6,7 @@
 #include "Kismet/BlueprintFunctionLibrary.h"
 #include "GeometryScript/GeometryScriptTypes.h"
 #include "GeometryScript/GeometryScriptSelectionTypes.h"
+#include "GeometryScript/MeshSpatialFunctions.h"
 #include "MeshBasicEditFunctions.generated.h"
 
 class UDynamicMesh;
@@ -85,6 +86,26 @@ public:
 	void UpdateAttributesForCombineMode(UE::Geometry::FDynamicMesh3& Target, const UE::Geometry::FDynamicMesh3& Source);
 };
 
+/**
+ * Options for merging vertex pairs
+ */
+USTRUCT(BlueprintType)
+struct GEOMETRYSCRIPTINGCORE_API FGeometryScriptMergeVertexOptions
+{
+	GENERATED_BODY()
+public:
+
+	// Whether to restrict merges to boundary vertices
+	UPROPERTY(EditAnywhere, BlueprintReadWrite, Category = Options)
+	bool bOnlyBoundary = true;
+
+	// Whether to allow the merge to introduce a non-boundary bowtie vertex (has no effect if bOnlyBoundary is true)
+	UPROPERTY(EditAnywhere, BlueprintReadWrite, Category = Options)
+	bool bAllowNonBoundaryBowties = false;
+
+};
+
+
 
 UCLASS(meta = (ScriptName = "GeometryScript_MeshEdits"))
 class GEOMETRYSCRIPTINGCORE_API UGeometryScriptLibrary_MeshBasicEditFunctions : public UBlueprintFunctionLibrary
@@ -228,6 +249,60 @@ public:
 		int& NumDeleted,
 		bool bDeferChangeNotifications = false );
 
+	/**
+	 * Attempt to merge together two vertices, and report whether they were merged.
+	 * Note that some merges may be prevented because they would create non-manifold edges in the mesh, which are not supported.
+	 * 
+	 * @param TargetMesh Mesh in which to merge vertices 
+	 * @param VertexKeep Vertex to keep after merge
+	 * @param VertexDiscard Vertex to discard after merge
+	 * @param Options Options for merge, controlling which merges should be permitted
+	 * @param bSuccess Flag indicating whether the merge succeeded
+	 * @param InterpParam The kept vertex is moved to interpolated position Lerp(VertexKeep Position, VertexDiscard Position, InterpParam)
+	 * @param bDeferChangeNotifications Whether to defer change notifications after this operation
+	 */
+	UFUNCTION(BlueprintCallable, Category = "GeometryScript|MeshEdits", meta = (ScriptMethod))
+	static UPARAM(DisplayName = "Target Mesh") UDynamicMesh*
+	MergeMeshVertexPair(
+		UDynamicMesh* TargetMesh,
+		int32 VertexKeep,
+		int32 VertexDiscard,
+		FGeometryScriptMergeVertexOptions Options,
+		bool& bSuccess,
+		double InterpParam = .5,
+		bool bDeferChangeNotifications = false,
+		UGeometryScriptDebug* Debug = nullptr
+	);
+
+	/**
+	 * Attempt to merge together vertices in one selection to their closest vertices in the second selection, within a distance threshold.
+	 * Note that some merges may be prevented because they would create non-manifold edges in the mesh, which are not supported.
+	 * 
+	 * @param TargetMesh Mesh in which to merge vertices 
+	 * @param SelectionKeep Selection of vertices to be merged with SelectionDiscard, treated as 'kept' vertices for purposes of interpolation.
+	 * @param SelectionDiscard Selection of vertices to be merged with SelectionKeep. Note: If a vertex is in both selections, it will be treated as if it were only in SelectionKeep.
+	 * @param Options Options for merge, controlling which merges should be permitted
+	 * @param NumMerged Number of vertices merged
+	 * @param InterpParam The kept vertex is moved to interpolated position Lerp(VertexKeep Position, VertexDiscard Position, InterpParam)
+	 * @param DistanceThreshold Vertices further apart than this threshold will not be merged
+	 * @param bDeferChangeNotifications Whether to defer change notifications after this operation
+	 */
+	UFUNCTION(BlueprintCallable, Category = "GeometryScript|MeshEdits", meta = (ScriptMethod))
+	static UPARAM(DisplayName = "Target Mesh") UDynamicMesh*
+	MergeMeshVerticesInSelections(
+		UDynamicMesh* TargetMesh,
+		FGeometryScriptMeshSelection SelectionKeep,
+		FGeometryScriptMeshSelection SelectionDiscard,
+		FGeometryScriptMergeVertexOptions Options,
+		int& NumMerged,
+		double InterpParam = .5,
+		double DistanceThreshold = 1,
+		bool bDeferChangeNotifications = false,
+		UGeometryScriptDebug* Debug = nullptr
+	);
+
+
+
 	/**
 	 * Apply Append Transform to Append Mesh and then add its geometry to the Target Mesh.
 	 * @param AppendOptions Control how details like mesh attributes are handled when one mesh is appended to another.
