Clear LOD material remapping when setting up new materials on skeletal mesh LOD

diff --git a/Engine/Plugins/Runtime/GeometryScripting/Source/GeometryScriptingCore/Private/MeshAssetFunctions.cpp b/Engine/Plugins/Runtime/GeometryScripting/Source/GeometryScriptingCore/Private/MeshAssetFunctions.cpp
--- a/Engine/Plugins/Runtime/GeometryScripting/Source/GeometryScriptingCore/Private/MeshAssetFunctions.cpp
+++ b/Engine/Plugins/Runtime/GeometryScripting/Source/GeometryScriptingCore/Private/MeshAssetFunctions.cpp
@@ -1012,6 +1012,7 @@ UDynamicMesh* UGeometryScriptLibrary_StaticMeshFunctions::CopyMeshToSkeletalMesh
 			NewMaterial.UVChannelData = FMeshUVChannelInfo(1.f);		// this avoids an ensure in  UStaticMesh::GetUVChannelData
 			NewMaterials.Add(NewMaterial);
 		}
+		SkeletalLODInfo->LODMaterialMap.Empty();
 		
 		ToSkeletalMeshAsset->SetMaterials(NewMaterials);
 	}
