fix spacing on geometry script BVH function display names

diff --git a/Engine/Plugins/Runtime/GeometryScripting/Source/GeometryScriptingCore/Public/GeometryScript/MeshSpatialFunctions.h b/Engine/Plugins/Runtime/GeometryScripting/Source/GeometryScriptingCore/Public/GeometryScript/MeshSpatialFunctions.h
--- a/Engine/Plugins/Runtime/GeometryScripting/Source/GeometryScriptingCore/Public/GeometryScript/MeshSpatialFunctions.h
+++ b/Engine/Plugins/Runtime/GeometryScripting/Source/GeometryScriptingCore/Public/GeometryScript/MeshSpatialFunctions.h
@@ -69,7 +69,7 @@ public:
 	/**
 	* Builds a Bounding Volume Hierarchy (BVH) object for a mesh that can be used with multiple spatial queries.
 	*/
-	UFUNCTION(BlueprintCallable, Category = "GeometryScript|Spatial", meta=(ScriptMethod))
+	UFUNCTION(BlueprintCallable, Category = "GeometryScript|Spatial", meta=(ScriptMethod, DisplayName="Build BVH For Mesh"))
 	static UPARAM(DisplayName = "Target Mesh") UDynamicMesh* 
 	BuildBVHForMesh(
 		UDynamicMesh* TargetMesh,
@@ -79,7 +79,7 @@ public:
 	/**
 	* Checks if the provided Bounding Volume Hierarchy (BVH) can still be used with the Mesh — it generally returns false if the mesh has been changed.
 	*/
-	UFUNCTION(BlueprintCallable, Category = "GeometryScript|Spatial", meta=(ScriptMethod))
+	UFUNCTION(BlueprintCallable, Category = "GeometryScript|Spatial", meta=(ScriptMethod, DisplayName="Is BVH Valid For Mesh"))
 	static UPARAM(DisplayName = "Target Mesh") UDynamicMesh* 
 	IsBVHValidForMesh(
 		UDynamicMesh* TargetMesh,
@@ -90,7 +90,7 @@ public:
 	/**
 	* Rebuilds the Bounding Volume Hierarchy (BVH) for the mesh in-place, which can reduce memory allocations, compared to building a new BVH.
 	*/
-	UFUNCTION(BlueprintCallable, Category = "GeometryScript|Spatial", meta=(ScriptMethod))
+	UFUNCTION(BlueprintCallable, Category = "GeometryScript|Spatial", meta=(ScriptMethod, DisplayName="Rebuild BVH For Mesh"))
 	static UPARAM(DisplayName = "Target Mesh") UDynamicMesh* 
 	RebuildBVHForMesh(
 		UDynamicMesh* TargetMesh,
