Fix up category and display name for some mesh material functions in geometry script

diff --git a/Engine/Plugins/Runtime/GeometryScripting/Source/GeometryScriptingCore/Public/GeometryScript/MeshMaterialFunctions.h b/Engine/Plugins/Runtime/GeometryScripting/Source/GeometryScriptingCore/Public/GeometryScript/MeshMaterialFunctions.h
--- a/Engine/Plugins/Runtime/GeometryScripting/Source/GeometryScriptingCore/Public/GeometryScript/MeshMaterialFunctions.h
+++ b/Engine/Plugins/Runtime/GeometryScripting/Source/GeometryScriptingCore/Public/GeometryScript/MeshMaterialFunctions.h
@@ -165,7 +165,7 @@ public:
 	* @param MaterialID the ID to be assigned to each triangle in the input list.
 	* @param bDeferChangeNotifications if true, the UDynamicMesh does not emit a change event/signal for this modification.
 	*/
-	UFUNCTION(BlueprintCallable, Category = "GeometryScript|Materials", meta=(ScriptMethod))
+	UFUNCTION(BlueprintCallable, Category = "GeometryScript|Materials", meta=(ScriptMethod, DisplayName="Set Material ID On Triangles"))
 	static UPARAM(DisplayName = "Target Mesh") UDynamicMesh* 
 	SetMaterialIDOnTriangles(
 		UDynamicMesh* TargetMesh,
@@ -179,7 +179,7 @@ public:
 	* @param MaterialID new Material ID to set
 	* @param bDeferChangeNotifications if true, the UDynamicMesh does not emit a change event/signal for this modification
 	*/
-	UFUNCTION(BlueprintCallable, Category = "GeometryScript|PolyGroups", meta=(ScriptMethod))
+	UFUNCTION(BlueprintCallable, Category = "GeometryScript|Materials", meta=(ScriptMethod, DisplayName="Set Material ID For Mesh Selection"))
 	static UPARAM(DisplayName = "Target Mesh") UDynamicMesh* 
 	SetMaterialIDForMeshSelection( 
 		UDynamicMesh* TargetMesh, 
