add box2d bounds overlap method to geometry script w/ support for testing for overlap in a wrapped [0, 1] space

diff --git a/Engine/Plugins/Runtime/GeometryScripting/Source/GeometryScriptingCore/Private/MeshUVFunctions.cpp b/Engine/Plugins/Runtime/GeometryScripting/Source/GeometryScriptingCore/Private/MeshUVFunctions.cpp
--- a/Engine/Plugins/Runtime/GeometryScripting/Source/GeometryScriptingCore/Private/MeshUVFunctions.cpp
+++ b/Engine/Plugins/Runtime/GeometryScripting/Source/GeometryScriptingCore/Private/MeshUVFunctions.cpp
@@ -1544,5 +1544,65 @@ UDynamicMesh* UGeometryScriptLibrary_MeshUVFunctions::ComputeMeshLocalUVParam(
 }
 
 
+bool UGeometryScriptLibrary_MeshUVFunctions::IntersectsUVBox2D(FBox2D A, FBox2D B, bool bWrappedToUnitRange)
+{
+	if (!A.bIsValid || !B.bIsValid)
+	{
+		return false;
+	}
+	if (!bWrappedToUnitRange)
+	{
+		return A.Intersect(B);
+	}
+	else
+	{
+		// Wrap the min value to [0,1] and shift the max value to have the same offset (without wrapping)
+		auto WrapMin = [](double& MinV, double& MaxV) -> void
+		{
+			double WrapMinV = FMath::Wrap(MinV, 0., 1.);
+			// Offset MinV and MaxV by (WrapMinV-MinV), so MinV becomes wrapped and MaxV keeps the same distance
+			MaxV += WrapMinV - MinV;
+			MinV = WrapMinV;
+		};
+		// Test for overlap in a single dimension, with wrapping
+		auto Overlaps1DWrapped = [](double AMin, double AMax, double BMin, double BMax) -> bool
+		{
+			// either value covers the full range, must overlap
+			if (AMax-AMin >= 1 || BMax - BMin >= 1)
+			{
+				return true;
+			}
+			if (BMax > 1)
+			{
+				if (AMax > 1)
+				{
+					// both wrap
+					return true;
+				}
+				// only range B wraps, convert to the only-range-A-wraps case
+				Swap(AMin, BMin);
+				Swap(AMax, BMax);
+			}
+			if (AMax > 1)
+			{
+				AMax -= 1.;
+				return BMin <= AMax || BMax >= AMin;
+			}
+			else
+			{
+				return !(AMax < BMin || AMin > BMax);
+			}
+		};
+		WrapMin(A.Min.X, A.Max.X);
+		WrapMin(B.Min.X, B.Max.X);
+		if (Overlaps1DWrapped(A.Min.X, A.Max.X, B.Min.X, B.Max.X))
+		{
+			WrapMin(A.Min.Y, A.Max.Y);
+			WrapMin(B.Min.Y, B.Max.Y);
+			return Overlaps1DWrapped(A.Min.Y, A.Max.Y, B.Min.Y, B.Max.Y);
+		}
+		return false;
+	}
+}
 
 #undef LOCTEXT_NAMESPACE
\ No newline at end of file
diff --git a/Engine/Plugins/Runtime/GeometryScripting/Source/GeometryScriptingCore/Public/GeometryScript/MeshUVFunctions.h b/Engine/Plugins/Runtime/GeometryScripting/Source/GeometryScriptingCore/Public/GeometryScript/MeshUVFunctions.h
--- a/Engine/Plugins/Runtime/GeometryScripting/Source/GeometryScriptingCore/Public/GeometryScript/MeshUVFunctions.h
+++ b/Engine/Plugins/Runtime/GeometryScripting/Source/GeometryScriptingCore/Public/GeometryScript/MeshUVFunctions.h
@@ -647,5 +647,16 @@ public:
 		double UVRotationDeg = 0.0,
 		UGeometryScriptDebug* Debug = nullptr );
 
+	/**
+	 * Test two Box2D bounds for intersection, with optional support for working in a wrapped space
+	 * 
+	 * @param A First box
+	 * @param B Second box
+	 * @param bWrappedCoordinates Whether to test the boxes for intersection in a space wrapped to unit range of [0, 1]
+	 * @return Whether the boxes intersect
+	 */
+	UFUNCTION(BlueprintPure, Category = "GeometryScript|UVs", meta = (DisplayName = "Intersects (UV Box2D)"))
+	static bool IntersectsUVBox2D(FBox2D A, FBox2D B, bool bWrappedToUnitRange = false);
+
 
 };
\ No newline at end of file
