Fix up display name for "Split Mesh By PolyGroups"

diff --git a/Engine/Plugins/Runtime/GeometryScripting/Source/GeometryScriptingCore/Public/GeometryScript/MeshDecompositionFunctions.h b/Engine/Plugins/Runtime/GeometryScripting/Source/GeometryScriptingCore/Public/GeometryScript/MeshDecompositionFunctions.h
--- a/Engine/Plugins/Runtime/GeometryScripting/Source/GeometryScriptingCore/Public/GeometryScript/MeshDecompositionFunctions.h
+++ b/Engine/Plugins/Runtime/GeometryScripting/Source/GeometryScriptingCore/Public/GeometryScript/MeshDecompositionFunctions.h
@@ -54,7 +54,7 @@ public:
 	 * @param ComponentPolygroups Original Polygroup for each Mesh in ComponentMeshes is returned here
 	 * @param MeshPool New meshes in ComponentMeshes output list are allocated from this pool if it is provided (highly recommended!!)
 	 */
-	UFUNCTION(BlueprintCallable, Category = "GeometryScript|Decomposition", meta=(ScriptMethod))
+	UFUNCTION(BlueprintCallable, Category = "GeometryScript|Decomposition", meta = (ScriptMethod, DisplayName = "Split Mesh By PolyGroups"))
 	static UPARAM(DisplayName = "Target Mesh") UDynamicMesh* 
 	SplitMeshByPolygroups(  
 		UDynamicMesh* TargetMesh, 
