Add geometry script nodes to add and find polygroup layers by name

diff --git a/Engine/Plugins/Runtime/GeometryScripting/Source/GeometryScriptingCore/Private/MeshPolygroupFunctions.cpp b/Engine/Plugins/Runtime/GeometryScripting/Source/GeometryScriptingCore/Private/MeshPolygroupFunctions.cpp
--- a/Engine/Plugins/Runtime/GeometryScripting/Source/GeometryScriptingCore/Private/MeshPolygroupFunctions.cpp
+++ b/Engine/Plugins/Runtime/GeometryScripting/Source/GeometryScriptingCore/Private/MeshPolygroupFunctions.cpp
@@ -131,8 +131,79 @@ UDynamicMesh* UGeometryScriptLibrary_MeshPolygroupFunctions::SetNumExtendedPolyg
 	return TargetMesh;
 }
 
+UDynamicMesh* UGeometryScriptLibrary_MeshPolygroupFunctions::AddNamedPolygroupLayer(
+	UDynamicMesh* TargetMesh,
+	FName LayerName,
+	FGeometryScriptGroupLayer& GroupLayer,
+	bool& bAlreadyExisted,
+	UGeometryScriptDebug* Debug)
+{
+	if (TargetMesh == nullptr)
+	{
+		UE::Geometry::AppendError(Debug, EGeometryScriptErrorType::InvalidInputs, LOCTEXT("AddNamedPolygroupLayer_InvalidInput", "AddNamedPolygroupLayer: TargetMesh is Null"));
+		return TargetMesh;
+	}
+
+	TargetMesh->EditMesh([&GroupLayer, &bAlreadyExisted, LayerName](FDynamicMesh3& EditMesh)
+	{
+		if (EditMesh.HasAttributes() == false)
+		{
+			EditMesh.EnableAttributes();
+		}
+		const int32 InitialNumLayers = EditMesh.Attributes()->NumPolygroupLayers();
+		for (int32 LayerIdx = 0; LayerIdx < InitialNumLayers; ++LayerIdx)
+		{
+			if (EditMesh.Attributes()->GetPolygroupLayer(LayerIdx)->GetName() == LayerName)
+			{
+				GroupLayer.bDefaultLayer = false;
+				GroupLayer.ExtendedLayerIndex = LayerIdx;
+				return;
+			}
+		}
+		// layer not found -- add it
+		EditMesh.Attributes()->SetNumPolygroupLayers(InitialNumLayers + 1);
+		EditMesh.Attributes()->GetPolygroupLayer(InitialNumLayers)->SetName(LayerName);
+		GroupLayer.bDefaultLayer = false;
+		GroupLayer.ExtendedLayerIndex = InitialNumLayers;
+	}, EDynamicMeshChangeType::GeneralEdit, EDynamicMeshAttributeChangeFlags::Unknown, false);
+
+	return TargetMesh;
+}
 
+const UDynamicMesh* UGeometryScriptLibrary_MeshPolygroupFunctions::FindExtendedPolygroupLayerByName(
+	const UDynamicMesh* TargetMesh,
+	FName LayerName,
+	FGeometryScriptGroupLayer& GroupLayer,
+	EGeometryScriptSearchOutcomePins& Outcome,
+	UGeometryScriptDebug* Debug)
+{
+	Outcome = EGeometryScriptSearchOutcomePins::NotFound;
+	if (TargetMesh == nullptr)
+	{
+		UE::Geometry::AppendError(Debug, EGeometryScriptErrorType::InvalidInputs, LOCTEXT("FindExtendedPolygroupLayerByName_InvalidInput", "FindExtendedPolygroupLayerByName: TargetMesh is Null"));
+		return TargetMesh;
+	}
 
+	TargetMesh->ProcessMesh([&Outcome, &GroupLayer, LayerName](const FDynamicMesh3& EditMesh)
+	{
+		if (!EditMesh.HasAttributes())
+		{
+			return;
+		}
+		for (int32 LayerIndex = 0; LayerIndex < EditMesh.Attributes()->NumPolygroupLayers(); ++LayerIndex)
+		{
+			if (EditMesh.Attributes()->GetPolygroupLayer(LayerIndex)->GetName() == LayerName)
+			{
+				Outcome = EGeometryScriptSearchOutcomePins::Found;
+				GroupLayer.bDefaultLayer = false;
+				GroupLayer.ExtendedLayerIndex = LayerIndex;
+				return;
+			}
+		}
+	});
+
+	return TargetMesh;
+}
 
 UDynamicMesh* UGeometryScriptLibrary_MeshPolygroupFunctions::ClearPolygroups( 
 	UDynamicMesh* TargetMesh, 
diff --git a/Engine/Plugins/Runtime/GeometryScripting/Source/GeometryScriptingCore/Public/GeometryScript/MeshPolygroupFunctions.h b/Engine/Plugins/Runtime/GeometryScripting/Source/GeometryScriptingCore/Public/GeometryScript/MeshPolygroupFunctions.h
--- a/Engine/Plugins/Runtime/GeometryScripting/Source/GeometryScriptingCore/Public/GeometryScript/MeshPolygroupFunctions.h
+++ b/Engine/Plugins/Runtime/GeometryScripting/Source/GeometryScriptingCore/Public/GeometryScript/MeshPolygroupFunctions.h
@@ -33,6 +33,42 @@ public:
 		int NumLayers,
 		UGeometryScriptDebug* Debug = nullptr );
 
+	/**
+	 * Add an extended polygroup layer with the given name. If a layer with that name is already on the mesh, that existing layer will be returned and no new layer will be added.
+	 * 
+	 * @param TargetMesh Mesh to update
+	 * @param LayerName Name of the PolyGroup layer
+	 * @param GroupLayer PolyGroup layer with the requested name
+	 * @param bAlreadyExisted If true, the layer already existed and did not need to be added (and the output GroupLayer will refer to that existing layer)
+	 */
+	UFUNCTION(BlueprintCallable, Category = "GeometryScript|PolyGroups", meta = (ScriptMethod, DisplayName = "Add Named PolyGroup Layer"))
+	static UPARAM(DisplayName = "Target Mesh") UDynamicMesh* 
+	AddNamedPolygroupLayer( 
+		UDynamicMesh* TargetMesh, 
+		FName LayerName,
+		FGeometryScriptGroupLayer& GroupLayer,
+		bool& bAlreadyExisted,
+		UGeometryScriptDebug* Debug = nullptr );
+
+	/**
+	 * Find an extended PolyGroup layer by its name. If there are multiple layers with the same name, returns the first such layer.
+	 * 
+	 * @param TargetMesh Mesh to search
+	 * @param LayerName Name of the PolyGroup layer to find
+	 * @param GroupLayer PolyGroup layer with the requested name (if Outcome is Success)
+	 * @param Outcome Whether the requested layer was found
+	 */
+	UFUNCTION(BlueprintCallable, Category = "GeometryScript|PolyGroups", meta = (ScriptMethod, DisplayName = "Find Extended PolyGroup Layer By Name", ExpandEnumAsExecs = "Outcome"))
+	static UPARAM(DisplayName = "Target Mesh") const UDynamicMesh*
+	FindExtendedPolygroupLayerByName(
+		const UDynamicMesh* TargetMesh,
+		FName LayerName,
+		FGeometryScriptGroupLayer &GroupLayer,
+		EGeometryScriptSearchOutcomePins& Outcome,
+		UGeometryScriptDebug* Debug = nullptr);
+
+
+
 	/**
 	* Resets the triangle PolyGroup assignments within a PolyGroup Layer to the given Clear Value (or 0 if no Clear Value is specified).
 	* Note, this will have no effect if PolyGroups have not been enabled on the mesh, or if the requested Group Layer does not exist. 
