Geometry Scripting: When grabbing normals from a normal overlay, use the correct index to look them up.

#rnx

diff --git a/Engine/Plugins/Runtime/GeometryScripting/Source/GeometryScriptingCore/Private/MeshAssetFunctions.cpp b/Engine/Plugins/Runtime/GeometryScripting/Source/GeometryScriptingCore/Private/MeshAssetFunctions.cpp
--- a/Engine/Plugins/Runtime/GeometryScripting/Source/GeometryScriptingCore/Private/MeshAssetFunctions.cpp
+++ b/Engine/Plugins/Runtime/GeometryScripting/Source/GeometryScriptingCore/Private/MeshAssetFunctions.cpp
@@ -1201,7 +1201,8 @@ UDynamicMesh* UGeometryScriptLibrary_StaticMeshFunctions::CopyMorphTargetToSkele
 	TVertexInstanceAttributesConstRef<FVector3f> VertexNormals = MeshAttributes.GetVertexInstanceNormals();
 	
 	bool bMorphTargetIsEmpty = true;
-
+	
+	TArray<int32> ElementIndexes;
 	for (int32 SourceVID = 0; SourceVID < SourceMesh.VertexCount(); ++SourceVID)
 	{
 		const int32 TargetVID = NonManifoldMappingSupport.GetOriginalNonManifoldVertexID(SourceVID);
@@ -1217,7 +1218,25 @@ UDynamicMesh* UGeometryScriptLibrary_StaticMeshFunctions::CopyMorphTargetToSkele
 
 			if (Options.bCopyNormals)
 			{
-				const FVector3f N0 = SourceNormals ? SourceNormals->GetElement(SourceVID) : FVector3f{Normals[SourceVID]};
+				FVector3f N0;
+				if (SourceNormals)
+				{
+					// For now, we average the normals. In the future, we should detect discontinuous normals and transfer them
+					// exactly to the target mesh, using vertex matching for the triangles to figure out which vertex instance goes
+					// where.
+					N0 = FVector3f::ZeroVector;
+					ElementIndexes.Reset();
+					SourceNormals->GetVertexElements(SourceVID, ElementIndexes);
+					for (int32 ElementIdx: ElementIndexes)
+					{
+						N0 += SourceNormals->GetElement(ElementIdx);
+					}
+					N0.Normalize();
+				}
+				else
+				{
+					N0 = FVector3f{Normals[SourceVID]};
+				}
 				for (const FVertexInstanceID VertexInstanceID : MeshDescription->GetVertexVertexInstanceIDs(TargetVID))
 				{
 					const FVector3f N1 = VertexNormals.Get(VertexInstanceID);
