Fix simple polygon query functions not returning their computed results + clarify display names for different Polygon Area geometry script methods

diff --git a/Engine/Plugins/Runtime/GeometryScripting/Source/GeometryScriptingCore/Private/PolygonFunctions.cpp b/Engine/Plugins/Runtime/GeometryScripting/Source/GeometryScriptingCore/Private/PolygonFunctions.cpp
--- a/Engine/Plugins/Runtime/GeometryScripting/Source/GeometryScriptingCore/Private/PolygonFunctions.cpp
+++ b/Engine/Plugins/Runtime/GeometryScripting/Source/GeometryScriptingCore/Private/PolygonFunctions.cpp
@@ -130,7 +130,7 @@ double UGeometryScriptLibrary_SimplePolygonFunctions::GetPolygonArcLength(FGeome
 {
 	if (Polygon.Vertices.IsValid() && !Polygon.Vertices->IsEmpty())
 	{
-		UE::Geometry::CurveUtil::ArcLength<double, FVector2D>(*Polygon.Vertices, true);
+		return UE::Geometry::CurveUtil::ArcLength<double, FVector2D>(*Polygon.Vertices, true);
 	}
 	return 0.0;
 }
@@ -139,7 +139,7 @@ double UGeometryScriptLibrary_SimplePolygonFunctions::GetPolygonArea(FGeometrySc
 {
 	if (Polygon.Vertices.IsValid() && !Polygon.Vertices->IsEmpty())
 	{
-		UE::Geometry::CurveUtil::SignedArea2<double, FVector2D>(*Polygon.Vertices);
+		return UE::Geometry::CurveUtil::SignedArea2<double, FVector2D>(*Polygon.Vertices);
 	}
 	return 0.0;
 }
diff --git a/Engine/Plugins/Runtime/GeometryScripting/Source/GeometryScriptingCore/Public/GeometryScript/PolygonFunctions.h b/Engine/Plugins/Runtime/GeometryScripting/Source/GeometryScriptingCore/Public/GeometryScript/PolygonFunctions.h
--- a/Engine/Plugins/Runtime/GeometryScripting/Source/GeometryScriptingCore/Public/GeometryScript/PolygonFunctions.h
+++ b/Engine/Plugins/Runtime/GeometryScripting/Source/GeometryScriptingCore/Public/GeometryScript/PolygonFunctions.h
@@ -137,7 +137,7 @@ public:
 	/**
 	 * Returns the area enclosed by a Simple Polygon
 	 */
-	UFUNCTION(BlueprintCallable, Category = "GeometryScript|SimplePolygon", meta = (ScriptMethod))
+	UFUNCTION(BlueprintCallable, Category = "GeometryScript|SimplePolygon", meta = (ScriptMethod, DisplayName = "Get Simple Polygon Area"))
 	static double GetPolygonArea(FGeometryScriptSimplePolygon Polygon);
 
 	/**
@@ -225,9 +225,9 @@ public:
 	static FVector2D GetPolygonVertex(FGeometryScriptGeneralPolygonList PolygonList, bool& bIsValidVertex, int32 VertexIndex, int32 PolygonIndex, int32 HoleIndex = -1);
 
 	/**
-	 * Returns the area enclosed by a Polygon. Returns zero for an invalid PolygonIndex. 
+	 * Returns the area enclosed by a specified Polygon in a Polygon List. Returns zero for an invalid PolygonIndex. 
 	 */
-	UFUNCTION(BlueprintCallable, Category = "GeometryScript|PolygonList", meta = (ScriptMethod))
+	UFUNCTION(BlueprintCallable, Category = "GeometryScript|PolygonList", meta = (ScriptMethod, DisplayName = "Get Area of Polygon in List"))
 	static double GetPolygonArea(FGeometryScriptGeneralPolygonList PolygonList, bool& bValidIndex, int32 PolygonIndex);
 
 	/**
@@ -237,13 +237,13 @@ public:
 	static FBox2D GetPolygonBounds(FGeometryScriptGeneralPolygonList PolygonList, bool& bValidIndex, int32 PolygonIndex);
 
 	/**
-	 * Returns the area enclosed by a Polygon
+	 * Returns the area enclosed by a Polygon List
 	 */
 	UFUNCTION(BlueprintCallable, Category = "GeometryScript|PolygonList", meta = (ScriptMethod))
 	static double GetPolygonListArea(FGeometryScriptGeneralPolygonList PolygonList);
 
 	/**
-	 * Returns the bounding box of a Polygon
+	 * Returns the bounding box of a Polygon List
 	 */
 	UFUNCTION(BlueprintCallable, Category = "GeometryScript|PolygonList", meta = (ScriptMethod))
 	static FBox2D GetPolygonListBounds(FGeometryScriptGeneralPolygonList PolygonList);
