guard against infinite recursion or adding-to-loop-while-iterating issues when generated dynamic mesh actor adds new actors or triggers rebuilds of actors inside its own rebuild

diff --git a/Engine/Plugins/Runtime/GeometryScripting/Source/GeometryScriptingEditor/Private/GeometryActors/EditorGeometryGenerationSubsystem.cpp b/Engine/Plugins/Runtime/GeometryScripting/Source/GeometryScriptingEditor/Private/GeometryActors/EditorGeometryGenerationSubsystem.cpp
--- a/Engine/Plugins/Runtime/GeometryScripting/Source/GeometryScriptingEditor/Private/GeometryActors/EditorGeometryGenerationSubsystem.cpp
+++ b/Engine/Plugins/Runtime/GeometryScripting/Source/GeometryScriptingEditor/Private/GeometryActors/EditorGeometryGenerationSubsystem.cpp
@@ -10,13 +10,19 @@
 void UEditorGeometryGenerationManager::Tick(float DeltaTime)
 {
 	// rebuild all valid generated actors, if necessary
+	// note: we have a separate to-rebuild array in case the process of the rebuild modifies the ActiveGeneratedActors array
+	TArray<AGeneratedDynamicMeshActor*> ToRebuild;
 	for (AGeneratedDynamicMeshActor* Actor : ActiveGeneratedActors)
 	{
-		if (Actor->IsValidLowLevel() && IsValid(Actor) && !Actor->IsUnreachable() && Actor->GetLevel() != nullptr)
+		if (Actor->IsValidLowLevel() && IsValid(Actor) && !Actor->IsUnreachable() && Actor->GetLevel() != nullptr && Actor->WantsRebuild())
 		{
-			Actor->ExecuteRebuildGeneratedMeshIfPending();
+			ToRebuild.Add(Actor);
 		}
 	}
+	for (AGeneratedDynamicMeshActor* Actor : ToRebuild)
+	{
+		Actor->ExecuteRebuildGeneratedMeshIfPending();
+	}
 }
 
 bool UEditorGeometryGenerationManager::IsTickable() const
diff --git a/Engine/Plugins/Runtime/GeometryScripting/Source/GeometryScriptingEditor/Private/GeometryActors/GeneratedDynamicMeshActor.cpp b/Engine/Plugins/Runtime/GeometryScripting/Source/GeometryScriptingEditor/Private/GeometryActors/GeneratedDynamicMeshActor.cpp
--- a/Engine/Plugins/Runtime/GeometryScripting/Source/GeometryScriptingEditor/Private/GeometryActors/GeneratedDynamicMeshActor.cpp
+++ b/Engine/Plugins/Runtime/GeometryScripting/Source/GeometryScriptingEditor/Private/GeometryActors/GeneratedDynamicMeshActor.cpp
@@ -168,14 +168,24 @@ void AGeneratedDynamicMeshActor::UnregisterWithGenerationManager()
 	}
 }
 
-
-void AGeneratedDynamicMeshActor::ExecuteRebuildGeneratedMeshIfPending()
+bool AGeneratedDynamicMeshActor::WantsRebuild()
 {
 	if (bFrozen)
 	{
-		return;
+		return false;
 	}
 	if (bGeneratedMeshRebuildPending == false)
+	{
+		return false;
+	}
+
+	return true;
+}
+
+
+void AGeneratedDynamicMeshActor::ExecuteRebuildGeneratedMeshIfPending()
+{
+	if (!WantsRebuild())
 	{
 		return;
 	}
@@ -186,6 +196,16 @@ void AGeneratedDynamicMeshActor::ExecuteRebuildGeneratedMeshIfPending()
 		return;
 	}
 
+	if (!ensure(!bIsCurrentlyRebuilding))
+	{
+		return;
+	}
+	bIsCurrentlyRebuilding = true;
+	ON_SCOPE_EXIT
+	{
+		bIsCurrentlyRebuilding = false;
+	};
+
 	// Automatically defer collision updates during generated mesh rebuild. If we do not do this, then
 	// each mesh change will result in collision being rebuilt, which is very expensive !!
 	bool bEnabledDeferredCollision = false;
@@ -231,7 +251,8 @@ void AGeneratedDynamicMeshActor::RebuildGeneratedMesh(UDynamicMesh* TargetMesh)
 
 void AGeneratedDynamicMeshActor::MarkForMeshRebuild(bool bImmediate, bool bImmediateEvenIfFrozen)
 {
-	if (bImmediate)
+	// note we ignore the immediate rebuild flag if we're already rebuilding
+	if (bImmediate && !bIsCurrentlyRebuilding)
 	{
 		bool bRestoreFrozenState = bFrozen;
 		if (bImmediateEvenIfFrozen && bFrozen)
diff --git a/Engine/Plugins/Runtime/GeometryScripting/Source/GeometryScriptingEditor/Public/GeometryActors/GeneratedDynamicMeshActor.h b/Engine/Plugins/Runtime/GeometryScripting/Source/GeometryScriptingEditor/Public/GeometryActors/GeneratedDynamicMeshActor.h
--- a/Engine/Plugins/Runtime/GeometryScripting/Source/GeometryScriptingEditor/Public/GeometryActors/GeneratedDynamicMeshActor.h
+++ b/Engine/Plugins/Runtime/GeometryScripting/Source/GeometryScriptingEditor/Public/GeometryActors/GeneratedDynamicMeshActor.h
@@ -60,6 +60,8 @@ public:
 	 */
 	virtual void ExecuteRebuildGeneratedMeshIfPending();
 
+	virtual bool WantsRebuild();
+
 
 protected:
 	/**
@@ -116,6 +118,11 @@ protected:
 	// is where the mesh rebuilds are executed
 	bool bIsRegisteredWithGenerationManager = false;
 
+private:
+	// flag set during a rebuild -- used to help detect/avoid recursive rebuilds
+	bool bIsCurrentlyRebuilding = false;
+protected:
+
 	virtual void RegisterWithGenerationManager();
 	virtual void UnregisterWithGenerationManager();
 
