Expose dynamic mesh sculpt layers to geometry script

diff --git a/Engine/Plugins/Runtime/GeometryScripting/Source/GeometryScriptingCore/Private/MeshSculptLayersFunctions.cpp b/Engine/Plugins/Runtime/GeometryScripting/Source/GeometryScriptingCore/Private/MeshSculptLayersFunctions.cpp
new file mode 100644
--- /dev/null
+++ b/Engine/Plugins/Runtime/GeometryScripting/Source/GeometryScriptingCore/Private/MeshSculptLayersFunctions.cpp
@@ -0,0 +1,282 @@
+// Copyright Epic Games, Inc. All Rights Reserved.
+
+#include "GeometryScript/MeshSculptLayersFunctions.h"
+
+#include "DynamicMesh/DynamicMesh3.h"
+#include "DynamicMesh/DynamicMeshAttributeSet.h"
+#include "DynamicMesh/MeshNormals.h"
+
+#include "UDynamicMesh.h"
+
+#include UE_INLINE_GENERATED_CPP_BY_NAME(MeshSculptLayersFunctions)
+
+using namespace UE::Geometry;
+
+#define LOCTEXT_NAMESPACE "UGeometryScriptLibrary_MeshSculptLayerFunctions"
+
+
+UDynamicMesh* UGeometryScriptLibrary_MeshSculptLayersFunctions::EnableSculptLayers(
+	UDynamicMesh* TargetMesh,
+	int32 NumLayers,
+	UGeometryScriptDebug* Debug)
+{
+	if (TargetMesh == nullptr)
+	{
+		UE::Geometry::AppendError(Debug, EGeometryScriptErrorType::InvalidInputs, LOCTEXT("EnableSculptLayers_InvalidInput", "EnableSculptLayers: TargetMesh is Null"));
+		return TargetMesh;
+	}
+
+	TargetMesh->EditMesh([NumLayers](FDynamicMesh3& EditMesh)
+	{
+		if (!EditMesh.HasAttributes())
+		{
+			EditMesh.EnableAttributes();
+		}
+		EditMesh.Attributes()->EnableSculptLayers(NumLayers);
+	}, EDynamicMeshChangeType::GeneralEdit, EDynamicMeshAttributeChangeFlags::Unknown, false);
+
+	return TargetMesh;
+}
+
+UDynamicMesh* UGeometryScriptLibrary_MeshSculptLayersFunctions::SetSculptLayerWeight(
+	UDynamicMesh* TargetMesh,
+	int32 LayerIndex,
+	double Weight,
+	FGeometryScriptSculptLayerUpdateOptions Options,
+	UGeometryScriptDebug* Debug)
+{
+	if (TargetMesh == nullptr)
+	{
+		UE::Geometry::AppendError(Debug, EGeometryScriptErrorType::InvalidInputs, LOCTEXT("SetSculptLayerWeight_InvalidInput", "SetSculptLayerWeight: TargetMesh is Null"));
+		return TargetMesh;
+	}
+
+	TargetMesh->EditMesh([LayerIndex, Weight, &Options, Debug](FDynamicMesh3& EditMesh)
+	{
+		if (!EditMesh.HasAttributes())
+		{
+			UE::Geometry::AppendError(Debug, EGeometryScriptErrorType::InvalidInputs, LOCTEXT("SetSculptLayerWeight_NoAttribs", "SetSculptLayerWeight: TargetMesh does not have attributes enabled"));
+			return;
+		}
+		if (LayerIndex >= 0 && LayerIndex < EditMesh.Attributes()->NumSculptLayers())
+		{
+			TArray<double> Wts(EditMesh.Attributes()->GetSculptLayers()->GetLayerWeights());
+			Wts[LayerIndex] = Weight;
+			EditMesh.Attributes()->GetSculptLayers()->UpdateLayerWeights(Wts);
+			if (Options.bRecomputeNormals)
+			{
+				FMeshNormals::QuickRecomputeOverlayNormals(EditMesh);
+			}
+		}
+		else
+		{
+			UE::Geometry::AppendError(Debug, EGeometryScriptErrorType::InvalidInputs, LOCTEXT("SetSculptLayerWeight_LayerNotFound", "SetSculptLayerWeight: Requested sculpt layer was not enabled on TargetMesh. Use EnableSculptLayers first."));
+		}
+	}, EDynamicMeshChangeType::GeneralEdit, EDynamicMeshAttributeChangeFlags::Unknown, false);
+
+	return TargetMesh;
+}
+
+UDynamicMesh* UGeometryScriptLibrary_MeshSculptLayersFunctions::SetSculptLayerWeightsArray(
+	UDynamicMesh* TargetMesh,
+	TArray<double> SetWeights,
+	FGeometryScriptSculptLayerUpdateOptions Options,
+	UGeometryScriptDebug* Debug)
+{
+	if (TargetMesh == nullptr)
+	{
+		UE::Geometry::AppendError(Debug, EGeometryScriptErrorType::InvalidInputs, LOCTEXT("SetSculptLayerWeight_InvalidInput", "SetSculptLayerWeight: TargetMesh is Null"));
+		return TargetMesh;
+	}
+
+	TargetMesh->EditMesh([&SetWeights, &Options, Debug](FDynamicMesh3& EditMesh)
+	{
+		if (!EditMesh.HasAttributes())
+		{
+			UE::Geometry::AppendError(Debug, EGeometryScriptErrorType::InvalidInputs, LOCTEXT("SetSculptLayerWeight_NoAttribs", "SetSculptLayerWeight: TargetMesh does not have attributes enabled"));
+			return;
+		}
+		if (SetWeights.Num() > 0 && EditMesh.Attributes()->NumSculptLayers() > 0)
+		{
+			TArray<double> Wts(EditMesh.Attributes()->GetSculptLayers()->GetLayerWeights());
+			int32 NumToCopy = FMath::Min(Wts.Num(), SetWeights.Num());
+			for (int32 Idx = 0; Idx < NumToCopy; ++Idx)
+			{
+				Wts[Idx] = SetWeights[Idx];
+			}
+			EditMesh.Attributes()->GetSculptLayers()->UpdateLayerWeights(Wts);
+			if (Options.bRecomputeNormals)
+			{
+				FMeshNormals::QuickRecomputeOverlayNormals(EditMesh);
+			}
+		}
+		else
+		{
+			UE::Geometry::AppendError(Debug, EGeometryScriptErrorType::InvalidInputs, LOCTEXT("SetActiveSculptLayer_LayerNotFound", "SetActiveSculptLayer: Requested sculpt layer was not enabled on TargetMesh. Use EnableSculptLayers first."));
+		}
+	}, EDynamicMeshChangeType::GeneralEdit, EDynamicMeshAttributeChangeFlags::Unknown, false);
+
+	return TargetMesh;
+}
+
+TArray<double> UGeometryScriptLibrary_MeshSculptLayersFunctions::GetSculptLayerWeightsArray(UDynamicMesh* TargetMesh)
+{
+	TArray<double> Weights;
+
+	if (TargetMesh == nullptr)
+	{
+		UE_LOG(LogGeometry, Error, TEXT("GetActiveSculptLayer: TargetMesh is Null"));
+		return Weights;
+	}
+
+	TargetMesh->ProcessMesh([&Weights](const FDynamicMesh3& EditMesh)
+	{
+		if (!EditMesh.HasAttributes() || !EditMesh.Attributes()->NumSculptLayers())
+		{
+			UE_LOG(LogGeometry, Error, TEXT("GetActiveSculptLayer: TargetMesh does not have sculpt layers enabled"));
+			return;
+		}
+		Weights = EditMesh.Attributes()->GetSculptLayers()->GetLayerWeights();
+	});
+
+	return Weights;
+}
+
+UDynamicMesh* UGeometryScriptLibrary_MeshSculptLayersFunctions::SetActiveSculptLayer(
+	UDynamicMesh* TargetMesh,
+	int32 LayerIndex,
+	UGeometryScriptDebug* Debug)
+{
+	if (TargetMesh == nullptr)
+	{
+		UE::Geometry::AppendError(Debug, EGeometryScriptErrorType::InvalidInputs, LOCTEXT("SetActiveSculptLayer_InvalidInput", "SetActiveSculptLayer: TargetMesh is Null"));
+		return TargetMesh;
+	}
+
+	TargetMesh->EditMesh([LayerIndex, Debug](FDynamicMesh3& EditMesh)
+	{
+		if (!EditMesh.HasAttributes())
+		{
+			UE::Geometry::AppendError(Debug, EGeometryScriptErrorType::InvalidInputs, LOCTEXT("SetActiveSculptLayer_NoAttribs", "SetActiveSculptLayer: TargetMesh does not have attributes enabled"));
+			return;
+		}
+		if (LayerIndex >= 0 && LayerIndex < EditMesh.Attributes()->NumSculptLayers())
+		{
+			EditMesh.Attributes()->GetSculptLayers()->SetActiveLayer(LayerIndex);
+		}
+		else
+		{
+			UE::Geometry::AppendError(Debug, EGeometryScriptErrorType::InvalidInputs, LOCTEXT("SetActiveSculptLayer_LayerNotFound", "SetActiveSculptLayer: Requested sculpt layer was not enabled on TargetMesh. Use EnableSculptLayers first."));
+		}
+	}, EDynamicMeshChangeType::GeneralEdit, EDynamicMeshAttributeChangeFlags::Unknown, false);
+
+	return TargetMesh;
+}
+
+int32 UGeometryScriptLibrary_MeshSculptLayersFunctions::GetNumSculptLayers(const UDynamicMesh* TargetMesh)
+{
+	if (TargetMesh == nullptr)
+	{
+		UE_LOG(LogGeometry, Error, TEXT("GetNumSculptLayers: TargetMesh is Null"));
+		return -1;
+	}
+
+	int32 FoundLayerCount = -1;
+	TargetMesh->ProcessMesh([&FoundLayerCount](const FDynamicMesh3& EditMesh)
+	{
+		if (!EditMesh.HasAttributes())
+		{
+			UE_LOG(LogGeometry, Error, TEXT("GetNumSculptLayers: TargetMesh does not have attributes enabled"));
+			return;
+		}
+		FoundLayerCount = EditMesh.Attributes()->NumSculptLayers();
+	});
+	return FoundLayerCount;
+}
+
+int32 UGeometryScriptLibrary_MeshSculptLayersFunctions::GetActiveSculptLayer(const UDynamicMesh* TargetMesh)
+{
+	if (TargetMesh == nullptr)
+	{
+		UE_LOG(LogGeometry, Error, TEXT("GetActiveSculptLayer: TargetMesh is Null"));
+		return INDEX_NONE;
+	}
+
+	int32 ActiveSculptLayer = INDEX_NONE;
+	TargetMesh->ProcessMesh([&ActiveSculptLayer](const FDynamicMesh3& EditMesh)
+	{
+		if (!EditMesh.HasAttributes() || !EditMesh.Attributes()->NumSculptLayers())
+		{
+			UE_LOG(LogGeometry, Error, TEXT("GetActiveSculptLayer: TargetMesh does not have sculpt layers enabled"));
+			return;
+		}
+		ActiveSculptLayer = EditMesh.Attributes()->GetSculptLayers()->GetActiveLayer();
+	});
+	return ActiveSculptLayer;
+}
+
+UDynamicMesh* UGeometryScriptLibrary_MeshSculptLayersFunctions::DiscardSculptLayers(
+	UDynamicMesh* TargetMesh,
+	UGeometryScriptDebug* Debug)
+{
+	if (TargetMesh == nullptr)
+	{
+		UE::Geometry::AppendError(Debug, EGeometryScriptErrorType::InvalidInputs, LOCTEXT("DiscardSculptLayers_InvalidInput", "DiscardSculptLayers: TargetMesh is Null"));
+		return TargetMesh;
+	}
+
+	TargetMesh->EditMesh([](FDynamicMesh3& EditMesh)
+	{
+		if (!EditMesh.HasAttributes())
+		{
+			return;
+		}
+		EditMesh.Attributes()->DiscardSculptLayers();
+	}, EDynamicMeshChangeType::GeneralEdit, EDynamicMeshAttributeChangeFlags::Unknown, false);
+
+	return TargetMesh;
+}
+
+UDynamicMesh* UGeometryScriptLibrary_MeshSculptLayersFunctions::MergeSculptLayers(
+	UDynamicMesh* TargetMesh,
+	int32& OutActiveLayer,
+	int32 MergeLayerStart,
+	int32 MergeLayerNum,
+	bool bUseWeights,
+	UGeometryScriptDebug* Debug)
+{
+	if (TargetMesh == nullptr)
+	{
+		UE::Geometry::AppendError(Debug, EGeometryScriptErrorType::InvalidInputs, LOCTEXT("MergeSculptLayers_InvalidInput", "MergeSculptLayers: TargetMesh is Null"));
+		return TargetMesh;
+	}
+
+	TargetMesh->EditMesh([MergeLayerStart, MergeLayerNum, &OutActiveLayer, bUseWeights, Debug](FDynamicMesh3& EditMesh)
+	{
+		if (!EditMesh.HasAttributes())
+		{
+			UE::Geometry::AppendError(Debug, EGeometryScriptErrorType::InvalidInputs, LOCTEXT("MergeSculptLayers_NoAttribs", "MergeSculptLayers: TargetMesh does not have attributes enabled"));
+			return;
+		}
+		if (FDynamicMeshSculptLayers* SculptLayers = EditMesh.Attributes()->GetSculptLayers())
+		{
+			if (MergeLayerStart >= 0 && MergeLayerStart < SculptLayers->NumLayers())
+			{
+				int32 EndLayer = FMath::Min(SculptLayers->NumLayers() - 1, MergeLayerStart + MergeLayerNum);
+				if (EndLayer > MergeLayerStart)
+				{
+					SculptLayers->MergeSculptLayers(MergeLayerStart, EndLayer, bUseWeights);
+				}
+			}
+			else
+			{
+				UE::Geometry::AppendError(Debug, EGeometryScriptErrorType::InvalidInputs, LOCTEXT("MergeSculptLayers_LayerNotFound", "MergeSculptLayers: Requested merge start layer not found on TargetMesh."));
+			}
+		}
+	}, EDynamicMeshChangeType::GeneralEdit, EDynamicMeshAttributeChangeFlags::Unknown, false);
+
+	return TargetMesh;
+}
+
+
+#undef LOCTEXT_NAMESPACE
diff --git a/Engine/Plugins/Runtime/GeometryScripting/Source/GeometryScriptingCore/Public/GeometryScript/MeshSculptLayersFunctions.h b/Engine/Plugins/Runtime/GeometryScripting/Source/GeometryScriptingCore/Public/GeometryScript/MeshSculptLayersFunctions.h
new file mode 100644
--- /dev/null
+++ b/Engine/Plugins/Runtime/GeometryScripting/Source/GeometryScriptingCore/Public/GeometryScript/MeshSculptLayersFunctions.h
@@ -0,0 +1,102 @@
+// Copyright Epic Games, Inc. All Rights Reserved.
+
+#pragma once
+
+#include "CoreMinimal.h"
+#include "Kismet/BlueprintFunctionLibrary.h"
+#include "GeometryScript/GeometryScriptTypes.h"
+#include "GeometryScript/GeometryScriptSelectionTypes.h"
+#include "MeshSculptLayersFunctions.generated.h"
+
+class UDynamicMesh;
+
+USTRUCT(BlueprintType)
+struct GEOMETRYSCRIPTINGCORE_API FGeometryScriptSculptLayerUpdateOptions
+{
+	GENERATED_BODY()
+public:
+	/** Whether to recompute normals when changes to sculpt layers may have changed the mesh vertex positions */
+	UPROPERTY(EditAnywhere, BlueprintReadWrite, Category = Options)
+	bool bRecomputeNormals = true;
+};
+
+
+UCLASS(meta = (ScriptName = "GeometryScript_MeshSculptLayers"))
+class GEOMETRYSCRIPTINGCORE_API UGeometryScriptLibrary_MeshSculptLayersFunctions : public UBlueprintFunctionLibrary
+{
+	GENERATED_BODY()
+public:
+
+	/// Enable Sculpt Layers on the Target Mesh, if not already enabled, with at least the requested number of layers. Note: Will never decrease the number of enabled layers.
+	/// @param NumLayers The amount of layers that must be enabled on the mesh. Note: If this many layers (or more) are already enabled, the mesh will not be changed.
+	UFUNCTION(BlueprintCallable, Category = "GeometryScript|SculptLayers", meta = (ScriptMethod))
+	static UPARAM(DisplayName = "Target Mesh") UDynamicMesh*
+	EnableSculptLayers(
+		UDynamicMesh* TargetMesh,
+		int32 NumLayers,
+		UGeometryScriptDebug* Debug = nullptr);
+
+	/// Set the requested LayerIndex as the current active sculpt layer, if possible.
+	UFUNCTION(BlueprintCallable, Category = "GeometryScript|SculptLayers", meta = (ScriptMethod))
+	static UPARAM(DisplayName = "Target Mesh") UDynamicMesh*
+	SetActiveSculptLayer(
+		UDynamicMesh* TargetMesh,
+		int32 LayerIndex,
+		UGeometryScriptDebug* Debug = nullptr);
+
+	/// Set the weight of the layer at LayerIndex to the requested Weight
+	UFUNCTION(BlueprintCallable, Category = "GeometryScript|SculptLayers", meta = (ScriptMethod))
+	static UPARAM(DisplayName = "Target Mesh") UDynamicMesh*
+	SetSculptLayerWeight(
+		UDynamicMesh* TargetMesh,
+		int32 LayerIndex,
+		double Weight,
+		FGeometryScriptSculptLayerUpdateOptions Options = FGeometryScriptSculptLayerUpdateOptions(),
+		UGeometryScriptDebug* Debug = nullptr);
+
+	/// Set the weights of multiple layers to match the given Weights array.
+	/// Note: If the Weights array length is larger than the number of layers, will just set the weights of the existing layers. Will not add or remove layers.
+	UFUNCTION(BlueprintCallable, Category = "GeometryScript|SculptLayers", meta = (ScriptMethod))
+	static UPARAM(DisplayName = "Target Mesh") UDynamicMesh*
+	SetSculptLayerWeightsArray(
+		UDynamicMesh* TargetMesh,
+		TArray<double> Weights,
+		FGeometryScriptSculptLayerUpdateOptions Options = FGeometryScriptSculptLayerUpdateOptions(),
+		UGeometryScriptDebug* Debug = nullptr);
+
+	/// Get the weights of all sculpt layers on the mesh
+	UFUNCTION(BlueprintCallable, Category = "GeometryScript|SculptLayers", meta = (ScriptMethod))
+	static UPARAM(DisplayName = "Weights")
+	TArray<double> GetSculptLayerWeightsArray(UDynamicMesh* TargetMesh);
+
+	/// Get the number of sculpting layers active on the mesh
+	UFUNCTION(BlueprintPure, Category = "GeometryScript|SculptLayers")
+	static UPARAM(DisplayName = "Num Layers") int32 GetNumSculptLayers(const UDynamicMesh* TargetMesh);
+
+	/// Get the current sculpting layers active on the mesh, or -1 if the mesh does not have sculpting layers
+	UFUNCTION(BlueprintPure, Category = "GeometryScript|SculptLayers")
+	static UPARAM(DisplayName = "Layer Index") int32 GetActiveSculptLayer(const UDynamicMesh* TargetMesh);
+
+	/// Discard all sculpt layer data, leaving current vertex positions unchanged
+	UFUNCTION(BlueprintCallable, Category = "GeometryScript|SculptLayers", meta = (ScriptMethod, Keywords="Delete Remove"))
+	static UPARAM(DisplayName = "Target Mesh") UDynamicMesh*
+	DiscardSculptLayers(
+		UDynamicMesh* TargetMesh,
+		UGeometryScriptDebug* Debug = nullptr);
+
+	/// Merge a range of sculpt layers together. May change the Active Sculpt Layer.
+	/// @param OutActiveLayer The active layer after layers have been merged.
+	/// @param MergeLayerStart The first layer to merge
+	/// @param MergeLayerNum The number of layers to merge. Note: If larger than the number of existing layers after MergeLayerStart, will be clamped as needed.
+	/// @param bUseWeights If true, layers will be merged based on their current weights -- keeping the mesh shape unchanged. Otherwise, layers will be merged with weight 1, which may change the shape of the mesh.
+	UFUNCTION(BlueprintCallable, Category = "GeometryScript|SculptLayers", meta = (ScriptMethod, Keywords = "Combine"))
+	static UPARAM(DisplayName = "Target Mesh") UDynamicMesh*
+	MergeSculptLayers(
+		UDynamicMesh* TargetMesh,
+		int32& OutActiveLayer,
+		int32 MergeLayerStart,
+		int32 MergeLayerNum,
+		bool bUseWeights = true,
+		UGeometryScriptDebug* Debug = nullptr);
+
+};
\ No newline at end of file
diff --git a/Engine/Source/Runtime/GeometryCore/Private/DynamicMesh/DynamicMeshSculptLayers.cpp b/Engine/Source/Runtime/GeometryCore/Private/DynamicMesh/DynamicMeshSculptLayers.cpp
--- a/Engine/Source/Runtime/GeometryCore/Private/DynamicMesh/DynamicMeshSculptLayers.cpp
+++ b/Engine/Source/Runtime/GeometryCore/Private/DynamicMesh/DynamicMeshSculptLayers.cpp
@@ -284,7 +284,7 @@ bool FDynamicMeshSculptLayers::ValidateActiveLayer()
 		ActiveLayer = INDEX_NONE;
 		return false;
 	}
-	ActiveLayer = FMath::Clamp(ActiveLayer, 0, LayerWeights.Num());
+	ActiveLayer = FMath::Clamp(ActiveLayer, 0, LayerWeights.Num() - 1);
 
 	// If active layer has zero weight, try to find an active layer with non-zero weight instead
 	if (LayerWeights[ActiveLayer] == 0)
