GeometryScript: Fixed CreateVolumeFromMesh in editor to use actor factories to ensure that the Brush is correctly assigned to the resulting actor.

diff --git a/Engine/Plugins/Runtime/GeometryScripting/Source/GeometryScriptingEditor/Private/CreateNewAssetUtilityFunctions.cpp b/Engine/Plugins/Runtime/GeometryScripting/Source/GeometryScriptingEditor/Private/CreateNewAssetUtilityFunctions.cpp
--- a/Engine/Plugins/Runtime/GeometryScripting/Source/GeometryScriptingEditor/Private/CreateNewAssetUtilityFunctions.cpp
+++ b/Engine/Plugins/Runtime/GeometryScripting/Source/GeometryScriptingEditor/Private/CreateNewAssetUtilityFunctions.cpp
@@ -36,6 +36,9 @@
 #include "ModelingObjectsCreationAPI.h"
 #include "Engine/SkinnedAssetCommon.h"
 
+#include "ActorFactories/ActorFactory.h"
+#include "AssetSelection.h"
+
 #include UE_INLINE_GENERATED_CPP_BY_NAME(CreateNewAssetUtilityFunctions)
 
 using namespace UE::Geometry;
@@ -135,9 +138,22 @@ AVolume* UGeometryScriptLibrary_CreateNewAssetFunctions::CreateNewVolumeFromMesh
 
 	GEditor->BeginTransaction(LOCTEXT("CreateNewVolumeFromMesh_Transaction", "Create Volume"));
 
-	FActorSpawnParameters SpawnInfo;
 	FTransform NewActorTransform = FTransform::Identity;
-	AVolume* NewVolumeActor = (AVolume*)CreateInWorld->SpawnActor(VolumeClass, &NewActorTransform, SpawnInfo);
+	AVolume* const NewVolumeActor = [VolumeClass, CreateInWorld, &NewActorTransform]() -> AVolume*
+	{
+		if (UActorFactory* const VolumeFactory = FActorFactoryAssetProxy::GetFactoryForAssetObject(VolumeClass))
+		{
+			AActor* const Actor = VolumeFactory->CreateActor(VolumeClass, CreateInWorld->GetCurrentLevel(), NewActorTransform);
+			return Cast<AVolume>(Actor);
+		}
+		return nullptr;
+	}();
+
+	if (!NewVolumeActor)
+	{
+		UE::Geometry::AppendError(Debug, EGeometryScriptErrorType::InvalidInputs, LOCTEXT("CreateNewVolumeFromMesh_FailedSpawn", "CreateNewVolumeFromMesh: Failed to spawn new volume actor."));
+		return nullptr;
+	}
 
 	NewVolumeActor->BrushType = EBrushType::Brush_Add;
 	UModel* Model = NewObject<UModel>(NewVolumeActor);
