Add support for applying nanite settings in Geometry Script's Skeletal Mesh asset update/creation methods

diff --git a/Engine/Plugins/Runtime/GeometryScripting/Source/GeometryScriptingCore/Private/MeshAssetFunctions.cpp b/Engine/Plugins/Runtime/GeometryScripting/Source/GeometryScriptingCore/Private/MeshAssetFunctions.cpp
--- a/Engine/Plugins/Runtime/GeometryScripting/Source/GeometryScriptingCore/Private/MeshAssetFunctions.cpp
+++ b/Engine/Plugins/Runtime/GeometryScripting/Source/GeometryScriptingCore/Private/MeshAssetFunctions.cpp
@@ -994,6 +994,11 @@ UDynamicMesh* UGeometryScriptLibrary_StaticMeshFunctions::CopyMeshToSkeletalMesh
 	SkeletalLODInfo->ReductionSettings.MaxNumOfVerts = MAX_int32; 
 	SkeletalLODInfo->ReductionSettings.BaseLOD = TargetLOD.LODIndex;
 
+	if (Options.bApplyNaniteSettings)
+	{
+		ToSkeletalMeshAsset->NaniteSettings = Options.NewNaniteSettings;
+	}
+
 	// update materials on the Asset
 	if (Options.bReplaceMaterials)
 	{
diff --git a/Engine/Plugins/Runtime/GeometryScripting/Source/GeometryScriptingEditor/Private/CreateNewAssetUtilityFunctions.cpp b/Engine/Plugins/Runtime/GeometryScripting/Source/GeometryScriptingEditor/Private/CreateNewAssetUtilityFunctions.cpp
--- a/Engine/Plugins/Runtime/GeometryScripting/Source/GeometryScriptingEditor/Private/CreateNewAssetUtilityFunctions.cpp
+++ b/Engine/Plugins/Runtime/GeometryScripting/Source/GeometryScriptingEditor/Private/CreateNewAssetUtilityFunctions.cpp
@@ -379,6 +379,8 @@ USkeletalMesh* UGeometryScriptLibrary_CreateNewAssetFunctions::CreateNewSkeletal
 	
 	AssetOptions.bEnableRecomputeNormals = Options.bEnableRecomputeNormals;
 	AssetOptions.bEnableRecomputeTangents = Options.bEnableRecomputeTangents;
+	AssetOptions.bApplyNaniteSettings = Options.bApplyNaniteSettings;
+	AssetOptions.NaniteSettings = Options.NaniteSettings;
 
 	/**
 	 * We are making a copy of each LOD mesh since UDynamicMesh can potentially be editable asynchronously in the future, 
diff --git a/Engine/Plugins/Runtime/GeometryScripting/Source/GeometryScriptingEditor/Public/GeometryScript/CreateNewAssetUtilityFunctions.h b/Engine/Plugins/Runtime/GeometryScripting/Source/GeometryScriptingEditor/Public/GeometryScript/CreateNewAssetUtilityFunctions.h
--- a/Engine/Plugins/Runtime/GeometryScripting/Source/GeometryScriptingEditor/Public/GeometryScript/CreateNewAssetUtilityFunctions.h
+++ b/Engine/Plugins/Runtime/GeometryScripting/Source/GeometryScriptingEditor/Public/GeometryScript/CreateNewAssetUtilityFunctions.h
@@ -91,9 +91,17 @@ public:
 	UPROPERTY(BlueprintReadWrite, Category = Options)
 	TMap<FName, TObjectPtr<UMaterialInterface>> Materials;
 
-	/** If true, will use the skeleton proportions (if availabale) stored in the dynamic mesh. */
+	/** If true, will use the skeleton proportions (if available) stored in the dynamic mesh. */
 	UPROPERTY(BlueprintReadWrite, Category = Options)
 	bool bUseMeshBoneProportions = false;
+
+	/** Whether to apply the provided Nanite Settings to the new Skeletal Mesh asset. If false, the default settings (nanite disabled) will be used. */
+	UPROPERTY(BlueprintReadWrite, Category = Options)
+	bool bApplyNaniteSettings = false;
+
+	/** Nanite settings to apply to the new Skeletal Mesh Asset, if bApplyNaniteSettings is true */
+	UPROPERTY(BlueprintReadWrite, Category = Options)
+	FMeshNaniteSettings NaniteSettings;
 	
 	/** Use the original vertex order found in the source data. This is useful if the inbound mesh was originally non-manifold, and needs to keep
 	 *  the non-manifold structure when re-created. */
diff --git a/Engine/Plugins/Runtime/MeshModelingToolset/Source/ModelingComponentsEditorOnly/Private/AssetUtils/CreateSkeletalMeshUtil.cpp b/Engine/Plugins/Runtime/MeshModelingToolset/Source/ModelingComponentsEditorOnly/Private/AssetUtils/CreateSkeletalMeshUtil.cpp
--- a/Engine/Plugins/Runtime/MeshModelingToolset/Source/ModelingComponentsEditorOnly/Private/AssetUtils/CreateSkeletalMeshUtil.cpp
+++ b/Engine/Plugins/Runtime/MeshModelingToolset/Source/ModelingComponentsEditorOnly/Private/AssetUtils/CreateSkeletalMeshUtil.cpp
@@ -102,6 +102,11 @@ UE::AssetUtils::ECreateSkeletalMeshResult UE::AssetUtils::CreateSkeletalMeshAsse
 		Materials.Add(FSkeletalMaterial());
 		MaterialView = Materials;
 	}
+
+	if (Options.bApplyNaniteSettings)
+	{
+		NewSkeletalMesh->NaniteSettings = Options.NaniteSettings;
+	}
 	
 	if (!FStaticToSkeletalMeshConverter::InitializeSkeletalMeshFromMeshDescriptions(
 		NewSkeletalMesh, MeshDescriptions, MaterialView, 
diff --git a/Engine/Plugins/Runtime/MeshModelingToolset/Source/ModelingComponentsEditorOnly/Public/AssetUtils/CreateSkeletalMeshUtil.h b/Engine/Plugins/Runtime/MeshModelingToolset/Source/ModelingComponentsEditorOnly/Public/AssetUtils/CreateSkeletalMeshUtil.h
--- a/Engine/Plugins/Runtime/MeshModelingToolset/Source/ModelingComponentsEditorOnly/Public/AssetUtils/CreateSkeletalMeshUtil.h
+++ b/Engine/Plugins/Runtime/MeshModelingToolset/Source/ModelingComponentsEditorOnly/Public/AssetUtils/CreateSkeletalMeshUtil.h
@@ -4,6 +4,7 @@
 
 #include "DynamicMesh/DynamicMesh3.h"
 
+#include "Engine/EngineTypes.h" // for NaniteSettings
 
 class UMaterialInterface;
 class USkeletalMesh;
@@ -65,6 +66,11 @@ namespace UE::AssetUtils
 		// Controls whether or not the RecomputeTangents option will be enabled on the Asset
 		bool bEnableRecomputeTangents = true;
 
+		// Whether to use the specified NaniteSettings when generating the Asset. If false, the default settings (nanite disabled) will be used.
+		bool bApplyNaniteSettings = false;
+		// Nanite settings to set on skeletal mesh asset, if bApplyNaniteSettings is true
+		FMeshNaniteSettings NaniteSettings = FMeshNaniteSettings();
+
 		// Whether to attempt to re-create the original non-manifold mesh from the (optional) non-manifold index map on the dynamic mesh.
 		bool bConvertBackToNonManifold = false;
 
