Geometry Scripting: Allow overriding the ref skeleton from the dynamic mesh.

#rnx

diff --git a/Engine/Plugins/Runtime/GeometryScripting/Source/GeometryScriptingCore/Private/MeshAssetFunctions.cpp b/Engine/Plugins/Runtime/GeometryScripting/Source/GeometryScriptingCore/Private/MeshAssetFunctions.cpp
--- a/Engine/Plugins/Runtime/GeometryScripting/Source/GeometryScriptingCore/Private/MeshAssetFunctions.cpp
+++ b/Engine/Plugins/Runtime/GeometryScripting/Source/GeometryScriptingCore/Private/MeshAssetFunctions.cpp
@@ -899,11 +899,24 @@ UDynamicMesh* UGeometryScriptLibrary_StaticMeshFunctions::CopyMeshToSkeletalMesh
 		GEditor->BeginTransaction(LOCTEXT("UpdateSkeletalMesh", "Update Skeletal Mesh"));
 	}
 
+	// If this option is set, override the bone hierarchy mismatch settings and warn the user that we're doing so.
+	if (Options.bRemapBoneIndicesToMatchAsset)
+	{
+		UE::Geometry::AppendWarning(Debug, EGeometryScriptErrorType::InvalidInputs,
+			LOCTEXT("RemapBoneIndicesToMatchAssetDeprecated", "The bRemapBoneIndicesToMatchAsset option is deprecated. Use BoneHierarchyMismatchHandling instead."));
+		Options.BoneHierarchyMismatchHandling = EGeometryScriptBoneHierarchyMismatchHandling::RemapGeometryToReferenceSkeleton;
+	}
+
 	// make sure transactional flag is on for this asset
 	ToSkeletalMeshAsset->SetFlags(RF_Transactional);
 
 	verify(ToSkeletalMeshAsset->Modify());
 
+	if (Options.bDeferMeshPostEditChange == false)
+	{
+		ToSkeletalMeshAsset->PreEditChange(nullptr);
+	}
+	
 	// Ensure we have enough LODInfos to cover up to the requested LOD.
 	for (int32 LODIndex = ToSkeletalMeshAsset->GetLODNum(); LODIndex <= TargetLOD.LODIndex; LODIndex++)
 	{
@@ -926,7 +939,7 @@ UDynamicMesh* UGeometryScriptLibrary_StaticMeshFunctions::CopyMeshToSkeletalMesh
 	// to root, since we can't verify that any current skin binding is valid.
 	TArray<int32> BoneRemapping;
 	const FDynamicMesh3& Mesh{FromDynamicMesh->GetMeshRef()}; 
-	if (Options.bRemapBoneIndicesToMatchAsset && Mesh.HasAttributes() && Mesh.Attributes()->HasBones())
+	if (Options.BoneHierarchyMismatchHandling == EGeometryScriptBoneHierarchyMismatchHandling::RemapGeometryToReferenceSkeleton && Mesh.HasAttributes() && Mesh.Attributes()->HasBones())
 	{
 		const FDynamicMeshBoneNameAttribute* SrcBoneNames = Mesh.Attributes()->GetBoneNames();
 		TArray<FName> DstBoneNames = ToSkeletalMeshAsset->GetRefSkeleton().GetRawRefBoneNames();
@@ -960,26 +973,77 @@ UDynamicMesh* UGeometryScriptLibrary_StaticMeshFunctions::CopyMeshToSkeletalMesh
 		Converter.Convert(&ReadMesh, *MeshDescription, !Options.bEnableRecomputeTangents);
 	});
 
-	if (Options.bRemapBoneIndicesToMatchAsset)
+	bool bForceRootBinding = false;
+	
+	switch (Options.BoneHierarchyMismatchHandling)
 	{
+	case EGeometryScriptBoneHierarchyMismatchHandling::DoNothing:
+		break;
+		
+	case EGeometryScriptBoneHierarchyMismatchHandling::RemapGeometryToReferenceSkeleton:
 		if (!BoneRemapping.IsEmpty())
 		{
 			FSkeletalMeshOperations::RemapBoneIndicesOnSkinWeightAttribute(*MeshDescription, BoneRemapping);
 		}
 		else
 		{
-			using namespace UE::AnimationCore;
-			FBoneWeight RootWeight(0, 1.0f);
-			FBoneWeights RootBinding = FBoneWeights::Create({RootWeight});
-			for (const FName AttributeName: MeshAttributes.GetSkinWeightProfileNames())
+			bForceRootBinding = true;
+		}
+		break;
+	case EGeometryScriptBoneHierarchyMismatchHandling::CreateNewReferenceSkeleton:
+		// We only grab ref skeleton from LOD 0. As of now, it is left to the user to ensure that bone hierarchies
+		// for lower LODs are a strict subset of the ref skeleton (either being created). This restriction may be
+		// lifted in the future and automatic fixing/rejection performed. 
+		if (TargetLOD.LODIndex == 0)
+		{
+			static const FName RootBoneName("Root");
+			FReferenceSkeleton RefSkeleton;
+
 			{
-				FSkinWeightsVertexAttributesRef SkinWeights(MeshAttributes.GetVertexSkinWeights(AttributeName));
+				// Scoped here so that the modifier's destructor can complete the construction of the ref skeleton.
+				FReferenceSkeletonModifier Modifier(RefSkeleton, nullptr);
+				
+				if (MeshAttributes.GetNumBones())
+				{
+					// For now, we assume the bone hierarchy is consistent and can construct a well-formed ref skeleton.
+					FSkeletalMeshAttributesShared::FBoneNameAttributesConstRef BoneNames(MeshAttributes.GetBoneNames());
+					FSkeletalMeshAttributesShared::FBoneParentIndexAttributesConstRef BoneParents(MeshAttributes.GetBoneParentIndices());
+					FSkeletalMeshAttributesShared::FBonePoseAttributesConstRef BonePoses(MeshAttributes.GetBonePoses());
 
-				for (FVertexID VertexID: MeshDescription->Vertices().GetElementIDs())
+					for (int32 BoneIndex = 0; BoneIndex < MeshAttributes.GetNumBones(); BoneIndex++)
+					{
+						Modifier.Add(FMeshBoneInfo(BoneNames.Get(BoneIndex), BoneNames.Get(BoneIndex).ToString(), BoneParents.Get(BoneIndex)),
+							BonePoses.Get(BoneIndex));
+					}
+				}
+				else
 				{
-					SkinWeights.Set(VertexID, RootBinding);
+					Modifier.Add(FMeshBoneInfo(RootBoneName, RootBoneName.ToString(), INDEX_NONE), FTransform::Identity);
+					bForceRootBinding = true;
 				}
 			}
+
+			ToSkeletalMeshAsset->SetRefSkeleton(RefSkeleton);
+			ToSkeletalMeshAsset->CalculateInvRefMatrices();
+		}
+		break;
+	}
+
+	if (bForceRootBinding)
+	{
+		MeshAttributes.BoneAttributes();
+		
+		using namespace UE::AnimationCore;
+		FBoneWeight RootWeight(0, 1.0f);
+		FBoneWeights RootBinding = FBoneWeights::Create({RootWeight});
+		for (const FName AttributeName: MeshAttributes.GetSkinWeightProfileNames())
+		{
+			FSkinWeightsVertexAttributesRef SkinWeights(MeshAttributes.GetVertexSkinWeights(AttributeName));
+
+			for (FVertexID VertexID: MeshDescription->Vertices().GetElementIDs())
+			{
+				SkinWeights.Set(VertexID, RootBinding);
+			}
 		}
 	}
 	
diff --git a/Engine/Plugins/Runtime/GeometryScripting/Source/GeometryScriptingCore/Public/GeometryScript/MeshAssetFunctions.h b/Engine/Plugins/Runtime/GeometryScripting/Source/GeometryScriptingCore/Public/GeometryScript/MeshAssetFunctions.h
--- a/Engine/Plugins/Runtime/GeometryScripting/Source/GeometryScriptingCore/Public/GeometryScript/MeshAssetFunctions.h
+++ b/Engine/Plugins/Runtime/GeometryScripting/Source/GeometryScriptingCore/Public/GeometryScript/MeshAssetFunctions.h
@@ -12,6 +12,30 @@ class UStaticMesh;
 class UDynamicMesh;
 class UMaterialInterface;
 
+/** A set of options on how to handle the situation where the bone hierarchy on a skeletal geometry does not match the
+  * reference skeleton on the skeletal asset being copied to. Does not apply when copying geometry to static meshes.
+  */
+UENUM(BlueprintType)
+enum class EGeometryScriptBoneHierarchyMismatchHandling : uint8
+{
+	/** Do nothing to fix the mismatch. This is dangerous and should not be used lightly. The reference skeleton and the geometry's 
+	  * bone hierarchy may not match, resulting in visual glitches during rendering. */
+	DoNothing = 0,
+	
+	/** Remap the bone bindings and bone hierarchy on the geometry to match the reference skeleton. If no bone information is present
+	  * then all vertices are mapped to the root bone. 
+	  */
+	RemapGeometryToReferenceSkeleton = 1,
+	
+	/** Generate a new reference skeleton on the skeletal mesh asset that matches the bone hierarchy of the geometry being copied in.
+	  * Note that virtual bones are not retained from the old reference skeleton.  
+	  * If no bone information is present, then a ref skeleton is created, with a single root bone at the origin, and all vertices bound
+	  * to that root bone. 
+	  * No attempt is made to ensure that this reference skeleton is compatible with the skeleton object. If this is a requirement, then
+	  * it is the user's responsibility to ensure they are.
+	  */
+	CreateNewReferenceSkeleton = 2
+};
 
 
 USTRUCT(BlueprintType)
@@ -72,10 +96,16 @@ struct GEOMETRYSCRIPTINGCORE_API FGeometryScriptCopyMeshToAssetOptions
 	UPROPERTY(EditAnywhere, BlueprintReadWrite, Category = Options)
 	bool bEnableRemoveDegenerates = false;
 
-	/** Remap the bone indices to match the asset. This requires the source mesh to have bone information present. If no bone information is present
-	 *  then all bone weights are mapped to the root.
-	 */
+	/** An option that specifies, for skeletal mesh assets, how mismatches between the existing reference skeleton on the asset, and the bone
+	  * hierarchy stored on the geometry are handled. By default, no attempt is made to resolve this mismatch.
+	  */
 	UPROPERTY(EditAnywhere, BlueprintReadWrite, Category = Options)
+	EGeometryScriptBoneHierarchyMismatchHandling BoneHierarchyMismatchHandling = EGeometryScriptBoneHierarchyMismatchHandling::DoNothing;
+	
+	/** Deprecated. Use BoneHierarchyMismatchHandling instead.
+	 */
+	// UE_DEPRECATED(5.6, "Deprecated. Use BoneHierarchyMismatchHandling instead.")
+	UPROPERTY(EditAnywhere, BlueprintReadWrite, AdvancedDisplay, Category = Deprecated)
 	bool bRemapBoneIndicesToMatchAsset = false;
 
 	/** Use the original vertex order found in the source data. This is useful if the inbound mesh was originally non-manifold, and needs to keep
