add inverse transform mesh methods to geometry script

diff --git a/Engine/Plugins/Runtime/GeometryScripting/Source/GeometryScriptingCore/Private/MeshTransformFunctions.cpp b/Engine/Plugins/Runtime/GeometryScripting/Source/GeometryScriptingCore/Private/MeshTransformFunctions.cpp
--- a/Engine/Plugins/Runtime/GeometryScripting/Source/GeometryScriptingCore/Private/MeshTransformFunctions.cpp
+++ b/Engine/Plugins/Runtime/GeometryScripting/Source/GeometryScriptingCore/Private/MeshTransformFunctions.cpp
@@ -36,6 +36,27 @@ UDynamicMesh* UGeometryScriptLibrary_MeshTransformFunctions::TransformMesh(
 	return TargetMesh;
 }
 
+UDynamicMesh* UGeometryScriptLibrary_MeshTransformFunctions::InverseTransformMesh(
+	UDynamicMesh* TargetMesh,
+	FTransform Transform,
+	bool bFixOrientationForNegativeScale,
+	UGeometryScriptDebug* Debug)
+{
+	if (TargetMesh == nullptr)
+	{
+		UE::Geometry::AppendError(Debug, EGeometryScriptErrorType::InvalidInputs, LOCTEXT("InverseTransformMesh_InvalidInput", "InverseTransformMesh: TargetMesh is Null"));
+		return TargetMesh;
+	}
+
+	// todo: publish correct change types
+	TargetMesh->EditMesh([&](FDynamicMesh3& EditMesh)
+	{
+		MeshTransforms::ApplyTransformInverse(EditMesh, (FTransformSRT3d)Transform, bFixOrientationForNegativeScale);
+
+	}, EDynamicMeshChangeType::GeneralEdit, EDynamicMeshAttributeChangeFlags::Unknown, false);
+
+	return TargetMesh;
+}
 
 
 UDynamicMesh* UGeometryScriptLibrary_MeshTransformFunctions::TranslateMesh(
@@ -143,6 +164,37 @@ UDynamicMesh* UGeometryScriptLibrary_MeshTransformFunctions::TransformMeshSelect
 	return TargetMesh;
 }
 
+UDynamicMesh* UGeometryScriptLibrary_MeshTransformFunctions::InverseTransformMeshSelection(
+	UDynamicMesh* TargetMesh,
+	FGeometryScriptMeshSelection Selection,
+	FTransform Transform,
+	UGeometryScriptDebug* Debug)
+{
+	if (TargetMesh == nullptr)
+	{
+		UE::Geometry::AppendError(Debug, EGeometryScriptErrorType::InvalidInputs, LOCTEXT("InverseTransformMeshSelection_InvalidInput", "InverseTransformMeshSelection: TargetMesh is Null"));
+		return TargetMesh;
+	}
+
+	// todo: publish correct change types
+	TargetMesh->EditMesh([&](FDynamicMesh3& EditMesh) 
+	{
+		TArray<int32> Vertices;
+		Selection.ConvertToMeshIndexArray(EditMesh, Vertices, EGeometryScriptIndexType::Vertex);
+		for ( int32 vid : Vertices)
+		{
+			if (EditMesh.IsVertex(vid))
+			{
+				EditMesh.SetVertex(vid, Transform.InverseTransformPosition( EditMesh.GetVertex(vid) ) );
+			}
+		}
+
+		// normals?
+
+	}, EDynamicMeshChangeType::GeneralEdit, EDynamicMeshAttributeChangeFlags::Unknown, false);
+
+	return TargetMesh;
+}
 
 
 UDynamicMesh* UGeometryScriptLibrary_MeshTransformFunctions::TranslateMeshSelection(
diff --git a/Engine/Plugins/Runtime/GeometryScripting/Source/GeometryScriptingCore/Public/GeometryScript/MeshTransformFunctions.h b/Engine/Plugins/Runtime/GeometryScripting/Source/GeometryScriptingCore/Public/GeometryScript/MeshTransformFunctions.h
--- a/Engine/Plugins/Runtime/GeometryScripting/Source/GeometryScriptingCore/Public/GeometryScript/MeshTransformFunctions.h
+++ b/Engine/Plugins/Runtime/GeometryScripting/Source/GeometryScriptingCore/Public/GeometryScript/MeshTransformFunctions.h
@@ -28,6 +28,17 @@ public:
 		bool bFixOrientationForNegativeScale = true,
 		UGeometryScriptDebug* Debug = nullptr);
 
+	/**
+	 * Applies the inverse of the provided Transform to the TargetMesh
+	 */
+	UFUNCTION(BlueprintCallable, Category = "GeometryScript|Transforms", meta = (ScriptMethod, AdvancedDisplay = "bFixOrientationForNegativeScale"))
+	static UPARAM(DisplayName = "Target Mesh") UDynamicMesh*
+	InverseTransformMesh(
+		UDynamicMesh* TargetMesh,
+		FTransform Transform,
+		bool bFixOrientationForNegativeScale = true,
+		UGeometryScriptDebug* Debug = nullptr);
+
 	/**
 	* Applies the provided Translation to the vertices of a Mesh.
 	*/
@@ -74,6 +85,17 @@ public:
 		FTransform Transform,
 		UGeometryScriptDebug* Debug = nullptr);
 
+	/**
+	 * Applies the inverse of the given Transform to the vertices identified by the Selection of the TargetMesh.
+	 */
+	UFUNCTION(BlueprintCallable, Category = "GeometryScript|Transforms", meta=(ScriptMethod, AdvancedDisplay = "bFixOrientationForNegativeScale"))
+	static UPARAM(DisplayName = "Target Mesh") UDynamicMesh*
+	InverseTransformMeshSelection(
+		UDynamicMesh* TargetMesh,
+		FGeometryScriptMeshSelection Selection,
+		FTransform Transform,
+		UGeometryScriptDebug* Debug = nullptr);
+
 	/**
 	* Applies the given Translation to the vertices identified by the Selection of the mesh.
 	*/ 
