add geometry script support for more boolean operations: trim inside, trim outside, new polygroup inside, new polygroup outside

diff --git a/Engine/Plugins/Runtime/GeometryScripting/Source/GeometryScriptingCore/Private/MeshBooleanFunctions.cpp b/Engine/Plugins/Runtime/GeometryScripting/Source/GeometryScriptingCore/Private/MeshBooleanFunctions.cpp
--- a/Engine/Plugins/Runtime/GeometryScripting/Source/GeometryScriptingCore/Private/MeshBooleanFunctions.cpp
+++ b/Engine/Plugins/Runtime/GeometryScripting/Source/GeometryScriptingCore/Private/MeshBooleanFunctions.cpp
@@ -55,6 +55,18 @@ UDynamicMesh* UGeometryScriptLibrary_MeshBooleanFunctions::ApplyMeshBoolean(
 	case EGeometryScriptBooleanOperation::Subtract: 
 		ApplyOperation = FMeshBoolean::EBooleanOp::Difference;
 		break;	
+	case EGeometryScriptBooleanOperation::TrimInside:
+		ApplyOperation = FMeshBoolean::EBooleanOp::TrimInside;
+		break;
+	case EGeometryScriptBooleanOperation::TrimOutside:
+		ApplyOperation = FMeshBoolean::EBooleanOp::TrimOutside;
+		break;
+	case EGeometryScriptBooleanOperation::NewPolyGroupInside:
+		ApplyOperation = FMeshBoolean::EBooleanOp::NewGroupInside;
+		break;
+	case EGeometryScriptBooleanOperation::NewPolyGroupOutside:
+		ApplyOperation = FMeshBoolean::EBooleanOp::NewGroupOutside;
+		break;
 	}
 
 	FDynamicMesh3 NewResultMesh;
diff --git a/Engine/Plugins/Runtime/GeometryScripting/Source/GeometryScriptingCore/Public/GeometryScript/MeshBooleanFunctions.h b/Engine/Plugins/Runtime/GeometryScripting/Source/GeometryScriptingCore/Public/GeometryScript/MeshBooleanFunctions.h
--- a/Engine/Plugins/Runtime/GeometryScripting/Source/GeometryScriptingCore/Public/GeometryScript/MeshBooleanFunctions.h
+++ b/Engine/Plugins/Runtime/GeometryScripting/Source/GeometryScriptingCore/Public/GeometryScript/MeshBooleanFunctions.h
@@ -14,7 +14,11 @@ enum class EGeometryScriptBooleanOperation : uint8
 {
 	Union,
 	Intersection,
-	Subtract
+	Subtract,
+	TrimInside,
+	TrimOutside,
+	NewPolyGroupInside UMETA(DisplayName = "New PolyGroup Inside"),
+	NewPolyGroupOutside UMETA(DisplayName = "New PolyGroup Outside")
 };
 
 
