add option to choose the output space for geometry script mesh booleans

diff --git a/Engine/Plugins/Runtime/GeometryScripting/Source/GeometryScriptingCore/Private/MeshBooleanFunctions.cpp b/Engine/Plugins/Runtime/GeometryScripting/Source/GeometryScriptingCore/Private/MeshBooleanFunctions.cpp
--- a/Engine/Plugins/Runtime/GeometryScripting/Source/GeometryScriptingCore/Private/MeshBooleanFunctions.cpp
+++ b/Engine/Plugins/Runtime/GeometryScripting/Source/GeometryScriptingCore/Private/MeshBooleanFunctions.cpp
@@ -98,8 +98,17 @@ UDynamicMesh* UGeometryScriptLibrary_MeshBooleanFunctions::ApplyMeshBoolean(
 		return TargetMesh;
 	}
 
-	// Boolean result is in the space of TargetTransform, so invert that
-	MeshTransforms::ApplyTransformInverse(NewResultMesh, (FTransformSRT3d)TargetTransform, true);
+	if (Options.OutputTransformSpace == EGeometryScriptBooleanOutputSpace::TargetTransformSpace)
+	{
+		// Put the result in the local space of the target mesh
+		MeshTransforms::ApplyTransformInverse(NewResultMesh, (FTransformSRT3d)TargetTransform, true);
+	}
+	else if (Options.OutputTransformSpace == EGeometryScriptBooleanOutputSpace::ToolTransformSpace)
+	{
+		// Put the result in the local space of the tool mesh
+		MeshTransforms::ApplyTransformInverse(NewResultMesh, (FTransformSRT3d)ToolTransform, true);
+	}
+	// Note EGeometryScriptBooleanOutputSpace::SharedTransformSpace does not need a transform
 
 	if (NewBoundaryEdges.Num() > 0 && Options.bFillHoles)
 	{
diff --git a/Engine/Plugins/Runtime/GeometryScripting/Source/GeometryScriptingCore/Public/GeometryScript/MeshBooleanFunctions.h b/Engine/Plugins/Runtime/GeometryScripting/Source/GeometryScriptingCore/Public/GeometryScript/MeshBooleanFunctions.h
--- a/Engine/Plugins/Runtime/GeometryScripting/Source/GeometryScriptingCore/Public/GeometryScript/MeshBooleanFunctions.h
+++ b/Engine/Plugins/Runtime/GeometryScripting/Source/GeometryScriptingCore/Public/GeometryScript/MeshBooleanFunctions.h
@@ -22,6 +22,17 @@ enum class EGeometryScriptBooleanOperation : uint8
 };
 
 
+// Options for the output coordinate space for the mesh boolean result
+UENUM(BlueprintType)
+enum class EGeometryScriptBooleanOutputSpace : uint8
+{
+	// Transform the boolean result into the local space of the target mesh
+	TargetTransformSpace,
+	// Transform the boolean result into the local space of the tool mesh
+	ToolTransformSpace,
+	// Keep the boolean result in the shared space where the boolean was computed
+	SharedTransformSpace
+};
 
 USTRUCT(BlueprintType)
 struct GEOMETRYSCRIPTINGCORE_API FGeometryScriptMeshBooleanOptions
@@ -40,6 +51,10 @@ public:
 	// Whether to allow the Mesh Boolean operation to generate an empty mesh as its result
 	UPROPERTY(EditAnywhere, BlueprintReadWrite, Category = Options)
 	bool bAllowEmptyResult = false;
+
+	// The coordinate space to use for the result mesh
+	UPROPERTY(EditAnywhere, BlueprintReadWrite, Category = Options)
+	EGeometryScriptBooleanOutputSpace OutputTransformSpace = EGeometryScriptBooleanOutputSpace::TargetTransformSpace;
 };
 
 
