Geometry Script: Remove unneeded verify check for on a UObject::Modify call.
The script nodes aren't always called in an undo block, so the call may return false in that case.
Also removed some (void) casts since Modify is no longer marked with [[nodiscard]].

diff --git a/Engine/Plugins/Runtime/GeometryScripting/Source/GeometryScriptingCore/Private/MeshAssetFunctions.cpp b/Engine/Plugins/Runtime/GeometryScripting/Source/GeometryScriptingCore/Private/MeshAssetFunctions.cpp
--- a/Engine/Plugins/Runtime/GeometryScripting/Source/GeometryScriptingCore/Private/MeshAssetFunctions.cpp
+++ b/Engine/Plugins/Runtime/GeometryScripting/Source/GeometryScriptingCore/Private/MeshAssetFunctions.cpp
@@ -910,7 +910,7 @@ UDynamicMesh* UGeometryScriptLibrary_StaticMeshFunctions::CopyMeshToSkeletalMesh
 	// make sure transactional flag is on for this asset
 	ToSkeletalMeshAsset->SetFlags(RF_Transactional);
 
-	verify(ToSkeletalMeshAsset->Modify());
+	ToSkeletalMeshAsset->Modify();
 
 	if (Options.bDeferMeshPostEditChange == false)
 	{
@@ -1166,7 +1166,7 @@ UDynamicMesh* UGeometryScriptLibrary_StaticMeshFunctions::CopyMorphTargetToSkele
 	// make sure transactional flag is on for this asset
 	ToSkeletalMeshAsset->SetFlags(RF_Transactional);
 
-	(void)ToSkeletalMeshAsset->Modify();
+	ToSkeletalMeshAsset->Modify();
 	
 	// Ensure we have enough LODInfos to cover up to the requested LOD.
 	for (int32 LODIndex = ToSkeletalMeshAsset->GetLODNum(); LODIndex <= TargetLOD.LODIndex; LODIndex++)
@@ -1394,7 +1394,7 @@ UDynamicMesh* UGeometryScriptLibrary_StaticMeshFunctions::CopySkinWeightProfileT
 	// make sure transactional flag is on for this asset
 	ToSkeletalMeshAsset->SetFlags(RF_Transactional);
 
-	(void)ToSkeletalMeshAsset->Modify();
+	ToSkeletalMeshAsset->Modify();
 
 	// Ensure we have enough LODInfos to cover up to the requested LOD.
 	for (int32 LODIndex = ToSkeletalMeshAsset->GetLODNum(); LODIndex <= TargetLOD.LODIndex; LODIndex++)
