add capsule fitting as option in convex hulls -> simple collision geometry script fitting method

diff --git a/Engine/Plugins/Runtime/GeometryScripting/Source/GeometryScriptingCore/Private/CollisionFunctions.cpp b/Engine/Plugins/Runtime/GeometryScripting/Source/GeometryScriptingCore/Private/CollisionFunctions.cpp
--- a/Engine/Plugins/Runtime/GeometryScripting/Source/GeometryScriptingCore/Private/CollisionFunctions.cpp
+++ b/Engine/Plugins/Runtime/GeometryScripting/Source/GeometryScriptingCore/Private/CollisionFunctions.cpp
@@ -22,6 +22,7 @@
 #include "ShapeApproximation/MeshSimpleShapeApproximation.h"
 #include "MinVolumeSphere3.h"
 #include "MinVolumeBox3.h"
+#include "FitCapsule3.h"
 
 #include "Clustering/FaceNormalClustering.h"
 #include "CompGeom/ConvexDecomposition3.h"
@@ -41,6 +42,7 @@
 #include "PhysicsEngine/BodySetup.h"
 
 // requires ModelingComponents
+#include "Physics/CollisionGeometryConversion.h"
 #include "Physics/PhysicsDataCollection.h"
 #include "UObject/UObjectIterator.h"
 
@@ -1074,7 +1076,7 @@ void UGeometryScriptLibrary_CollisionFunctions::ApproximateConvexHullsWithSimple
 {
 	bHasApproximated = false;
 
-	if (!ApproximateOptions.bFitBoxes && !ApproximateOptions.bFitSpheres)
+	if (!ApproximateOptions.bFitBoxes && !ApproximateOptions.bFitSpheres && !ApproximateOptions.bFitCapsules)
 	{
 		// nothing to approximate
 		return;
@@ -1114,10 +1116,11 @@ void UGeometryScriptLibrary_CollisionFunctions::ApproximateConvexHullsWithSimple
 		const FKConvexElem& Elem = ConvexElems[ConvexIdx];
 		double HullVolume = UELocal::GetConvexElemVolume(Elem);
 
-		bool bFoundBox = false, bFoundSphere = false;
+		bool bFoundBox = false, bFoundSphere = false, bFoundCapsule = false;
 		UE::Geometry::FSphere3d ApproxSphere;
 		FOrientedBox3d ApproxBox;
-		double SphereVolumeDiff = FMathd::MaxReal, BoxVolumeDiff = FMathd::MaxReal;
+		FCapsule3d ApproxCapsule;
+		double SphereVolumeDiff = FMathd::MaxReal, BoxVolumeDiff = FMathd::MaxReal, CapsuleVolumeDiff = FMathd::MaxReal;
 
 		if (ApproximateOptions.bFitSpheres)
 		{
@@ -1151,23 +1154,46 @@ void UGeometryScriptLibrary_CollisionFunctions::ApproximateConvexHullsWithSimple
 			}
 		}
 
-		if (bFoundSphere || bFoundBox)
+		if (ApproximateOptions.bFitCapsules)
+		{
+			
+			if (FFitCapsule3d::Solve(Elem.VertexData.Num(), [&Elem](int32 Idx) {return Elem.VertexData[Idx];}, ApproxCapsule))
+			{
+				double CapsuleVolume = ApproxCapsule.Volume();
+				CapsuleVolumeDiff = GetVolumeDifferenceFrac(HullVolume, CapsuleVolume);
+				bFoundCapsule = (CapsuleVolumeDiff < ApproximateOptions.VolumeDiffThreshold_Fraction)
+					&& SampleHull(Elem, [&ApproxCapsule, &ApproximateOptions](FVector V) -> bool
+					{
+						return FMath::Abs(ApproxCapsule.SignedDistance(V)) < ApproximateOptions.DistanceThreshold;
+					});
+			}
+		}
+
+		if (bFoundSphere || bFoundBox || bFoundCapsule)
 		{
 			// Add the approximating element
-			if (BoxVolumeDiff < SphereVolumeDiff)
+			switch (FMathd::Min3Index(BoxVolumeDiff, SphereVolumeDiff, CapsuleVolumeDiff))
+			{
+			case 0: // Box is better
 			{
 				FKBoxElem& ElemBox = SimpleCollision.AggGeom.BoxElems.Emplace_GetRef();
-				ElemBox.Center = ApproxBox.Center();
-				ElemBox.Rotation = (FRotator)ApproxBox.Frame.Rotation;
-				ElemBox.X = ApproxBox.Extents.X * 2;
-				ElemBox.Y = ApproxBox.Extents.Y * 2;
-				ElemBox.Z = ApproxBox.Extents.Z * 2;
+				GetFKElement(ApproxBox, ElemBox);
+				break;
 			}
-			else // sphere is better
+			case 1: // Sphere is better
 			{
 				FKSphereElem& ElemSphere = SimpleCollision.AggGeom.SphereElems.Emplace_GetRef();
-				ElemSphere.Center = ApproxSphere.Center;
-				ElemSphere.Radius = ApproxSphere.Radius;
+				GetFKElement(ApproxSphere, ElemSphere);
+				break;
+			}
+			case 2: // Capsule is better
+			{
+				FKSphylElem& ElemCapsule = SimpleCollision.AggGeom.SphylElems.Emplace_GetRef();
+				GetFKElement(ApproxCapsule, ElemCapsule);
+				break;
+			}
+			default: // Min3Index can only return 0, 1, or 2
+				checkNoEntry();
 			}
 
 			// Remove the approximated convex hull
diff --git a/Engine/Plugins/Runtime/GeometryScripting/Source/GeometryScriptingCore/Public/GeometryScript/CollisionFunctions.h b/Engine/Plugins/Runtime/GeometryScripting/Source/GeometryScriptingCore/Public/GeometryScript/CollisionFunctions.h
--- a/Engine/Plugins/Runtime/GeometryScripting/Source/GeometryScriptingCore/Public/GeometryScript/CollisionFunctions.h
+++ b/Engine/Plugins/Runtime/GeometryScripting/Source/GeometryScriptingCore/Public/GeometryScript/CollisionFunctions.h
@@ -327,6 +327,10 @@ public:
 	UPROPERTY(EditAnywhere, BlueprintReadWrite, Category = Options)
 	bool bFitBoxes = true;
 
+	/** Whether to attempt to replace convex hulls with capsules */
+	UPROPERTY(EditAnywhere, BlueprintReadWrite, Category = Options)
+	bool bFitCapsules = false;
+
 	/** Approximating shape should be at least this close to the original shape */
 	UPROPERTY(EditAnywhere, BlueprintReadWrite, Category = Options)
 	float DistanceThreshold = 5.f;
