GeometryScript: Added height map bake support to BakeTexture and BakeVertex.

diff --git a/Engine/Plugins/Runtime/GeometryScripting/Source/GeometryScriptingCore/Private/MeshBakeFunctions.cpp b/Engine/Plugins/Runtime/GeometryScripting/Source/GeometryScriptingCore/Private/MeshBakeFunctions.cpp
--- a/Engine/Plugins/Runtime/GeometryScripting/Source/GeometryScriptingCore/Private/MeshBakeFunctions.cpp
+++ b/Engine/Plugins/Runtime/GeometryScripting/Source/GeometryScriptingCore/Private/MeshBakeFunctions.cpp
@@ -17,6 +17,7 @@
 #include "Sampling/MeshPropertyMapEvaluator.h"
 #include "Sampling/MeshUVShellMapEvaluator.h"
 #include "Sampling/MeshResampleImageEvaluator.h"
+#include "Sampling/MeshHeightMapEvaluator.h"
 
 #include "DynamicMesh/MeshTransforms.h"
 
@@ -154,6 +155,11 @@ namespace GeometryScriptBakeLocals
 			}
 			break;
 		}
+		case EMeshMapEvaluatorType::Height:
+		{
+			TexType = FTexture2DBuilder::ETextureType::ColorLinear;
+			break;
+		}
 		case EMeshMapEvaluatorType::UVShell:
 		{
 			TexType = FTexture2DBuilder::ETextureType::Color;
@@ -244,6 +250,21 @@ namespace GeometryScriptBakeLocals
 		return Result;
 	}
 
+	FMeshHeightMapEvaluator::EHeightRangeMode GetHeightRangeMode(EGeometryScriptBakeHeightRangeMode RangeMode)
+	{
+		FMeshHeightMapEvaluator::EHeightRangeMode Result = FMeshHeightMapEvaluator::EHeightRangeMode::Absolute;
+		switch (RangeMode)
+		{
+		case EGeometryScriptBakeHeightRangeMode::Absolute:
+			Result = FMeshHeightMapEvaluator::EHeightRangeMode::Absolute;
+			break;
+		case EGeometryScriptBakeHeightRangeMode::RelativeBounds:
+			Result = FMeshHeightMapEvaluator::EHeightRangeMode::RelativeBounds;
+			break;
+		}
+		return Result;
+	}
+
 	bool GetMeshTangents(const FDynamicMesh3* Mesh, TSharedPtr<FMeshTangentsd, ESPMode::ThreadSafe>& Tangents)
 	{
 		if (!Tangents)
@@ -490,6 +511,15 @@ namespace GeometryScriptBakeLocals
 			Result = PropertyEval;
 			break;
 		}
+		case EGeometryScriptBakeTypes::Height:
+		{
+			FGeometryScriptBakeType_Height* HeightOptions = static_cast<FGeometryScriptBakeType_Height*>(Options.Options.Get());
+			TSharedPtr<FMeshHeightMapEvaluator> HeightEval = MakeShared<FMeshHeightMapEvaluator>();
+			HeightEval->RangeMode = GetHeightRangeMode(HeightOptions->HeightRangeMode);
+			HeightEval->Range = FInterval1f::MakeFromUnordered(HeightOptions->InnerDistance, HeightOptions->OuterDistance);
+			Result = HeightEval;
+			break;
+		}
 		case EGeometryScriptBakeTypes::UVShell:
 		{
 			FGeometryScriptBakeType_UVShell* UVShellOptions = static_cast<FGeometryScriptBakeType_UVShell*>(Options.Options.Get());
@@ -823,7 +853,11 @@ namespace GeometryScriptBakeLocals
 			auto IsValidBakeType = [&BakeVertexPrefix](EGeometryScriptBakeTypes BakeType, TArray<FGeometryScriptDebugMessage>* Debug)
 			{
 				// The None bake type is permitted for PerChannel bakes since it allows users to specify a channel that should not be overwritten.
-				const bool bIsValid = (BakeType == EGeometryScriptBakeTypes::AmbientOcclusion || BakeType == EGeometryScriptBakeTypes::Curvature || BakeType == EGeometryScriptBakeTypes::Constant || BakeType == EGeometryScriptBakeTypes::None);
+				const bool bIsValid = (BakeType == EGeometryScriptBakeTypes::AmbientOcclusion
+					|| BakeType == EGeometryScriptBakeTypes::Curvature
+					|| BakeType == EGeometryScriptBakeTypes::Height
+					|| BakeType == EGeometryScriptBakeTypes::Constant
+					|| BakeType == EGeometryScriptBakeTypes::None);
 				if (!bIsValid)
 				{
 					const FText BakeTypeName = FText::FromName(StaticEnum<EGeometryScriptBakeTypes>()->GetNameByIndex(static_cast<int32>(BakeType)));
@@ -1144,6 +1178,18 @@ FGeometryScriptBakeTypeOptions UGeometryScriptLibrary_MeshBakeFunctions::MakeBak
 	return Output;
 }
 
+FGeometryScriptBakeTypeOptions UGeometryScriptLibrary_MeshBakeFunctions::MakeBakeTypeHeight(EGeometryScriptBakeHeightRangeMode RangeMode, float InnerDistance, float OuterDistance)
+{
+	FGeometryScriptBakeTypeOptions Output;
+	Output.BakeType = EGeometryScriptBakeTypes::Height;
+	const TSharedPtr<FGeometryScriptBakeType_Height> HeightOptions = MakeShared<FGeometryScriptBakeType_Height>();
+	HeightOptions->HeightRangeMode = RangeMode;
+	HeightOptions->InnerDistance = InnerDistance;
+	HeightOptions->OuterDistance = OuterDistance;
+	Output.Options = HeightOptions;
+	return Output;
+}
+
 FGeometryScriptBakeTypeOptions UGeometryScriptLibrary_MeshBakeFunctions::MakeBakeTypeConstant(
 	float Value)
 {
diff --git a/Engine/Plugins/Runtime/GeometryScripting/Source/GeometryScriptingCore/Public/GeometryScript/MeshBakeFunctions.h b/Engine/Plugins/Runtime/GeometryScripting/Source/GeometryScriptingCore/Public/GeometryScript/MeshBakeFunctions.h
--- a/Engine/Plugins/Runtime/GeometryScripting/Source/GeometryScriptingCore/Public/GeometryScript/MeshBakeFunctions.h
+++ b/Engine/Plugins/Runtime/GeometryScripting/Source/GeometryScriptingCore/Public/GeometryScript/MeshBakeFunctions.h
@@ -77,7 +77,9 @@ enum class EGeometryScriptBakeTypes : uint8
 	/* Constant value */
 	Constant,
 	/* UV Shell */
-	UVShell                UMETA(DisplayName = "UV Shell")
+	UVShell                UMETA(DisplayName = "UV Shell"),
+	/* Height */
+	Height
 };
 
 UENUM(BlueprintType)
@@ -180,6 +182,30 @@ struct GEOMETRYSCRIPTINGCORE_API FGeometryScriptBakeType_Curvature : public FGeo
 	EGeometryScriptBakeCurvatureClampMode Clamping = EGeometryScriptBakeCurvatureClampMode::None;
 };
 
+UENUM(BlueprintType)
+enum class EGeometryScriptBakeHeightRangeMode : uint8
+{
+	/** Height range defined in absolute units in object space */
+	Absolute,
+	/** Height range defined as a ratio of the maximal bounding box axis */
+	RelativeBounds
+};
+
+struct GEOMETRYSCRIPTINGCORE_API FGeometryScriptBakeType_Height : public FGeometryScriptBakeTypes
+{
+	/**
+	 * When set to Absolute, Inner/Outer distance are interpreted as an absolute distance range in object space.
+	 * When set to RelativeBounds, Inner/Outer distance are interpreted as a ratio of the maximal bounding box axis.
+	 */
+	EGeometryScriptBakeHeightRangeMode HeightRangeMode = EGeometryScriptBakeHeightRangeMode::Absolute;
+
+	/** The absolute or relative inner distance from the target mesh, in object space units. */
+	float InnerDistance = -1.0f;
+
+	/** The absolute or relative outer distance from the target mesh, in object space units. */
+	float OuterDistance = 1.0f;
+};
+
 struct GEOMETRYSCRIPTINGCORE_API FGeometryScriptBakeType_UVShell : public FGeometryScriptBakeTypes
 {
 	/** The source UV layer to sample. */
@@ -578,6 +604,12 @@ public:
 		FLinearColor ShellColor = FLinearColor::Gray,
 		FLinearColor BackgroundColor = FLinearColor(0.0f, 0.0f, 0.0f, 0.0f));
 
+	UFUNCTION(BlueprintPure, Category = "GeometryScript|Bake")
+	static UPARAM(DisplayName="Bake Type Out") FGeometryScriptBakeTypeOptions MakeBakeTypeHeight(
+		EGeometryScriptBakeHeightRangeMode RangeMode = EGeometryScriptBakeHeightRangeMode::Absolute,
+		float InnerDistance = -1.0f,
+		float OuterDistance = 1.0f);
+
 	UFUNCTION(BlueprintPure, Category = "GeometryScript|Bake")
 	static UPARAM(DisplayName="Bake Type Out") FGeometryScriptBakeTypeOptions MakeBakeTypeConstant(
 		float Value = 0.0f);
