Add control over whether lightmap UV generation is enabled in the build settings of a target LOD when updating a static mesh via geometry script

Default to copy the existing setting, or LOD 0's setting if the target LOD doesn't exist yet

diff --git a/Engine/Plugins/Runtime/GeometryScripting/Source/GeometryScriptingCore/Private/MeshAssetFunctions.cpp b/Engine/Plugins/Runtime/GeometryScripting/Source/GeometryScriptingCore/Private/MeshAssetFunctions.cpp
--- a/Engine/Plugins/Runtime/GeometryScripting/Source/GeometryScriptingCore/Private/MeshAssetFunctions.cpp
+++ b/Engine/Plugins/Runtime/GeometryScripting/Source/GeometryScriptingCore/Private/MeshAssetFunctions.cpp
@@ -180,12 +180,43 @@ UDynamicMesh*  UGeometryScriptLibrary_StaticMeshFunctions::CopyMeshToStaticMesh(
 	// mark as modified
 	ToStaticMeshAsset->Modify();
 
-	auto ConfigureBuildSettingsFromOptions = [](FStaticMeshSourceModel& SourceModel, FGeometryScriptCopyMeshToAssetOptions& Options) -> FVector
+	// Decide whether to generate lightmap UVs by referencing the current asset settings (before they're modified below)
+	bool bShouldGenerateLightmapUVs = false;
+	if (Options.GenerateLightmapUVs == EGeometryScriptGenerateLightmapUVOptions::MatchTargetLODSetting)
+	{
+		const FStaticMeshSourceModel* UseReferenceSourceModel = nullptr;
+		if (TargetLOD.bWriteHiResSource)
+		{
+			UseReferenceSourceModel = &ToStaticMeshAsset->GetHiResSourceModel();
+		}
+		else
+		{
+			if (ToStaticMeshAsset->IsSourceModelValid(TargetLOD.LODIndex))
+			{
+				UseReferenceSourceModel = &ToStaticMeshAsset->GetSourceModel(TargetLOD.LODIndex);
+			}
+			else if (ToStaticMeshAsset->IsSourceModelValid(0))
+			{
+				UseReferenceSourceModel = &ToStaticMeshAsset->GetSourceModel(0);
+			}
+		}
+		if (UseReferenceSourceModel)
+		{
+			bShouldGenerateLightmapUVs = UseReferenceSourceModel->BuildSettings.bGenerateLightmapUVs;
+		}
+	}
+	else
+	{
+		bShouldGenerateLightmapUVs = (Options.GenerateLightmapUVs == EGeometryScriptGenerateLightmapUVOptions::GenerateLightmapUVs);
+	}
+	
+	auto ConfigureBuildSettingsFromOptions = [bShouldGenerateLightmapUVs](FStaticMeshSourceModel& SourceModel, FGeometryScriptCopyMeshToAssetOptions& Options) -> FVector
 												{
 													FMeshBuildSettings& BuildSettings = SourceModel.BuildSettings;
 													BuildSettings.bRecomputeNormals  = Options.bEnableRecomputeNormals;
 													BuildSettings.bRecomputeTangents = Options.bEnableRecomputeTangents;
 													BuildSettings.bRemoveDegenerates = Options.bEnableRemoveDegenerates;
+													BuildSettings.bGenerateLightmapUVs = bShouldGenerateLightmapUVs;
 													if (!Options.bUseBuildScale) // if we're not using build scale, set asset BuildScale to 1,1,1
 													{
 														BuildSettings.BuildScale3D = FVector::OneVector;
diff --git a/Engine/Plugins/Runtime/GeometryScripting/Source/GeometryScriptingCore/Public/GeometryScript/MeshAssetFunctions.h b/Engine/Plugins/Runtime/GeometryScripting/Source/GeometryScriptingCore/Public/GeometryScript/MeshAssetFunctions.h
--- a/Engine/Plugins/Runtime/GeometryScripting/Source/GeometryScriptingCore/Public/GeometryScript/MeshAssetFunctions.h
+++ b/Engine/Plugins/Runtime/GeometryScripting/Source/GeometryScriptingCore/Public/GeometryScript/MeshAssetFunctions.h
@@ -37,6 +37,18 @@ enum class EGeometryScriptBoneHierarchyMismatchHandling : uint8
 	CreateNewReferenceSkeleton = 2
 };
 
+/** Options to control whether lightmap UVs are generated */
+UENUM(BlueprintType)
+enum class EGeometryScriptGenerateLightmapUVOptions : uint8
+{
+	/** Match the lightmap UV generation setting of the target LOD, if it exists. For a new LOD, match LOD 0. */
+	MatchTargetLODSetting UMETA(DisplayName = "Match Target LOD Setting"),
+	/** Generate lightmap UVs */
+	GenerateLightmapUVs UMETA(DisplayName = "Generate Lightmap UVs"),
+	/** Do not generate lightmap UVs */
+	DoNotGenerateLightmapUVs UMETA(DisplayName = "Do Not Generate Lightmap UVs")
+};
+
 
 USTRUCT(BlueprintType)
 struct GEOMETRYSCRIPTINGCORE_API FGeometryScriptCopyMeshFromAssetOptions
@@ -121,6 +133,10 @@ struct GEOMETRYSCRIPTINGCORE_API FGeometryScriptCopyMeshToAssetOptions
 	UPROPERTY(EditAnywhere, BlueprintReadWrite, Category = Options)
 	bool bReplaceMaterials = false;
 
+	// Whether to generate lightmap UVs
+	UPROPERTY(EditAnywhere, BlueprintReadWrite, Category = Options)
+	EGeometryScriptGenerateLightmapUVOptions GenerateLightmapUVs = EGeometryScriptGenerateLightmapUVOptions::MatchTargetLODSetting;
+
 	// New materials to set if Replace Materials is enabled. Ignored otherwise.
 	UPROPERTY(EditAnywhere, BlueprintReadWrite, Category = Options)
 	TArray<TObjectPtr<UMaterialInterface>> NewMaterials;
