add tool names to geometry script mesh solidify / mesh morphology nodes (i.e., voxel wrap and voxel offset), to help users find them when coming from the modeling mode

diff --git a/Engine/Plugins/Runtime/GeometryScripting/Source/GeometryScriptingCore/Public/GeometryScript/MeshVoxelFunctions.h b/Engine/Plugins/Runtime/GeometryScripting/Source/GeometryScriptingCore/Public/GeometryScript/MeshVoxelFunctions.h
--- a/Engine/Plugins/Runtime/GeometryScripting/Source/GeometryScriptingCore/Public/GeometryScript/MeshVoxelFunctions.h
+++ b/Engine/Plugins/Runtime/GeometryScripting/Source/GeometryScriptingCore/Public/GeometryScript/MeshVoxelFunctions.h
@@ -122,7 +122,7 @@ public:
 	/**
 	* Replaces the mesh with a voxelized-and-meshed approximation (VoxWrap operation).
 	*/
-	UFUNCTION(BlueprintCallable, Category = "GeometryScript|Voxel", meta=(ScriptMethod, Keywords="Vox Wrap"))
+	UFUNCTION(BlueprintCallable, Category = "GeometryScript|Voxel", meta = (ScriptMethod, Keywords = "Vox Wrap", DisplayName = "Apply Mesh Solidify (Voxel Wrap)"))
 	static UPARAM(DisplayName = "Target Mesh") UDynamicMesh*
 	ApplyMeshSolidify(
 		UDynamicMesh* TargetMesh,
@@ -130,9 +130,9 @@ public:
 		UGeometryScriptDebug* Debug = nullptr);
 	
 	/**
-	* Replaces the mesh with an SDF-based offset mesh approximation.
+	* Replaces the mesh with an SDF-based offset mesh approximation (VoxOffset operation).
 	*/
-	UFUNCTION(BlueprintCallable, Category = "GeometryScript|Voxel", meta=(ScriptMethod, Keywords="Vox Offset Dilate Contract"))
+	UFUNCTION(BlueprintCallable, Category = "GeometryScript|Voxel", meta=(ScriptMethod, Keywords="Vox Offset Dilate Contract", DisplayName = "Apply Mesh Morphology (Voxel Offset)"))
 	static UPARAM(DisplayName = "Target Mesh") UDynamicMesh*
 	ApplyMeshMorphology(
 		UDynamicMesh* TargetMesh,
