Add a vertex-overlap-based split mesh method to geometry script, matching the one in modeling tools

diff --git a/Engine/Plugins/Runtime/GeometryScripting/Source/GeometryScriptingCore/Private/MeshDecompositionFunctions.cpp b/Engine/Plugins/Runtime/GeometryScripting/Source/GeometryScriptingCore/Private/MeshDecompositionFunctions.cpp
--- a/Engine/Plugins/Runtime/GeometryScripting/Source/GeometryScriptingCore/Private/MeshDecompositionFunctions.cpp
+++ b/Engine/Plugins/Runtime/GeometryScripting/Source/GeometryScriptingCore/Private/MeshDecompositionFunctions.cpp
@@ -5,6 +5,7 @@
 #include "DynamicMesh/DynamicMesh3.h"
 #include "DynamicMesh/DynamicMeshAttributeSet.h"
 #include "Selections/MeshConnectedComponents.h"
+#include "VertexConnectedComponents.h"
 #include "Polygroups/PolygroupSet.h"
 #include "DynamicMeshEditor.h"
 #include "UDynamicMesh.h"
@@ -88,6 +89,39 @@ UDynamicMesh* UGeometryScriptLibrary_MeshDecompositionFunctions::SplitMeshByComp
 	return TargetMesh;
 }
 
+UDynamicMesh* UGeometryScriptLibrary_MeshDecompositionFunctions::SplitMeshByVertexOverlap(  
+	UDynamicMesh* TargetMesh, 
+	TArray<UDynamicMesh*>& ComponentMeshes,
+	UDynamicMeshPool* MeshPool,
+	double ConnectVerticesThreshold,
+	UGeometryScriptDebug* Debug)
+{
+	if (TargetMesh == nullptr)
+	{
+		UE::Geometry::AppendError(Debug, EGeometryScriptErrorType::InvalidInputs, LOCTEXT("SplitMeshByVertexOverlap_InvalidInput", "SplitMeshByVertexOverlap: TargetMesh is Null"));
+		return TargetMesh;
+	}
+	
+	double UseThreshold = FMath::Max(ConnectVerticesThreshold, UE_DOUBLE_KINDA_SMALL_NUMBER);
+
+	TArray<FDynamicMesh3> SplitMeshes;
+
+	TargetMesh->ProcessMesh([&SplitMeshes, UseThreshold](const FDynamicMesh3& EditMesh) -> void
+	{
+		FVertexConnectedComponents Components(EditMesh.MaxVertexID());
+		Components.ConnectTriangles(EditMesh);
+		Components.ConnectCloseVertices(EditMesh, UseThreshold, 2);
+		FDynamicMeshEditor::SplitMesh(&EditMesh, SplitMeshes, [&Components, &EditMesh](int32 TID) -> int32
+		{
+			return Components.GetComponent(EditMesh.GetTriangle(TID).A);
+		});
+	});
+
+	BuildNewDynamicMeshes(TargetMesh, MeshPool, SplitMeshes, ComponentMeshes);
+
+	return TargetMesh;
+}
+
 
 UDynamicMesh* UGeometryScriptLibrary_MeshDecompositionFunctions::SplitMeshByMaterialIDs(  
 	UDynamicMesh* TargetMesh, 
diff --git a/Engine/Plugins/Runtime/GeometryScripting/Source/GeometryScriptingCore/Public/GeometryScript/MeshDecompositionFunctions.h b/Engine/Plugins/Runtime/GeometryScripting/Source/GeometryScriptingCore/Public/GeometryScript/MeshDecompositionFunctions.h
--- a/Engine/Plugins/Runtime/GeometryScripting/Source/GeometryScriptingCore/Public/GeometryScript/MeshDecompositionFunctions.h
+++ b/Engine/Plugins/Runtime/GeometryScripting/Source/GeometryScriptingCore/Public/GeometryScript/MeshDecompositionFunctions.h
@@ -23,7 +23,7 @@ public:
 	 * @param ComponentMeshes New List of meshes is returned here
 	 * @param MeshPool New meshes in ComponentMeshes output list are allocated from this pool if it is provided (highly recommended!!)
 	 */
-	UFUNCTION(BlueprintCallable, Category = "GeometryScript|Decomposition", meta=(ScriptMethod))
+	UFUNCTION(BlueprintCallable, Category = "GeometryScript|Decomposition", meta = (ScriptMethod, Keywords = "Connected Islands"))
 	static UPARAM(DisplayName = "Target Mesh") UDynamicMesh* 
 	SplitMeshByComponents(  
 		UDynamicMesh* TargetMesh, 
@@ -31,6 +31,22 @@ public:
 		UDynamicMeshPool* MeshPool,
 		UGeometryScriptDebug* Debug = nullptr);
 
+	/**
+	 * Create a new Mesh for each vertex-connected or vertex-overlapping part of TargetMesh.
+	 * New meshes are drawn from MeshPool if it is provided, otherwise new UDynamicMesh instances are allocated
+	 * @param ComponentMeshes New List of meshes is returned here
+	 * @param MeshPool New meshes in ComponentMeshes output list are allocated from this pool if it is provided (highly recommended!!)
+	 * @param ConnectVerticesThreshold Vertices closer than this distance will be classified as part of the same component, even if they aren't connected by the mesh triangulation
+	 */
+	UFUNCTION(BlueprintCallable, Category = "GeometryScript|Decomposition", meta = (ScriptMethod, Keywords = "Islands Connected Components"))
+	static UPARAM(DisplayName = "Target Mesh") UDynamicMesh*
+	SplitMeshByVertexOverlap(
+		UDynamicMesh* TargetMesh,
+		TArray<UDynamicMesh*>& ComponentMeshes,
+		UDynamicMeshPool* MeshPool,
+		double ConnectVerticesThreshold = .001,
+		UGeometryScriptDebug* Debug = nullptr);
+
 	/**
 	 * Create a new Mesh for each MaterialID of TargetMesh.
 	 * New meshes are drawn from MeshPool if it is provided, otherwise new UDynamicMesh instances are allocated
