add option for geometry script 'voxel' functions (solidify and offset) to operate with caller-specified bounds rather than the input mesh bounds

diff --git a/Engine/Plugins/Runtime/GeometryScripting/Source/GeometryScriptingCore/Private/MeshVoxelFunctions.cpp b/Engine/Plugins/Runtime/GeometryScripting/Source/GeometryScriptingCore/Private/MeshVoxelFunctions.cpp
--- a/Engine/Plugins/Runtime/GeometryScripting/Source/GeometryScriptingCore/Private/MeshVoxelFunctions.cpp
+++ b/Engine/Plugins/Runtime/GeometryScripting/Source/GeometryScriptingCore/Private/MeshVoxelFunctions.cpp
@@ -71,14 +71,23 @@ UDynamicMesh* UGeometryScriptLibrary_MeshVoxelFunctions::ApplyMeshSolidify(
 	{
 		FDynamicMesh3& UseSourceMesh = EditMesh;
 
-		FAxisAlignedBox3d SourceMeshBounds = EditMesh.GetBounds(true);
+		const bool bUseCustomBounds = (bool)Options.CustomBounds.IsValid;
+		FAxisAlignedBox3d UseBounds = bUseCustomBounds ?
+			FAxisAlignedBox3d(Options.CustomBounds) : EditMesh.GetBounds(true);
+
+		if (bUseCustomBounds && UseBounds.Diagonal().GetMin() <= 0)
+		{
+			UE::Geometry::AppendError(Debug, EGeometryScriptErrorType::InvalidInputs, LOCTEXT("ApplyMeshSolidify_EmptyCustomBounds", "ApplyMeshSolidify: Empty custom bounds"));
+			EditMesh.Clear();
+			return;
+		}
 
 		FDynamicMesh3 ThickenedMesh;
 		if (Options.bThickenShells)
 		{
 			// thickness should be at least a cell wide so we don't end up deleting a bunch of the input surface
 			double CellSize = (Options.GridParameters.SizeMethod == EGeometryScriptGridSizingMethod::GridResolution) ?
-				(SourceMeshBounds.MaxDim() / Options.GridParameters.GridResolution) : Options.GridParameters.GridCellSize;
+				(UseBounds.MaxDim() / Options.GridParameters.GridResolution) : Options.GridParameters.GridCellSize;
 			double SafeThickness = FMathd::Max(CellSize * 2, Options.ShellThickness);
 			ThickenedMesh.Copy(UseSourceMesh, false, false, false, false);
 			FMeshNormals::QuickComputeVertexNormals(ThickenedMesh);
@@ -94,15 +103,17 @@ UDynamicMesh* UGeometryScriptLibrary_MeshVoxelFunctions::ApplyMeshSolidify(
 		TFastWindingTree<FDynamicMesh3> FastWinding(&Spatial);
 
 		TImplicitSolidify<FDynamicMesh3> Solidify(&UseSourceMesh, &Spatial, &FastWinding);
+		Solidify.bUseCustomBounds = bUseCustomBounds;
+		Solidify.CustomBounds = UseBounds;
 
 		if (Options.GridParameters.SizeMethod == EGeometryScriptGridSizingMethod::GridResolution)
 		{
 			int ClampResolution = GetClampedGridResolution(Options.GridParameters.GridResolution, TEXT("ApplyMeshSolidify"), Debug);
-			Solidify.SetCellSizeAndExtendBounds(SourceMeshBounds, Options.ExtendBounds, ClampResolution);
+			Solidify.SetCellSizeAndExtendBounds(UseBounds, Options.ExtendBounds, ClampResolution);
 		}
 		else
 		{
-			double UseCellSize = GetClampedCellSize(Options.GridParameters.GridCellSize, SourceMeshBounds, TEXT("ApplyMeshSolidify"), Debug);
+			double UseCellSize = GetClampedCellSize(Options.GridParameters.GridCellSize, UseBounds, TEXT("ApplyMeshSolidify"), Debug);
 			Solidify.ExtendBounds = Options.ExtendBounds;
 			Solidify.MeshCellSize = UseCellSize;
 		}
@@ -151,7 +162,21 @@ UDynamicMesh* UGeometryScriptLibrary_MeshVoxelFunctions::ApplyMeshMorphology(
 		TImplicitMorphology<FDynamicMesh3> Morphology;
 		Morphology.Source = &EditMesh;
 		Morphology.SourceSpatial = &Spatial;
+		
+		const bool bUseCustomBounds = (bool)Options.CustomBounds.IsValid;
+		FAxisAlignedBox3d UseBounds = bUseCustomBounds ?
+			FAxisAlignedBox3d(Options.CustomBounds) : Spatial.GetBoundingBox();
+
+		if (bUseCustomBounds && UseBounds.Diagonal().GetMin() <= 0)
+		{
+			UE::Geometry::AppendError(Debug, EGeometryScriptErrorType::InvalidInputs, LOCTEXT("ApplyMeshMorphology_EmptyCustomBounds", "ApplyMeshMorphology: Empty custom bounds"));
+			EditMesh.Clear();
+			return;
+		}
+		
 		Morphology.MorphologyOp = static_cast<TImplicitMorphology<FDynamicMesh3>::EMorphologyOp>(Options.Operation);
+		Morphology.bUseCustomBounds = bUseCustomBounds;
+		Morphology.CustomBounds = UseBounds;
 		Morphology.Distance = Options.Distance;
 		if (Morphology.Distance == 0)
 		{
@@ -164,11 +189,11 @@ UDynamicMesh* UGeometryScriptLibrary_MeshVoxelFunctions::ApplyMeshMorphology(
 		if (SDFGridParameters.SizeMethod == EGeometryScriptGridSizingMethod::GridResolution)
 		{
 			int ClampResolution = GetClampedGridResolution(SDFGridParameters.GridResolution, TEXT("ApplyMeshMorphology"), Debug);
-			Morphology.SetGridCellSize(Spatial.GetBoundingBox(), Options.Distance, ClampResolution);
+			Morphology.SetGridCellSize(UseBounds, Options.Distance, ClampResolution);
 		}
 		else
 		{
-			double UseCellSize = GetClampedCellSize(SDFGridParameters.GridCellSize, Spatial.GetBoundingBox(), TEXT("ApplyMeshMorphology"), Debug);
+			double UseCellSize = GetClampedCellSize(SDFGridParameters.GridCellSize, UseBounds, TEXT("ApplyMeshMorphology"), Debug);
 			Morphology.GridCellSize = UseCellSize;
 		}
 
@@ -177,11 +202,11 @@ UDynamicMesh* UGeometryScriptLibrary_MeshVoxelFunctions::ApplyMeshMorphology(
 		if (MeshGridParameters.SizeMethod == EGeometryScriptGridSizingMethod::GridResolution)
 		{
 			int ClampResolution = GetClampedGridResolution(MeshGridParameters.GridResolution, TEXT("ApplyMeshMorphology"), Debug);
-			Morphology.SetMeshCellSize(Spatial.GetBoundingBox(), Options.Distance, ClampResolution);
+			Morphology.SetMeshCellSize(UseBounds, Options.Distance, ClampResolution);
 		}
 		else
 		{
-			double UseCellSize = GetClampedCellSize(MeshGridParameters.GridCellSize, Spatial.GetBoundingBox(), TEXT("ApplyMeshMorphology"), Debug);
+			double UseCellSize = GetClampedCellSize(MeshGridParameters.GridCellSize, UseBounds, TEXT("ApplyMeshMorphology"), Debug);
 			Morphology.MeshCellSize = UseCellSize;
 		}
 
diff --git a/Engine/Plugins/Runtime/GeometryScripting/Source/GeometryScriptingCore/Public/GeometryScript/MeshVoxelFunctions.h b/Engine/Plugins/Runtime/GeometryScripting/Source/GeometryScriptingCore/Public/GeometryScript/MeshVoxelFunctions.h
--- a/Engine/Plugins/Runtime/GeometryScripting/Source/GeometryScriptingCore/Public/GeometryScript/MeshVoxelFunctions.h
+++ b/Engine/Plugins/Runtime/GeometryScripting/Source/GeometryScriptingCore/Public/GeometryScript/MeshVoxelFunctions.h
@@ -47,15 +47,24 @@ public:
 	UPROPERTY(EditAnywhere, BlueprintReadWrite, Category = Options)
 	FGeometryScript3DGridParameters GridParameters;
 
+	// If valid, will be used to define the region of space to operate on.
+	// Otherwise, standard bounds based on the input mesh will be computed.
+	UPROPERTY(EditAnywhere, BlueprintReadWrite, Category = Options)
+	FBox CustomBounds = FBox(EForceInit::ForceInit);
+
+	// Space with generalized winding number higher than this threshold is considered to be inside the input surface.
 	UPROPERTY(EditAnywhere, BlueprintReadWrite, Category = Options)
 	float WindingThreshold = 0.5;
 
+	// If the solid surface extends beyond the bounds provided, whether to close off the surface at that boundary or leave it open
 	UPROPERTY(EditAnywhere, BlueprintReadWrite, Category = Options)
 	bool bSolidAtBoundaries = true;
 
-	UPROPERTY(EditAnywhere, BlueprintReadWrite, Category = Options)
+	// Amount to extend bounds, applied to both min and max extents. Only applied to default input-mesh-based bounds, *not* Custom Bounds.
+	UPROPERTY(EditAnywhere, BlueprintReadWrite, Category = Options, meta = (DisplayName = "Extend Mesh Bounds"))
 	float ExtendBounds = 1.0;
 
+	// Number of search steps to take when finding the marching cubes surface vertex positions.
 	UPROPERTY(EditAnywhere, BlueprintReadWrite, Category = Options)
 	int SurfaceSearchSteps = 3;
 
@@ -102,6 +111,11 @@ public:
 	UPROPERTY(EditAnywhere, BlueprintReadWrite, Category = Options)
 	FGeometryScript3DGridParameters MeshGridParameters;
 
+	// If valid, will be used to define the region of space to operate on.
+	// Otherwise, standard bounds based on the input mesh will be computed.
+	UPROPERTY(EditAnywhere, BlueprintReadWrite, Category = Options)
+	FBox CustomBounds = FBox(EForceInit::ForceInit);
+
 	UPROPERTY(EditAnywhere, BlueprintReadWrite, Category = Options)
 	EGeometryScriptMorphologicalOpType Operation = EGeometryScriptMorphologicalOpType::Dilate;
 
diff --git a/Engine/Source/Runtime/GeometryCore/Public/Implicit/Morphology.h b/Engine/Source/Runtime/GeometryCore/Public/Implicit/Morphology.h
--- a/Engine/Source/Runtime/GeometryCore/Public/Implicit/Morphology.h
+++ b/Engine/Source/Runtime/GeometryCore/Public/Implicit/Morphology.h
@@ -56,6 +56,12 @@ public:
 	// size of the cells used when meshing the output (marching cubes' cube size)
 	double MeshCellSize = 1.0;
 
+	/** Whether to use a custom bounding box instead of the input mesh bounds to define the domain to solidify */
+	bool bUseCustomBounds = false;
+
+	/** Custom bounds to use, if bUseCustomBounds == true; ignored otherwise */
+	FAxisAlignedBox3d CustomBounds;
+
 	// Set cell sizes to hit the target voxel counts along the max dimension of the bounds
 	void SetCellSizesAndDistance(FAxisAlignedBox3d Bounds, double DistanceIn, int TargetInputVoxelCount, int TargetOutputVoxelCount)
 	{
@@ -67,13 +73,13 @@ public:
 	// Set input grid cell size to hit the target voxel counts along the max dimension of the bounds
 	void SetGridCellSize(FAxisAlignedBox3d Bounds, double DistanceIn, int TargetInputVoxelCount)
 	{
-		GridCellSize = (Bounds.MaxDim() + Distance * 2.0) / double(TargetInputVoxelCount);
+		GridCellSize = (Bounds.MaxDim() + (!bUseCustomBounds ? DistanceIn * 2.0 : 0.0)) / double(TargetInputVoxelCount);
 	}
 
 	// Set output meshing cell size to hit the target voxel counts along the max dimension of the bounds
 	void SetMeshCellSize(FAxisAlignedBox3d Bounds, double DistanceIn, int TargetInputVoxelCount)
 	{
-		MeshCellSize = (Bounds.MaxDim() + Distance * 2.0) / double(TargetInputVoxelCount);
+		MeshCellSize = (Bounds.MaxDim() + (!bUseCustomBounds ? DistanceIn * 2.0 : 0.0)) / double(TargetInputVoxelCount);
 	}
 
 	/** if this function returns true, we should abort calculation */
@@ -165,11 +171,18 @@ protected:
 		TTriLinearGridInterpolant<MeshSDFType> Interpolant = ComputedSDF.MakeInterpolant();
 
 		MarchingCubes.IsoValue = SignedOffset;
-		MarchingCubes.Bounds = SourceSpatial->GetBoundingBox();
-		MarchingCubes.Bounds.Expand(GridCellSize);
-		if (MarchingCubes.IsoValue < 0)
+		if (bUseCustomBounds)
 		{
-			MarchingCubes.Bounds.Expand(ComputedSDF.NarrowBandMaxDistance);
+			MarchingCubes.Bounds = CustomBounds;
+		}
+		else
+		{
+			MarchingCubes.Bounds = SourceSpatial->GetBoundingBox();
+			MarchingCubes.Bounds.Expand(GridCellSize);
+			if (MarchingCubes.IsoValue < 0)
+			{
+				MarchingCubes.Bounds.Expand(ComputedSDF.NarrowBandMaxDistance);
+			}
 		}
 		MarchingCubes.RootMode = ERootfindingModes::SingleLerp;
 		MarchingCubes.CubeSize = MeshCellSize;
@@ -218,7 +231,10 @@ protected:
 		SecondSDF.Spatial = nullptr;
 
 		FAxisAlignedBox3d Bounds = MarchingCubes.Bounds;
-		Bounds.Expand(MeshCellSize); // (because mesh may spill one cell over bounds)
+		if (!bUseCustomBounds)
+		{
+			Bounds.Expand(MeshCellSize); // (because mesh may spill one cell over bounds)
+		}
 
 		SecondSDF.NarrowBandMaxDistance = UnsignedOffset + SecondSDF.CellSize;
 		SecondSDF.ExactBandWidth = FMath::CeilToInt32(SecondSDF.NarrowBandMaxDistance / SecondSDF.CellSize);
@@ -229,7 +245,10 @@ protected:
 			SecondSpatial.Build();
 			SecondSDF.Spatial = &SecondSpatial;
 			SecondSDF.ComputeMode = MeshSDFType::EComputeModes::NarrowBand_SpatialFloodFill;
-			Bounds = SecondSpatial.GetBoundingBox(); // Use the tighter bounds from the AABB tree since we have it
+			if (!bUseCustomBounds)
+			{
+				Bounds = SecondSpatial.GetBoundingBox(); // Use the tighter bounds from the AABB tree since we have it
+			}
 		}
 		else
 		{
@@ -251,10 +270,13 @@ protected:
 		MarchingCubes.Reset();
 		MarchingCubes.IsoValue = SignedOffset;
 		MarchingCubes.Bounds = Bounds;
-		MarchingCubes.Bounds.Expand(GridCellSize);
-		if (MarchingCubes.IsoValue < 0)
+		if (!bUseCustomBounds)
 		{
-			MarchingCubes.Bounds.Expand(NarrowBandMaxDistance);
+			MarchingCubes.Bounds.Expand(GridCellSize);
+			if (MarchingCubes.IsoValue < 0)
+			{
+				MarchingCubes.Bounds.Expand(NarrowBandMaxDistance);
+			}
 		}
 
 
diff --git a/Engine/Source/Runtime/GeometryCore/Public/Implicit/Solidify.h b/Engine/Source/Runtime/GeometryCore/Public/Implicit/Solidify.h
--- a/Engine/Source/Runtime/GeometryCore/Public/Implicit/Solidify.h
+++ b/Engine/Source/Runtime/GeometryCore/Public/Implicit/Solidify.h
@@ -39,7 +39,13 @@ public:
 	/** Inside/outside winding number threshold */
 	double WindingThreshold = .5;
 
-	/** How much to extend bounds considered by marching cubes outside the original surface bounds */
+	/** Whether to use a custom bounding box instead of the input mesh bounds to define the domain to solidify */
+	bool bUseCustomBounds = false;
+
+	/** Custom bounds to use, if bUseCustomBounds == true; ignored otherwise */
+	FAxisAlignedBox3d CustomBounds;
+
+	/** How much to extend bounds considered by marching cubes outside the original surface bounds. Ignored if bUseCustomBounds == true. */
 	double ExtendBounds = 1;
 
 	/** What to do if the surface extends outside the marching cubes bounds -- if true, puts a solid surface at the boundary */
@@ -56,7 +62,7 @@ public:
 	 */
 	void SetCellSizeAndExtendBounds(FAxisAlignedBox3d Bounds, double ExtendBoundsIn, int TargetOutputVoxelCount)
 	{
-		ExtendBounds = ExtendBoundsIn;
+		ExtendBounds = !bUseCustomBounds ? ExtendBoundsIn : 0.0;
 		MeshCellSize = (Bounds.MaxDim() + ExtendBounds * 2.0) / double(TargetOutputVoxelCount);
 	}
 
@@ -95,12 +101,15 @@ public:
 			return MarchingCubes;
 		}
 
-		FAxisAlignedBox3d InternalBounds = SourceSpatial->GetBoundingBox();
-		InternalBounds.Expand(ExtendBounds);
+		FAxisAlignedBox3d UseBounds = bUseCustomBounds ? CustomBounds : SourceSpatial->GetBoundingBox();
+		if (!bUseCustomBounds)
+		{
+			UseBounds.Expand(ExtendBounds);
+		}
 		
 		MarchingCubes.CubeSize = MeshCellSize;
 
-		MarchingCubes.Bounds = InternalBounds;
+		MarchingCubes.Bounds = UseBounds;
 		// expand marching cubes bounds beyond the 'internal' bounds to ensure we sample outside the bounds, if solid-at-boundaries is requested
 		if (bSolidAtBoundaries)
 		{
@@ -114,9 +123,9 @@ public:
 
 		if (bSolidAtBoundaries)
 		{
-			MarchingCubes.Implicit = [this, InternalBounds](const FVector3d& Pos)
+			MarchingCubes.Implicit = [this, UseBounds](const FVector3d& Pos)
 			{
-				return InternalBounds.Contains(Pos) ? SourceWinding->FastWindingNumber(Pos) : -(WindingThreshold + 1);
+				return UseBounds.Contains(Pos) ? SourceWinding->FastWindingNumber(Pos) : -(WindingThreshold + 1);
 			};
 		}
 		else
@@ -131,7 +140,10 @@ public:
 		for ( int32 VertIdx : Source->VertexIndicesItr() )
 		{
 			FVector3d Vertex = Source->GetVertex(VertIdx);
-			// Only add vertices that are inside the spatial bounds (only vertices that are not on any triangles will be outside)
+			// Only add vertices that are inside the spatial bounds
+			// In the non-custom-bounds case, this will only filter out unreferenced vertices, so should not miss surfaces.
+			// However, in the case of custom bounds smaller than the mesh bounds, it may cause us to miss surfaces.
+			// TODO: consider adding additional sampling options to better support this case
 			if (MarchingCubes.Bounds.Contains(Vertex))
 			{
 				MCSeeds.Add(Vertex);
