add geometry script methods to sort an array of meshes, either by volume or by a second array of float values

diff --git a/Engine/Plugins/Runtime/GeometryScripting/Source/GeometryScriptingCore/Private/MeshDecompositionFunctions.cpp b/Engine/Plugins/Runtime/GeometryScripting/Source/GeometryScriptingCore/Private/MeshDecompositionFunctions.cpp
--- a/Engine/Plugins/Runtime/GeometryScripting/Source/GeometryScriptingCore/Private/MeshDecompositionFunctions.cpp
+++ b/Engine/Plugins/Runtime/GeometryScripting/Source/GeometryScriptingCore/Private/MeshDecompositionFunctions.cpp
@@ -4,6 +4,7 @@
 
 #include "DynamicMesh/DynamicMesh3.h"
 #include "DynamicMesh/DynamicMeshAttributeSet.h"
+#include "MeshQueries.h"
 #include "Selections/MeshConnectedComponents.h"
 #include "VertexConnectedComponents.h"
 #include "Polygroups/PolygroupSet.h"
@@ -205,6 +206,152 @@ UDynamicMesh* UGeometryScriptLibrary_MeshDecompositionFunctions::SplitMeshByPoly
 	return TargetMesh;
 }
 
+namespace UE::MeshDecompositionFunctionsLocals
+{
+	template<typename ValueType>
+	TArray<ValueType> ValuePerMeshArrayHelper(TArray<UDynamicMesh*>& Meshes, TFunctionRef<float(const FDynamicMesh3&)> MeshValueFn, ValueType DefaultValue = 0)
+	{
+		TArray<ValueType> Values;
+		Values.SetNumUninitialized(Meshes.Num());
+		for (int32 Idx = 0; Idx < Meshes.Num(); ++Idx)
+		{
+			float MeshValue = DefaultValue;
+			if (Meshes[Idx])
+			{
+				Meshes[Idx]->ProcessMesh([&MeshValue, &MeshValueFn](const FDynamicMesh3& Mesh) -> void
+				{
+					MeshValue = (float)MeshValueFn(Mesh);
+				});
+			}
+			Values[Idx] = MeshValue;
+		}
+		return Values;
+	}
+
+	template<typename SortType, typename ValueType>
+	void SortByValuesArrayHelper(TArray<SortType>& ToSort, const TArray<ValueType>& Values, bool bStableSort, EArraySortOrder SortOrder)
+	{
+		const int32 NumValues = Values.Num();
+		check(ToSort.Num() == NumValues);
+
+		// make a reference array of indices, that we will sort based on values
+		TArray<int32> Indices;
+		Indices.SetNumUninitialized(NumValues);
+		for (int32 Idx = 0; Idx < NumValues; ++Idx)
+		{
+			Indices[Idx] = Idx;
+		}
+
+		// sort it
+		if (!bStableSort)
+		{
+			if (SortOrder == EArraySortOrder::Ascending)
+			{
+				Indices.Sort([&Values](int32 A, int32 B) { return Values[A] < Values[B]; });
+			}
+			else
+			{
+				Indices.Sort([&Values](int32 A, int32 B) { return Values[A] > Values[B]; });
+			}
+		}
+		else
+		{
+			if (SortOrder == EArraySortOrder::Ascending)
+			{
+				Indices.StableSort([&Values](int32 A, int32 B) { return Values[A] < Values[B]; });
+			}
+			else
+			{
+				Indices.StableSort([&Values](int32 A, int32 B) { return Values[A] > Values[B]; });
+			}
+		}
+
+		// swap the ToSort array to follow the ordering of the reference array
+		for (int32 Idx = 0; Idx < NumValues; ++Idx)
+		{
+			int32 SwapFromIndex = Indices[Idx];
+			while (SwapFromIndex < Idx)
+			{
+				SwapFromIndex = Indices[SwapFromIndex];
+			}
+
+			if (SwapFromIndex != Idx)
+			{
+				Swap(ToSort[Idx], ToSort[SwapFromIndex]);
+			}
+		}
+	}
+}
+
+
+void UGeometryScriptLibrary_MeshDecompositionFunctions::
+SortMeshesByVolume(
+	UPARAM(Ref) TArray<UDynamicMesh*>& Meshes,
+	bool bStableSort,
+	EArraySortOrder SortOrder)
+{
+	using namespace UE::MeshDecompositionFunctionsLocals;
+
+	// compute mesh volumes
+	TArray<float> MeshVolumes = ValuePerMeshArrayHelper<float>(Meshes, [](const FDynamicMesh3& Mesh) -> float
+		{
+			return (float)TMeshQueries<FDynamicMesh3>::GetVolumeNonWatertight(Mesh);
+		}, 0.0f);
+	
+	// sort by volumes
+	SortByValuesArrayHelper<UDynamicMesh*, float>(Meshes, MeshVolumes, bStableSort, SortOrder);
+}
+
+void UGeometryScriptLibrary_MeshDecompositionFunctions::
+SortMeshesByArea(
+	UPARAM(Ref) TArray<UDynamicMesh*>& Meshes,
+	bool bStableSort,
+	EArraySortOrder SortOrder)
+{
+	using namespace UE::MeshDecompositionFunctionsLocals;
+
+	// compute mesh areas
+	TArray<float> MeshAreas = ValuePerMeshArrayHelper<float>(Meshes, [](const FDynamicMesh3& Mesh) -> float
+		{
+			return (float)TMeshQueries<FDynamicMesh3>::GetVolumeArea(Mesh).Y;
+		}, 0.0f);
+
+	// sort by areas
+	SortByValuesArrayHelper<UDynamicMesh*, float>(Meshes, MeshAreas, bStableSort, SortOrder);
+}
+
+void UGeometryScriptLibrary_MeshDecompositionFunctions::
+SortMeshesByBoundsVolume(
+	UPARAM(Ref) TArray<UDynamicMesh*>& Meshes,
+	bool bStableSort,
+	EArraySortOrder SortOrder)
+{
+	using namespace UE::MeshDecompositionFunctionsLocals;
+
+	// compute mesh bounds volumes
+	TArray<float> MeshBoundsVolumes = ValuePerMeshArrayHelper<float>(Meshes, [](const FDynamicMesh3& Mesh) -> float
+		{
+			return (float)Mesh.GetBounds().Volume();
+		}, 0.0f);
+
+	// sort by bounds volumes
+	SortByValuesArrayHelper<UDynamicMesh*, float>(Meshes, MeshBoundsVolumes, bStableSort, SortOrder);
+}
+
+void UGeometryScriptLibrary_MeshDecompositionFunctions::
+SortMeshesByCustomValues(
+	UPARAM(Ref) TArray<UDynamicMesh*>& Meshes,
+	const TArray<float>& ValuesToSortBy,
+	bool bStableSort,
+	EArraySortOrder SortOrder)
+{
+	if (Meshes.Num() != ValuesToSortBy.Num())
+	{
+		UE::Geometry::MakeScriptError(EGeometryScriptErrorType::InvalidInputs, LOCTEXT("SortMeshesByCustomValues_InvalidInput", "SortMeshesByCustomValues: Meshes and Values array must have same number of elements"));
+		return;
+	}
+	UE::MeshDecompositionFunctionsLocals::SortByValuesArrayHelper<UDynamicMesh*, float>(Meshes, ValuesToSortBy, bStableSort, SortOrder);
+}
 
 
 UDynamicMesh* UGeometryScriptLibrary_MeshDecompositionFunctions::GetSubMeshFromMesh(
diff --git a/Engine/Plugins/Runtime/GeometryScripting/Source/GeometryScriptingCore/Public/GeometryScript/MeshDecompositionFunctions.h b/Engine/Plugins/Runtime/GeometryScripting/Source/GeometryScriptingCore/Public/GeometryScript/MeshDecompositionFunctions.h
--- a/Engine/Plugins/Runtime/GeometryScripting/Source/GeometryScriptingCore/Public/GeometryScript/MeshDecompositionFunctions.h
+++ b/Engine/Plugins/Runtime/GeometryScripting/Source/GeometryScriptingCore/Public/GeometryScript/MeshDecompositionFunctions.h
@@ -4,6 +4,7 @@
 
 #include "CoreMinimal.h"
 #include "Kismet/BlueprintFunctionLibrary.h"
+#include "Kismet/KismetArrayLibrary.h"
 #include "GeometryScript/GeometryScriptTypes.h"
 #include "GeometryScript/GeometryScriptSelectionTypes.h"
 #include "MeshDecompositionFunctions.generated.h"
@@ -80,6 +81,67 @@ public:
 		UDynamicMeshPool* MeshPool,
 		UGeometryScriptDebug* Debug = nullptr);
 
+	/**
+	 * Sort meshes by their volume
+	 * 
+	 * Note: For meshes with open boundary, volume is computed with respect to the average vertex position.
+	 * 
+	 * @param Meshes The meshes to sort
+	 * @param bStableSort Whether to preserve ordering for meshes with the same volume
+	 * @param SortOrder Whether to sort in order of increasing or decreasing volume
+	 */
+	UFUNCTION(BlueprintCallable, Category = "GeometryScript|Decomposition")
+	static void
+	SortMeshesByVolume(
+		UPARAM(Ref) TArray<UDynamicMesh*>& Meshes,
+		bool bStableSort = false,
+		EArraySortOrder SortOrder = EArraySortOrder::Ascending);
+
+	/**
+	 * Sort meshes by their surface area
+	 *
+	 * @param Meshes The meshes to sort
+	 * @param bStableSort Whether to preserve ordering for meshes with the same surface area
+	 * @param SortOrder Whether to sort in order of increasing or decreasing surface area
+	 */
+	UFUNCTION(BlueprintCallable, Category = "GeometryScript|Decomposition")
+	static void
+	SortMeshesByArea(
+		UPARAM(Ref) TArray<UDynamicMesh*>& Meshes,
+		bool bStableSort = false,
+		EArraySortOrder SortOrder = EArraySortOrder::Ascending);
+
+	/**
+	 * Sort meshes by their axis-aligned bounding box volume
+	 *
+	 * @param Meshes The meshes to sort
+	 * @param bStableSort Whether to preserve ordering for meshes with the same bounds volume
+	 * @param SortOrder Whether to sort in order of increasing or decreasing bounds volume
+	 */
+	UFUNCTION(BlueprintCallable, Category = "GeometryScript|Decomposition")
+	static void
+	SortMeshesByBoundsVolume(
+		UPARAM(Ref) TArray<UDynamicMesh*>& Meshes,
+		bool bStableSort = false,
+		EArraySortOrder SortOrder = EArraySortOrder::Ascending);
+
+	/**
+	 * Sort meshes according to the values in a second array, which must have the same length as the Meshes array
+	 * For example, if the values array is [3, 2, 1], with Ascending Sort Order, the Meshes array would be reversed
+	 * 
+	 * @param Meshes The meshes to sort
+	 * @param ValuesToSortBy The values to reference for sort ordering, which must be one-to-one with the meshes
+	 * @param bStableSort Whether to preserve ordering for meshes with the same corresponding value
+	 * @param SortOrder Whether to sort in order of increasing or decreasing value
+	 */
+	UFUNCTION(BlueprintCallable, Category = "GeometryScript|Decomposition")
+	static void
+	SortMeshesByCustomValues(
+		UPARAM(Ref) TArray<UDynamicMesh*>& Meshes,
+		const TArray<float>& ValuesToSortBy,
+		bool bStableSort = false,
+		EArraySortOrder SortOrder = EArraySortOrder::Ascending);
+
 	/**
 	 * CopyMeshSelectionToMesh should be used instead of this function
 	 */
