GeometryScript: Expose GutterSize to BakeTexture node.

diff --git a/Engine/Plugins/Runtime/GeometryScripting/Source/GeometryScriptingCore/Private/MeshBakeFunctions.cpp b/Engine/Plugins/Runtime/GeometryScripting/Source/GeometryScriptingCore/Private/MeshBakeFunctions.cpp
--- a/Engine/Plugins/Runtime/GeometryScripting/Source/GeometryScriptingCore/Private/MeshBakeFunctions.cpp
+++ b/Engine/Plugins/Runtime/GeometryScripting/Source/GeometryScriptingCore/Private/MeshBakeFunctions.cpp
@@ -604,6 +604,8 @@ namespace GeometryScriptBakeLocals
 		Baker.SetProjectionDistance(BakeOptions.ProjectionDistance);
 		Baker.SetSamplesPerPixel(GetSamplesPerPixel(BakeOptions.SamplesPerPixel));
 		Baker.SetFilter(GetFilter(BakeOptions.FilteringType));
+		Baker.SetGutterEnabled(BakeOptions.GutterSize > 0);
+		Baker.SetGutterSize(BakeOptions.GutterSize);
 		TUniquePtr<TImageBuilder<FVector4f>> SampleFilterMask = GetSampleFilterMask(BakeOptions, BakeTexturePrefix, Debug);
 		if (SampleFilterMask)
 		{
diff --git a/Engine/Plugins/Runtime/GeometryScripting/Source/GeometryScriptingCore/Public/GeometryScript/MeshBakeFunctions.h b/Engine/Plugins/Runtime/GeometryScripting/Source/GeometryScriptingCore/Public/GeometryScript/MeshBakeFunctions.h
--- a/Engine/Plugins/Runtime/GeometryScripting/Source/GeometryScriptingCore/Public/GeometryScript/MeshBakeFunctions.h
+++ b/Engine/Plugins/Runtime/GeometryScripting/Source/GeometryScriptingCore/Public/GeometryScript/MeshBakeFunctions.h
@@ -276,6 +276,10 @@ struct GEOMETRYSCRIPTINGCORE_API FGeometryScriptBakeTextureOptions
 	UPROPERTY(EditAnywhere, BlueprintReadWrite, Category = Options)
 	EGeometryScriptBakeBitDepth BitDepth = EGeometryScriptBakeBitDepth::ChannelBits8;
 
+	/** The pixel distance to pad baked content past the UV borders. */
+	UPROPERTY(EditAnywhere, BlueprintReadWrite, Category = Options, meta=(ClampMin=0))
+	int32 GutterSize = 4;
+
 	/** Number of samples per pixel */
 	UPROPERTY(EditAnywhere, BlueprintReadWrite, Category = Options)
 	EGeometryScriptBakeSamplesPerPixel SamplesPerPixel = EGeometryScriptBakeSamplesPerPixel::Sample1;
