[GeometryScript] Set mesh description material slot names when copying a dynamic mesh to a skeletal mesh. (Matches functionality of copying to static mesh).

diff --git a/Engine/Plugins/Runtime/GeometryScripting/Source/GeometryScriptingCore/Private/MeshAssetFunctions.cpp b/Engine/Plugins/Runtime/GeometryScripting/Source/GeometryScriptingCore/Private/MeshAssetFunctions.cpp
--- a/Engine/Plugins/Runtime/GeometryScripting/Source/GeometryScriptingCore/Private/MeshAssetFunctions.cpp
+++ b/Engine/Plugins/Runtime/GeometryScripting/Source/GeometryScriptingCore/Private/MeshAssetFunctions.cpp
@@ -1115,6 +1115,16 @@ UDynamicMesh* UGeometryScriptLibrary_StaticMeshFunctions::CopyMeshToSkeletalMesh
 		SkeletalLODInfo->LODMaterialMap.Empty();
 		
 		ToSkeletalMeshAsset->SetMaterials(NewMaterials);
+
+		// Set material slot names on the mesh description
+		TPolygonGroupAttributesRef<FName> PolygonGroupImportedMaterialSlotNames = MeshAttributes.GetPolygonGroupMaterialSlotNames();
+		for (int32 SlotIdx = 0; SlotIdx < NewMaterials.Num(); ++SlotIdx)
+		{
+			if (SlotIdx < PolygonGroupImportedMaterialSlotNames.GetNumElements())
+			{
+				PolygonGroupImportedMaterialSlotNames.Set(SlotIdx, NewMaterials[SlotIdx].ImportedMaterialSlotName);
+			}
+		}
 	}
 
 	ToSkeletalMeshAsset->CommitMeshDescription(TargetLOD.LODIndex);
