GeometryScript: Fixed output param display name for GetMaxMaterialID.

diff --git a/Engine/Plugins/Runtime/GeometryScripting/Source/GeometryScriptingCore/Public/GeometryScript/MeshMaterialFunctions.h b/Engine/Plugins/Runtime/GeometryScripting/Source/GeometryScriptingCore/Public/GeometryScript/MeshMaterialFunctions.h
--- a/Engine/Plugins/Runtime/GeometryScripting/Source/GeometryScriptingCore/Public/GeometryScript/MeshMaterialFunctions.h
+++ b/Engine/Plugins/Runtime/GeometryScripting/Source/GeometryScriptingCore/Public/GeometryScript/MeshMaterialFunctions.h
@@ -20,7 +20,7 @@ class UGeometryScriptLibrary_MeshMaterialFunctions : public UBlueprintFunctionLi
 public:
 
 	UFUNCTION(BlueprintPure, Category = "GeometryScript|MeshQueries", meta=(ScriptMethod))
-	static UE_API UPARAM(DisplayName = "Enabled") int GetMaxMaterialID( UDynamicMesh* TargetMesh, bool& bHasMaterialIDs );
+	static UE_API UPARAM(DisplayName = "Max ID") int GetMaxMaterialID( UDynamicMesh* TargetMesh, bool& bHasMaterialIDs );
 
 	/**
 	* Enables per-triangle Material IDs on a mesh and initializes the values to 0. 
