hide the 'debug' pins in geometry script

diff --git a/Engine/Plugins/Runtime/GeometryScripting/Source/GeometryScriptingCore/Public/GeometryScript/CollisionFunctions.h b/Engine/Plugins/Runtime/GeometryScripting/Source/GeometryScriptingCore/Public/GeometryScript/CollisionFunctions.h
--- a/Engine/Plugins/Runtime/GeometryScripting/Source/GeometryScriptingCore/Public/GeometryScript/CollisionFunctions.h
+++ b/Engine/Plugins/Runtime/GeometryScripting/Source/GeometryScriptingCore/Public/GeometryScript/CollisionFunctions.h
@@ -375,7 +375,7 @@ public:
 	 * 
 	 * @return true on success.
 	 */
-	UFUNCTION(BlueprintCallable, Category = "GeometryScript|Collision")
+	UFUNCTION(BlueprintCallable, Category = "GeometryScript|Collision", meta = (HidePin = "Debug"))
 	static UE_API bool SetStaticMeshCustomComplexCollision(
 		UStaticMesh* StaticMeshAsset,
 		UStaticMesh* StaticMeshCollisionAsset,
@@ -387,7 +387,7 @@ public:
 	/** 
 	* Generates Simple Collision shapes for a Static Mesh Asset based on the input Dynamic Mesh.
 	*/
-	UFUNCTION(BlueprintCallable, Category = "GeometryScript|Collision")
+	UFUNCTION(BlueprintCallable, Category = "GeometryScript|Collision", meta = (HidePin = "Debug"))
 	static UE_API UPARAM(DisplayName = "Dynamic Mesh") UDynamicMesh* 
 	SetStaticMeshCollisionFromMesh(
 		UDynamicMesh* FromDynamicMesh, 
@@ -399,7 +399,7 @@ public:
 	/**
 	 * Copy the Simple Collision Geometry from the Source Component to the Static Mesh Asset.
 	 */
-	UFUNCTION(BlueprintCallable, Category = "GeometryScript|Collision")
+	UFUNCTION(BlueprintCallable, Category = "GeometryScript|Collision", meta = (HidePin = "Debug"))
 	static UE_API void 
 	SetStaticMeshCollisionFromComponent(
 		UStaticMesh* StaticMeshAsset, 
@@ -417,7 +417,7 @@ public:
 	/** 
 	* Generate Simple Collision shapes for a Dynamic Mesh Component based on the input Dynamic Mesh. 
 	*/
-	UFUNCTION(BlueprintCallable, Category = "GeometryScript|Collision")
+	UFUNCTION(BlueprintCallable, Category = "GeometryScript|Collision", meta = (HidePin = "Debug"))
 	static UE_API UPARAM(DisplayName = "Dynamic Mesh") UDynamicMesh* 
 	SetDynamicMeshCollisionFromMesh(
 		UDynamicMesh* FromDynamicMesh, 
@@ -428,7 +428,7 @@ public:
     /** 
 	* Clears Simple Collisions from the Dynamic Mesh Component. 
 	*/
-	UFUNCTION(BlueprintCallable, Category = "GeometryScript|Collision")
+	UFUNCTION(BlueprintCallable, Category = "GeometryScript|Collision", meta = (HidePin = "Debug"))
 	static UE_API void
 	ResetDynamicMeshCollision(
 		UDynamicMeshComponent* Component,
@@ -444,7 +444,7 @@ public:
 	/*
 	 * Get the simple collision from a Primitive Component
 	 */
-	UFUNCTION(BlueprintCallable, Category = "GeometryScript|Collision")
+	UFUNCTION(BlueprintCallable, Category = "GeometryScript|Collision", meta = (HidePin = "Debug"))
 	static UE_API UPARAM(DisplayName = "Simple Collision") FGeometryScriptSimpleCollision
 	GetSimpleCollisionFromComponent(
 		UPrimitiveComponent* Component,
@@ -453,7 +453,7 @@ public:
 	/*
 	 * Set the simple collision on a Dynamic Mesh Component
 	 */
-	UFUNCTION(BlueprintCallable, Category = "GeometryScript|Collision")
+	UFUNCTION(BlueprintCallable, Category = "GeometryScript|Collision", meta = (HidePin = "Debug"))
 	static UE_API void
 	SetSimpleCollisionOfDynamicMeshComponent(
 		const FGeometryScriptSimpleCollision& SimpleCollision,
@@ -464,14 +464,14 @@ public:
 	/*
 	 * Get the simple collision from a Static Mesh
 	 */
-	UFUNCTION(BlueprintCallable, Category = "GeometryScript|Collision")
+	UFUNCTION(BlueprintCallable, Category = "GeometryScript|Collision", meta = (HidePin = "Debug"))
 	static UE_API UPARAM(DisplayName = "Simple Collision") FGeometryScriptSimpleCollision
 	GetSimpleCollisionFromStaticMesh(UStaticMesh* StaticMesh, UGeometryScriptDebug* Debug = nullptr);
 
 	/*
 	 * Set the simple collision on a Static Mesh
 	 */
-	UFUNCTION(BlueprintCallable, Category = "GeometryScript|Collision")
+	UFUNCTION(BlueprintCallable, Category = "GeometryScript|Collision", meta = (HidePin = "Debug"))
 	static UE_API void
 	SetSimpleCollisionOfStaticMesh(
 		const FGeometryScriptSimpleCollision& SimpleCollision,
@@ -483,7 +483,7 @@ public:
 	/**
 	* Generate Simple Collision shapes for an input Dynamic Mesh shape
 	*/
-	UFUNCTION(BlueprintCallable, Category = "GeometryScript|Collision")
+	UFUNCTION(BlueprintCallable, Category = "GeometryScript|Collision", meta = (HidePin = "Debug"))
 	static UE_API UPARAM(DisplayName = "Simple Collision") FGeometryScriptSimpleCollision
 	GenerateCollisionFromMesh(
 		UDynamicMesh* FromDynamicMesh,
@@ -503,7 +503,7 @@ public:
 	 * Transform simple collision shapes
 	 * @param bSuccess	Indicates whether all collision shapes were accurately transformed. On failure, shapes will still be copied over and a best-effort transform will still be applied.
 	 */
-	UFUNCTION(BlueprintCallable, Category = "GeometryScript|Collision", meta = (AutoCreateRefTerm = "TransformOptions"))
+	UFUNCTION(BlueprintCallable, Category = "GeometryScript|Collision", meta = (HidePin = "Debug", AutoCreateRefTerm = "TransformOptions"))
 	static UE_API UPARAM(DisplayName = "Transformed Collision") FGeometryScriptSimpleCollision TransformSimpleCollisionShapes(
 		const FGeometryScriptSimpleCollision& SimpleCollision,
 		FTransform Transform,
@@ -515,7 +515,7 @@ public:
 	/*
 	 * Add simple collision shapes from AppendCollision to CollisionToUpdate
 	 */
-	UFUNCTION(BlueprintCallable, Category = "GeometryScript|Collision", meta = (ScriptMethod))
+	UFUNCTION(BlueprintCallable, Category = "GeometryScript|Collision", meta = (ScriptMethod, HidePin = "Debug"))
 	static UE_API void CombineSimpleCollision(
 		UPARAM(ref) FGeometryScriptSimpleCollision& CollisionToUpdate,
 		const FGeometryScriptSimpleCollision& AppendCollision,
@@ -525,7 +525,7 @@ public:
 	/*
 	 * Combine the SimpleCollisionArray collision shapes into a single SimpleCollision
 	 */
-	UFUNCTION(BlueprintCallable, Category = "GeometryScript|Collision")
+	UFUNCTION(BlueprintCallable, Category = "GeometryScript|Collision", meta = (HidePin = "Debug"))
 	static UE_API void CombineSimpleCollisionArray(
 		const TArray<FGeometryScriptSimpleCollision>& SimpleCollisionArray,
 		FGeometryScriptSimpleCollision& SimpleCollision,
@@ -539,7 +539,7 @@ public:
 	 * @param ConvexSimplifyOptions	Options controlling how convex hulls are simplified
 	 * @param bHasSimplified		Indicates whether any convex hulls were modified
 	 */
-	UFUNCTION(BlueprintCallable, Category = "GeometryScript|Collision", meta = (ScriptMethod, AutoCreateRefTerm = "SimplifyOptions"))
+	UFUNCTION(BlueprintCallable, Category = "GeometryScript|Collision", meta = (ScriptMethod, HidePin = "Debug", AutoCreateRefTerm = "SimplifyOptions"))
 	static UE_API void SimplifyConvexHulls(
 		UPARAM(ref) FGeometryScriptSimpleCollision& SimpleCollision,
 		const FGeometryScriptConvexHullSimplificationOptions& SimplifyOptions,
@@ -555,7 +555,7 @@ public:
 	 * @param ConvexSimplifyOptions	Options controlling how convex hulls are approximated
 	 * @param bHasApproximated		Indicates whether any convex hulls were replaced with simpler approximations
 	 */
-	UFUNCTION(BlueprintCallable, Category = "GeometryScript|Collision", meta = (ScriptMethod, AutoCreateRefTerm = "ApproximateOptions"))
+	UFUNCTION(BlueprintCallable, Category = "GeometryScript|Collision", meta = (ScriptMethod, HidePin = "Debug", AutoCreateRefTerm = "ApproximateOptions"))
 	static UE_API void ApproximateConvexHullsWithSimplerCollisionShapes(
 		UPARAM(ref) FGeometryScriptSimpleCollision& SimpleCollision,
 		const FGeometryScriptConvexHullApproximationOptions& ApproximateOptions,
@@ -571,7 +571,7 @@ public:
 	 * @param bHasMerged			Indicates whether any shapes have been merged
 	 * @return						Simple Collision with collision shapes merged, as allowed by settings
 	 */
-	UFUNCTION(BlueprintCallable, Category = "GeometryScript|Collision", meta = (AutoCreateRefTerm = "MergeOptions"))
+	UFUNCTION(BlueprintCallable, Category = "GeometryScript|Collision", meta = (HidePin = "Debug", AutoCreateRefTerm = "MergeOptions"))
 	static UE_API UPARAM(DisplayName = "Merged Simple Collision") FGeometryScriptSimpleCollision MergeSimpleCollisionShapes(
 		const FGeometryScriptSimpleCollision& SimpleCollision,
 		const FGeometryScriptMergeSimpleCollisionOptions& MergeOptions,
@@ -586,7 +586,7 @@ public:
 	 * @param NegativeSpaceOptions	Options controlling how the negative space is generated
 	 * @return						A set of spheres that cover the negative space of the input shape
 	 */
-	UFUNCTION(BlueprintCallable, Category = "GeometryScript|Collision", meta = (AutoCreateRefTerm = "NegativeSpaceOptions"))
+	UFUNCTION(BlueprintCallable, Category = "GeometryScript|Collision", meta = (HidePin = "Debug", AutoCreateRefTerm = "NegativeSpaceOptions"))
 	static UE_API UPARAM(DisplayName = "Negative Space") FGeometryScriptSphereCovering ComputeNegativeSpace(
 		const FGeometryScriptDynamicMeshBVH& MeshBVH,
 		const FComputeNegativeSpaceOptions& NegativeSpaceOptions,
@@ -601,7 +601,7 @@ public:
 	 * @param Options		Options controlling the convex decomposition
 	 * @return				The resulting convex hulls as simple collision shapes
 	 */
-	UFUNCTION(BlueprintCallable, Category = "GeometryScript|Collision", meta = (AutoCreateRefTerm = "Options"))
+	UFUNCTION(BlueprintCallable, Category = "GeometryScript|Collision", meta = (HidePin = "Debug", AutoCreateRefTerm = "Options"))
 	static UE_API UPARAM(DisplayName = "Convex Decomposition") FGeometryScriptSimpleCollision ComputeNavigableConvexDecomposition(
 		const UDynamicMesh* TargetMesh,
 		const FNavigableConvexDecompositionOptions& Options,
diff --git a/Engine/Plugins/Runtime/GeometryScripting/Source/GeometryScriptingCore/Public/GeometryScript/MeshAssetFunctions.h b/Engine/Plugins/Runtime/GeometryScripting/Source/GeometryScriptingCore/Public/GeometryScript/MeshAssetFunctions.h
--- a/Engine/Plugins/Runtime/GeometryScripting/Source/GeometryScriptingCore/Public/GeometryScript/MeshAssetFunctions.h
+++ b/Engine/Plugins/Runtime/GeometryScripting/Source/GeometryScriptingCore/Public/GeometryScript/MeshAssetFunctions.h
@@ -218,7 +218,7 @@ public:
 	* Check if a Static Mesh Asset has the RequestedLOD available, ie if CopyMeshFromStaticMesh will be able to
 	* succeed for the given LODType and LODIndex. 
 	*/
-	UFUNCTION(BlueprintCallable, Category = "GeometryScript|StaticMesh", meta = (ExpandEnumAsExecs = "Outcome"))
+	UFUNCTION(BlueprintCallable, Category = "GeometryScript|StaticMesh", meta = (HidePin = "Debug", ExpandEnumAsExecs = "Outcome"))
 	static UE_API bool
 	CheckStaticMeshHasAvailableLOD(
 		UStaticMesh* StaticMeshAsset, 
@@ -243,7 +243,7 @@ public:
 	* @param bUseSectionMaterials Whether to use the mesh section indices as material IDs. If true, use GetSectionMaterialListFromStaticMesh to get the corresponding materials. If false, use GetMaterialListFromStaticMesh to get the materials instead.
 	*/
 	//~ Note this V2 version adds the bUseSectionMaterials parameter
-	UFUNCTION(BlueprintCallable, Category = "GeometryScript|StaticMesh", meta = (DisplayName = "Copy Mesh From Static Mesh", ExpandEnumAsExecs = "Outcome"))
+	UFUNCTION(BlueprintCallable, Category = "GeometryScript|StaticMesh", meta = (HidePin = "Debug", DisplayName = "Copy Mesh From Static Mesh", ExpandEnumAsExecs = "Outcome"))
 	static UE_API UPARAM(DisplayName = "Dynamic Mesh") UDynamicMesh* 
 	CopyMeshFromStaticMeshV2(
 		UStaticMesh* FromStaticMeshAsset, 
@@ -259,7 +259,7 @@ public:
 	*
 	* Note that the LOD Index in RequestedLOD will be silently clamped to the available number of LODs (SourceModel or RenderData)
 	*/
-	UFUNCTION(BlueprintCallable, Category = "GeometryScript|StaticMesh", meta = (DisplayName = "Copy Mesh From Static Mesh with Section Materials", ExpandEnumAsExecs = "Outcome"))
+	UFUNCTION(BlueprintCallable, Category = "GeometryScript|StaticMesh", meta = (HidePin = "Debug", DisplayName = "Copy Mesh From Static Mesh with Section Materials", ExpandEnumAsExecs = "Outcome"))
 	static UPARAM(DisplayName = "Dynamic Mesh") UDynamicMesh*
 	CopyMeshFromStaticMesh(
 		UStaticMesh* FromStaticMeshAsset,
@@ -277,7 +277,7 @@ public:
 	*
 	* @param bUseSectionMaterials Whether to assume Dynamic Mesh material IDs are section indices in the target Static Mesh. Should match the value passed to CopyMeshFromStaticMesh. Has no effect if replacing the asset materials.
 	*/
-	UFUNCTION(BlueprintCallable, Category = "GeometryScript|StaticMesh", meta = (ExpandEnumAsExecs = "Outcome"))
+	UFUNCTION(BlueprintCallable, Category = "GeometryScript|StaticMesh", meta = (HidePin = "Debug", ExpandEnumAsExecs = "Outcome"))
 	static UE_API UPARAM(DisplayName = "Dynamic Mesh") UDynamicMesh*
 	CopyMeshToStaticMesh(
 		UDynamicMesh* FromDynamicMesh,
@@ -313,7 +313,7 @@ public:
 	*
 	* Note that the LOD Index in RequestedLOD will be silently clamped to the available number of LODs (SourceModel or RenderData)
 	*/
-	UFUNCTION(BlueprintCallable, Category = "GeometryScript|StaticMesh", meta = (ExpandEnumAsExecs = "Outcome"))
+	UFUNCTION(BlueprintCallable, Category = "GeometryScript|StaticMesh", meta = (HidePin = "Debug", ExpandEnumAsExecs = "Outcome"))
 	static UE_API void
 	GetSectionMaterialListFromStaticMesh(
 		UStaticMesh* FromStaticMeshAsset, 
@@ -333,7 +333,7 @@ public:
 	*
 	* Note that the LOD Index in RequestedLOD will be silently clamped to the available number of LODs
 	*/
-	UFUNCTION(BlueprintCallable, Category = "GeometryScript|SkeletalMesh", meta = (DisplayName = "Get LOD Material List From Skeletal Mesh", ExpandEnumAsExecs = "Outcome"))
+	UFUNCTION(BlueprintCallable, Category = "GeometryScript|SkeletalMesh", meta = (HidePin = "Debug", DisplayName = "Get LOD Material List From Skeletal Mesh", ExpandEnumAsExecs = "Outcome"))
 	static UE_API void
 	GetLODMaterialListFromSkeletalMesh(
 		USkeletalMesh* FromSkeletalMeshAsset,
@@ -348,7 +348,7 @@ public:
 	 * Get the asset materials from the static mesh asset. These will match the DynamicMesh material if CopyMeshFromStaticMesh
 	 * was used to create a DynamicMesh with bUseSectionMaterials=false
 	 */
-	UFUNCTION(BlueprintCallable, Category = "GeometryScript|StaticMesh")
+	UFUNCTION(BlueprintCallable, Category = "GeometryScript|StaticMesh", meta = (HidePin = "Debug"))
 	static UE_API void 
 	GetMaterialListFromStaticMesh(const UStaticMesh* FromStaticMeshAsset,
 		TArray<UMaterialInterface*>& MaterialList,
@@ -359,7 +359,7 @@ public:
 	 * Get the asset materials from the skeletal mesh asset.
 	 * Note: For LOD-specific materials, use GetLODMaterialListFromSkeletalMesh instead.
 	 */
-	UFUNCTION(BlueprintCallable, Category = "GeometryScript|SkeletalMesh")
+	UFUNCTION(BlueprintCallable, Category = "GeometryScript|SkeletalMesh", meta = (HidePin = "Debug"))
 	static UE_API void
 	GetMaterialListFromSkeletalMesh(const USkeletalMesh* FromSkeletalMeshAsset,
 		TArray<UMaterialInterface*>& MaterialList,
@@ -386,7 +386,7 @@ public:
 	/** 
 	* Extracts a Dynamic Mesh from a Skeletal Mesh Asset. 
 	*/
-	UFUNCTION(BlueprintCallable, Category = "GeometryScript|SkeletalMesh", meta = (ExpandEnumAsExecs = "Outcome"))
+	UFUNCTION(BlueprintCallable, Category = "GeometryScript|SkeletalMesh", meta = (HidePin = "Debug", ExpandEnumAsExecs = "Outcome"))
 	static UE_API UPARAM(DisplayName = "Dynamic Mesh") UDynamicMesh* 
 	CopyMeshFromSkeletalMesh(
 		USkeletalMesh* FromSkeletalMeshAsset, 
@@ -399,7 +399,7 @@ public:
 	/** 
 	* Updates a Skeletal Mesh Asset with new geometry and bone weights data from a Dynamic Mesh.
 	*/
-	UFUNCTION(BlueprintCallable, Category = "GeometryScript|SkeletalMesh", meta = (ExpandEnumAsExecs = "Outcome"))
+	UFUNCTION(BlueprintCallable, Category = "GeometryScript|SkeletalMesh", meta = (HidePin = "Debug", ExpandEnumAsExecs = "Outcome"))
 	static UE_API UPARAM(DisplayName = "Dynamic Mesh") UDynamicMesh* 
 	CopyMeshToSkeletalMesh(
 		UDynamicMesh* FromDynamicMesh, 
@@ -417,7 +417,7 @@ public:
 	* @param ToSkeletalMeshAsset the asset we are writing the morph target into
 	* @param MorphTargetName the name of the morph target as it will appear in the UI
 	*/
-	UFUNCTION(BlueprintCallable, Category = "GeometryScript|SkeletalMesh", meta = (ExpandEnumAsExecs = "Outcome"))
+	UFUNCTION(BlueprintCallable, Category = "GeometryScript|SkeletalMesh", meta = (HidePin = "Debug", ExpandEnumAsExecs = "Outcome"))
 	static UE_API UPARAM(DisplayName = "Dynamic Mesh") UDynamicMesh* 
 	CopyMorphTargetToSkeletalMesh(
 		UDynamicMesh* FromMorphTarget, 
@@ -436,7 +436,7 @@ public:
 	* @param TargetProfileName the name of the skin weight profile as it will appear in the UI. Leave blank for the default profile.
 	* @param SourceProfileName The name of the skin weight profile to copy from the dynamic mesh. Leave blank for the default profile.
 	*/
-	UFUNCTION(BlueprintCallable, Category = "GeometryScript|SkeletalMesh", meta = (ExpandEnumAsExecs = "Outcome"))
+	UFUNCTION(BlueprintCallable, Category = "GeometryScript|SkeletalMesh", meta = (HidePin = "Debug", ExpandEnumAsExecs = "Outcome"))
 	static UE_API UPARAM(DisplayName = "Dynamic Mesh") UDynamicMesh* 
 	CopySkinWeightProfileToSkeletalMesh(
 		UDynamicMesh* FromDynamicMesh, 
diff --git a/Engine/Plugins/Runtime/GeometryScripting/Source/GeometryScriptingCore/Public/GeometryScript/MeshBakeFunctions.h b/Engine/Plugins/Runtime/GeometryScripting/Source/GeometryScriptingCore/Public/GeometryScript/MeshBakeFunctions.h
--- a/Engine/Plugins/Runtime/GeometryScripting/Source/GeometryScriptingCore/Public/GeometryScript/MeshBakeFunctions.h
+++ b/Engine/Plugins/Runtime/GeometryScripting/Source/GeometryScriptingCore/Public/GeometryScript/MeshBakeFunctions.h
@@ -620,7 +620,7 @@ public:
 	static UE_API UPARAM(DisplayName="Bake Type Out") FGeometryScriptBakeTypeOptions MakeBakeTypeConstant(
 		float Value = 0.0f);
 	
-	UFUNCTION(BlueprintCallable, Category = "GeometryScript|Bake")
+	UFUNCTION(BlueprintCallable, Category = "GeometryScript|Bake", meta = (HidePin = "Debug"))
 	static UE_API UPARAM(DisplayName="Textures Out") TArray<UTexture2D*> BakeTexture(
 		UDynamicMesh* TargetMesh,
 		FTransform TargetTransform,
@@ -632,7 +632,7 @@ public:
 		FGeometryScriptBakeTextureOptions BakeOptions,
 		UGeometryScriptDebug* Debug = nullptr);
 
-	UFUNCTION(BlueprintCallable, Category = "GeometryScript|Bake")
+	UFUNCTION(BlueprintCallable, Category = "GeometryScript|Bake", meta = (HidePin = "Debug"))
 	static UE_API UPARAM(DisplayName="Target Mesh") UDynamicMesh* BakeVertex(
 		UDynamicMesh* TargetMesh,
 		FTransform TargetTransform,
@@ -644,7 +644,7 @@ public:
 		FGeometryScriptBakeVertexOptions BakeOptions,
 		UGeometryScriptDebug* Debug = nullptr);
 
-	UFUNCTION(BlueprintCallable, Category = "GeometryScript|Bake")
+	UFUNCTION(BlueprintCallable, Category = "GeometryScript|Bake", meta = (HidePin = "Debug"))
 	static UE_API UPARAM(DisplayName="Textures Out") FGeometryScriptRenderCaptureTextures
 	BakeTextureFromRenderCaptures(
 		UDynamicMesh* TargetMesh,
diff --git a/Engine/Plugins/Runtime/GeometryScripting/Source/GeometryScriptingCore/Public/GeometryScript/MeshBasicEditFunctions.h b/Engine/Plugins/Runtime/GeometryScripting/Source/GeometryScriptingCore/Public/GeometryScript/MeshBasicEditFunctions.h
--- a/Engine/Plugins/Runtime/GeometryScripting/Source/GeometryScriptingCore/Public/GeometryScript/MeshBasicEditFunctions.h
+++ b/Engine/Plugins/Runtime/GeometryScripting/Source/GeometryScriptingCore/Public/GeometryScript/MeshBasicEditFunctions.h
@@ -136,7 +136,7 @@ public:
 	 * Set all vertex positions in the TargetMesh to the specified Positions.
 	 * @param PositionList new vertex Positions. Size must be less than or equal to the MaxVertexID of TargetMesh  (ie gaps are supported).
 	 */
-	UFUNCTION(BlueprintCallable, Category = "GeometryScript|MeshEdits", meta=(ScriptMethod))
+	UFUNCTION(BlueprintCallable, Category = "GeometryScript|MeshEdits", meta=(ScriptMethod, HidePin = "Debug"))
 	static UE_API UPARAM(DisplayName = "Target Mesh") UDynamicMesh*
 	SetAllMeshVertexPositions(
 		UDynamicMesh* TargetMesh,
@@ -192,7 +192,7 @@ public:
 	/**
 	* Adds a triangle (Vertex ID triplet) to the mesh and updates New Triangle Index with the resulting Triangle ID.
 	*/
-	UFUNCTION(BlueprintCallable, Category = "GeometryScript|MeshEdits", meta=(ScriptMethod))
+	UFUNCTION(BlueprintCallable, Category = "GeometryScript|MeshEdits", meta=(ScriptMethod, HidePin = "Debug"))
 	static UE_API UPARAM(DisplayName = "Target Mesh") UDynamicMesh* 
 	AddTriangleToMesh( 
 		UDynamicMesh* TargetMesh, 
@@ -205,7 +205,7 @@ public:
 	/**
 	* Adds a list of triangles to the mesh and populates the New Indices List with the corresponding new Triangle IDs.
 	*/
-	UFUNCTION(BlueprintCallable, Category = "GeometryScript|MeshEdits", meta=(ScriptMethod))
+	UFUNCTION(BlueprintCallable, Category = "GeometryScript|MeshEdits", meta=(ScriptMethod, HidePin = "Debug"))
 	static UE_API UPARAM(DisplayName = "Target Mesh") UDynamicMesh* 
 	AddTrianglesToMesh( 
 		UDynamicMesh* TargetMesh, 
@@ -263,7 +263,7 @@ public:
 	 * @param InterpParam The kept vertex is moved to interpolated position Lerp(VertexKeep Position, VertexDiscard Position, InterpParam)
 	 * @param bDeferChangeNotifications Whether to defer change notifications after this operation
 	 */
-	UFUNCTION(BlueprintCallable, Category = "GeometryScript|MeshEdits", meta = (ScriptMethod))
+	UFUNCTION(BlueprintCallable, Category = "GeometryScript|MeshEdits", meta = (ScriptMethod, HidePin = "Debug"))
 	static UE_API UPARAM(DisplayName = "Target Mesh") UDynamicMesh*
 	MergeMeshVertexPair(
 		UDynamicMesh* TargetMesh,
@@ -289,7 +289,7 @@ public:
 	 * @param DistanceThreshold Vertices further apart than this threshold will not be merged
 	 * @param bDeferChangeNotifications Whether to defer change notifications after this operation
 	 */
-	UFUNCTION(BlueprintCallable, Category = "GeometryScript|MeshEdits", meta = (ScriptMethod))
+	UFUNCTION(BlueprintCallable, Category = "GeometryScript|MeshEdits", meta = (ScriptMethod, HidePin = "Debug"))
 	static UE_API UPARAM(DisplayName = "Target Mesh") UDynamicMesh*
 	MergeMeshVerticesInSelections(
 		UDynamicMesh* TargetMesh,
@@ -309,7 +309,7 @@ public:
 	 * Apply Append Transform to Append Mesh and then add its geometry to the Target Mesh.
 	 * @param AppendOptions Control how details like mesh attributes are handled when one mesh is appended to another.
 	 */
-	UFUNCTION(BlueprintCallable, Category = "GeometryScript|MeshEdits", meta=(ScriptMethod))
+	UFUNCTION(BlueprintCallable, Category = "GeometryScript|MeshEdits", meta=(ScriptMethod, HidePin = "Debug"))
 	static UE_API UPARAM(DisplayName = "Target Mesh") UDynamicMesh* 
 	AppendMesh( 
 		UDynamicMesh* TargetMesh, 
@@ -324,7 +324,7 @@ public:
 	 * Also combines materials lists of the Target and Append meshes, and updates the output mesh materials to reference the combined list.
 	 * @param AppendOptions Control how details like mesh attributes are handled when one mesh is appended to another.
 	 */
-	UFUNCTION(BlueprintCallable, Category = "GeometryScript|MeshEdits", meta = (ScriptMethod))
+	UFUNCTION(BlueprintCallable, Category = "GeometryScript|MeshEdits", meta = (ScriptMethod, HidePin = "Debug"))
 	static UE_API UPARAM(DisplayName = "Target Mesh") UDynamicMesh*
 	AppendMeshWithMaterials(
 		UDynamicMesh* TargetMesh,
@@ -344,7 +344,7 @@ public:
 	 * @param bConstantTransformIsRelative if true, the Constant transform is applied "in the frame" of the Append Transform, otherwise it is applied as a second transform in local coordinates (ie rotate around the AppendTransform X axis, vs around the local X axis)
 	 * @param AppendOptions Control how details like mesh attributes are handled when one mesh is appended to another
 	 */
-	UFUNCTION(BlueprintCallable, Category = "GeometryScript|MeshEdits", meta=(ScriptMethod))
+	UFUNCTION(BlueprintCallable, Category = "GeometryScript|MeshEdits", meta=(ScriptMethod, HidePin = "Debug"))
 	static UE_API UPARAM(DisplayName = "Target Mesh") UDynamicMesh* 
 	AppendMeshTransformed( 
 		UDynamicMesh* TargetMesh, 
@@ -364,7 +364,7 @@ public:
 	 * @param bConstantTransformIsRelative if true, the Constant transform is applied "in the frame" of the Append Transform, otherwise it is applied as a second transform in local coordinates (ie rotate around the AppendTransform X axis, vs around the local X axis)
 	 * @param AppendOptions Control how details like mesh attributes are handled when one mesh is appended to another
 	 */
-	UFUNCTION(BlueprintCallable, Category = "GeometryScript|MeshEdits", meta=(ScriptMethod))
+	UFUNCTION(BlueprintCallable, Category = "GeometryScript|MeshEdits", meta=(ScriptMethod, HidePin = "Debug"))
 	static UE_API UPARAM(DisplayName = "Target Mesh") UDynamicMesh* 
 	AppendMeshTransformedWithMaterials( 
 		UDynamicMesh* TargetMesh, 
@@ -386,7 +386,7 @@ public:
 	 * @param bApplyTransformToFirstInstance if true, the AppendTransform is applied before the first mesh append, otherwise it is applied after
 	 * @param AppendOptions Control how details like mesh attributes are handled when one mesh is appended to another
 	 */
-	UFUNCTION(BlueprintCallable, Category = "GeometryScript|MeshEdits", meta=(ScriptMethod))
+	UFUNCTION(BlueprintCallable, Category = "GeometryScript|MeshEdits", meta=(ScriptMethod, HidePin = "Debug"))
 	static UE_API UPARAM(DisplayName = "Target Mesh") UDynamicMesh* 
 	AppendMeshRepeated( 
 		UDynamicMesh* TargetMesh, 
@@ -405,7 +405,7 @@ public:
 	 * @param bApplyTransformToFirstInstance if true, the AppendTransform is applied before the first mesh append, otherwise it is applied after
 	 * @param AppendOptions Control how details like mesh attributes are handled when one mesh is appended to another
 	 */
-	UFUNCTION(BlueprintCallable, Category = "GeometryScript|MeshEdits", meta=(ScriptMethod))
+	UFUNCTION(BlueprintCallable, Category = "GeometryScript|MeshEdits", meta=(ScriptMethod, HidePin = "Debug"))
 	static UE_API UPARAM(DisplayName = "Target Mesh") UDynamicMesh* 
 	AppendMeshRepeatedWithMaterials( 
 		UDynamicMesh* TargetMesh, 
@@ -425,7 +425,7 @@ public:
 	/**
 	* Adds a set of vertices/triangles to the mesh, with Normals, UVs, and Colors; returning the new triangles indices 
 	*/
-	UFUNCTION(BlueprintCallable, Category = "GeometryScript|MeshEdits", meta=(ScriptMethod))
+	UFUNCTION(BlueprintCallable, Category = "GeometryScript|MeshEdits", meta=(ScriptMethod, HidePin = "Debug"))
 	static UE_API UPARAM(DisplayName = "Target Mesh") UDynamicMesh* 
 	AppendBuffersToMesh( 
 		UDynamicMesh* TargetMesh, 
diff --git a/Engine/Plugins/Runtime/GeometryScripting/Source/GeometryScriptingCore/Public/GeometryScript/MeshBoneWeightFunctions.h b/Engine/Plugins/Runtime/GeometryScripting/Source/GeometryScriptingCore/Public/GeometryScript/MeshBoneWeightFunctions.h
--- a/Engine/Plugins/Runtime/GeometryScripting/Source/GeometryScriptingCore/Public/GeometryScript/MeshBoneWeightFunctions.h
+++ b/Engine/Plugins/Runtime/GeometryScripting/Source/GeometryScriptingCore/Public/GeometryScript/MeshBoneWeightFunctions.h
@@ -353,7 +353,7 @@ public:
 	 * @param bIsValidVertexID will be returned as true if the vertex ID is valid
 	 * @param Profile identifier for the bone/skin weight profile
 	 */
-	UFUNCTION(BlueprintCallable, Category = "GeometryScript|MeshQueries|BoneWeights", meta=(ScriptMethod))
+	UFUNCTION(BlueprintCallable, Category = "GeometryScript|MeshQueries|BoneWeights", meta=(ScriptMethod, HidePin = "Debug"))
 	static UE_API UPARAM(DisplayName = "Target Mesh") UDynamicMesh* 
 	SetVertexBoneWeights( 
 		UDynamicMesh* TargetMesh,
@@ -374,7 +374,7 @@ public:
 	 * @param Alpha The blending factor, ranging from 0 to 1, inclusive. Values outside of this range are clamped.
 	 * @param Result The resulting blend of the two bone weight arrays. 
 	 */
-	UFUNCTION(BlueprintPure, Category = "GeometryScript|MeshQueries|BoneWeights")
+	UFUNCTION(BlueprintPure, Category = "GeometryScript|MeshQueries|BoneWeights", meta = (HidePin = "Debug"))
 	static UE_API void
 	BlendBoneWeights(
 		const TArray<FGeometryScriptBoneWeight>& BoneWeightsA,
@@ -389,7 +389,7 @@ public:
 	 * @param BoneWeights input array of bone index/weight pairs for the Vertex
 	 * @param Profile identifier for the bone/skin weight profile
 	 */
-	UFUNCTION(BlueprintCallable, Category = "GeometryScript|MeshQueries|BoneWeights", meta=(ScriptMethod))
+	UFUNCTION(BlueprintCallable, Category = "GeometryScript|MeshQueries|BoneWeights", meta=(ScriptMethod, HidePin = "Debug"))
 	static UE_API UPARAM(DisplayName = "Target Mesh") UDynamicMesh* 
 	SetAllVertexBoneWeights( 
 		UDynamicMesh* TargetMesh,
@@ -409,7 +409,7 @@ public:
 	 *  @param Options The options to set for the pruning algorithm.
 	 *  @param Profile The skin weight profile to prune the bones from.
 	 */
-	UFUNCTION(BlueprintCallable, Category = "GeometryScript|MeshQueries|BoneWeights", meta=(ScriptMethod))
+	UFUNCTION(BlueprintCallable, Category = "GeometryScript|MeshQueries|BoneWeights", meta=(ScriptMethod, HidePin = "Debug"))
 	static UE_API UPARAM(DisplayName = "Target Mesh") UDynamicMesh* 
 	PruneBoneWeights( 
 		UDynamicMesh* TargetMesh,
@@ -425,7 +425,7 @@ public:
 	 *  @param Options The options to set for the binding algorithm.
 	 *  @param Profile The skin weight profile to update with the smooth binding.
 	 */
-	UFUNCTION(BlueprintCallable, Category = "GeometryScript|MeshQueries|BoneWeights", meta=(ScriptMethod))
+	UFUNCTION(BlueprintCallable, Category = "GeometryScript|MeshQueries|BoneWeights", meta=(ScriptMethod, HidePin = "Debug"))
 	static UE_API UPARAM(DisplayName = "Target Mesh") UDynamicMesh*
 	ComputeSmoothBoneWeights(
 		UDynamicMesh* TargetMesh, 
@@ -444,7 +444,7 @@ public:
 	 * @param Selection Optional subset of target mesh vertices to transfer weights to.
 	 * If left empty, skin weights will be transferred to all target mesh vertices.
 	 */
-	UFUNCTION(BlueprintCallable, Category = "GeometryScript|MeshQueries|BoneWeights", meta = (ScriptMethod))
+	UFUNCTION(BlueprintCallable, Category = "GeometryScript|MeshQueries|BoneWeights", meta = (ScriptMethod, HidePin = "Debug"))
 	static UE_API UPARAM(DisplayName = "Target Mesh") UDynamicMesh* 
 	TransferBoneWeightsFromMesh(
 		UDynamicMesh* SourceMesh,
@@ -459,7 +459,7 @@ public:
 	 * @param TargetMesh Mesh we are copying the bone attributes to.
 	 * @param Options An option object to control how the copying is performed.
 	 */
-	UFUNCTION(BlueprintCallable, Category = "GeometryScript|MeshQueries|BoneWeights", meta=(ScriptMethod))
+	UFUNCTION(BlueprintCallable, Category = "GeometryScript|MeshQueries|BoneWeights", meta=(ScriptMethod, HidePin = "Debug"))
 	static UE_API UPARAM(DisplayName = "Target Mesh") UDynamicMesh* 
 	CopyBonesFromMesh(
 		UDynamicMesh* SourceMesh,
@@ -473,7 +473,7 @@ public:
 	 * @param TargetMesh Mesh we are copying the bone attributes to.
 	 * @param Options An option object to control how the copying is performed.
 	 */
-	UFUNCTION(BlueprintCallable, Category = "GeometryScript|MeshQueries|BoneWeights", meta=(ScriptMethod))
+	UFUNCTION(BlueprintCallable, Category = "GeometryScript|MeshQueries|BoneWeights", meta=(ScriptMethod, HidePin = "Debug"))
 	static UE_API UPARAM(DisplayName = "Target Mesh") UDynamicMesh* 
 	CopyBonesFromSkeleton(
 		USkeleton* SourceSkeleton,
@@ -485,7 +485,7 @@ public:
 	 * Discard the bone attributes (skeleton) from the TargetMesh.
 	 * @param TargetMesh Mesh we are discarding bone attributes from.
 	 */
-	UFUNCTION(BlueprintCallable, Category = "GeometryScript|MeshQueries|BoneWeights", meta=(ScriptMethod))
+	UFUNCTION(BlueprintCallable, Category = "GeometryScript|MeshQueries|BoneWeights", meta=(ScriptMethod, HidePin = "Debug"))
 	static UE_API UPARAM(DisplayName = "Target Mesh") UDynamicMesh* 
 	DiscardBonesFromMesh(
 		UDynamicMesh* TargetMesh,
@@ -499,7 +499,7 @@ public:
 	 * @param BoneIndex The index of the bone with the given name. Will be set to -1 if bIsValidBoneName is false.
 	 */
 	UE_DEPRECATED(5.3, "Use GetBoneInfo instead.")
-	UFUNCTION(BlueprintCallable, Category = "GeometryScript|MeshQueries|BoneWeights", meta=(ScriptMethod))
+	UFUNCTION(BlueprintCallable, Category = "GeometryScript|MeshQueries|BoneWeights", meta=(ScriptMethod, HidePin = "Debug"))
 	static UE_API UPARAM(DisplayName = "Target Mesh") UDynamicMesh* 
 	GetBoneIndex(
 		UDynamicMesh* TargetMesh,
@@ -514,7 +514,7 @@ public:
 	 * @param TargetMesh Mesh containing the bone attributes.
 	 * @param BoneName The name of the root bone.
 	 */
-	UFUNCTION(BlueprintCallable, Category = "GeometryScript|MeshQueries|BoneWeights", meta=(ScriptMethod))
+	UFUNCTION(BlueprintCallable, Category = "GeometryScript|MeshQueries|BoneWeights", meta=(ScriptMethod, HidePin = "Debug"))
 	static UE_API UPARAM(DisplayName = "Target Mesh") UDynamicMesh* 
 	GetRootBoneName(
 		UDynamicMesh* TargetMesh,
@@ -530,7 +530,7 @@ public:
 	 * @param bIsValidBoneName Set to true if the TargetMesh contains a bone with the given name, false otherwise.
 	 * @param ChildrenInfo The informattion of the children.
 	 */
-	UFUNCTION(BlueprintCallable, Category = "GeometryScript|MeshQueries|BoneWeights", meta=(ScriptMethod))
+	UFUNCTION(BlueprintCallable, Category = "GeometryScript|MeshQueries|BoneWeights", meta=(ScriptMethod, HidePin = "Debug"))
 	static UE_API UPARAM(DisplayName = "Target Mesh") UDynamicMesh* 
 	GetBoneChildren(
 		UDynamicMesh* TargetMesh,
@@ -548,7 +548,7 @@ public:
 	 * @param bIsValidBoneName Set to true if the TargetMesh contains a bone with the given name, false otherwise.
 	 * @param BoneInfo The information about the bone.
 	 */
-	UFUNCTION(BlueprintCallable, Category = "GeometryScript|MeshQueries|BoneWeights", meta=(ScriptMethod))
+	UFUNCTION(BlueprintCallable, Category = "GeometryScript|MeshQueries|BoneWeights", meta=(ScriptMethod, HidePin = "Debug"))
 	static UE_API UPARAM(DisplayName = "Target Mesh") UDynamicMesh* 
 	GetBoneInfo(
 		UDynamicMesh* TargetMesh,
@@ -563,7 +563,7 @@ public:
 	 * @param TargetMesh Mesh containing the bone attributes.
 	 * @param BonesInfo Skeleton information.
 	 */
-	UFUNCTION(BlueprintCallable, Category = "GeometryScript|MeshQueries|BoneWeights", meta=(ScriptMethod))
+	UFUNCTION(BlueprintCallable, Category = "GeometryScript|MeshQueries|BoneWeights", meta=(ScriptMethod, HidePin = "Debug"))
 	static UE_API UPARAM(DisplayName = "Target Mesh") UDynamicMesh* 
 	GetAllBonesInfo(
 		UDynamicMesh* TargetMesh,
diff --git a/Engine/Plugins/Runtime/GeometryScripting/Source/GeometryScriptingCore/Public/GeometryScript/MeshBooleanFunctions.h b/Engine/Plugins/Runtime/GeometryScripting/Source/GeometryScriptingCore/Public/GeometryScript/MeshBooleanFunctions.h
--- a/Engine/Plugins/Runtime/GeometryScripting/Source/GeometryScriptingCore/Public/GeometryScript/MeshBooleanFunctions.h
+++ b/Engine/Plugins/Runtime/GeometryScripting/Source/GeometryScriptingCore/Public/GeometryScript/MeshBooleanFunctions.h
@@ -173,7 +173,7 @@ public:
 	* @param Operation selects the specific boolean operation
 	* @param Options selects additional options that are applied to the result
 	*/
-	UFUNCTION(BlueprintCallable, Category = "GeometryScript|Booleans", meta=(ScriptMethod))
+	UFUNCTION(BlueprintCallable, Category = "GeometryScript|Booleans", meta=(ScriptMethod, HidePin = "Debug"))
 	static UE_API UPARAM(DisplayName = "Target Mesh") UDynamicMesh* 
 	ApplyMeshBoolean( 
 		UDynamicMesh* TargetMesh, 
@@ -187,7 +187,7 @@ public:
 	/**
 	* Mesh-Boolean-Union an object with itself to repair self-intersections, remove floating geometry, etc.
 	*/
-	UFUNCTION(BlueprintCallable, Category = "GeometryScript|Booleans", meta=(ScriptMethod))
+	UFUNCTION(BlueprintCallable, Category = "GeometryScript|Booleans", meta=(ScriptMethod, HidePin = "Debug"))
 	static UE_API UPARAM(DisplayName = "Target Mesh") UDynamicMesh* 
 	ApplyMeshSelfUnion( 
 		UDynamicMesh* TargetMesh, 
@@ -201,7 +201,7 @@ public:
 	* @param CutFrame defines the plane used to cut the TargetMesh
 	* @param Options selects additional clean-up operations performed after the cut.
 	*/
-	UFUNCTION(BlueprintCallable, Category = "GeometryScript|Booleans", meta=(ScriptMethod))
+	UFUNCTION(BlueprintCallable, Category = "GeometryScript|Booleans", meta=(ScriptMethod, HidePin = "Debug"))
 	static UE_API UPARAM(DisplayName = "Target Mesh") UDynamicMesh* 
 	ApplyMeshPlaneCut( 
 		UDynamicMesh* TargetMesh, 
@@ -215,7 +215,7 @@ public:
 	 * @param CutFrame defines the plane used to slice the TargetMesh.
 	 * @param Options selects additional clean-up operations performed after the cut.
 	 */
-	UFUNCTION(BlueprintCallable, Category = "GeometryScript|Booleans", meta=(ScriptMethod))
+	UFUNCTION(BlueprintCallable, Category = "GeometryScript|Booleans", meta=(ScriptMethod, HidePin = "Debug"))
 	static UE_API UPARAM(DisplayName = "Target Mesh") UDynamicMesh* 
 	ApplyMeshPlaneSlice( 
 		UDynamicMesh* TargetMesh, 
@@ -229,7 +229,7 @@ public:
 	* @param MirrorFrame defines the plane used to mirror the TargetMesh.
 	* @param Options selects  additional clean-up operations performed after the mirror.
 	*/
-	UFUNCTION(BlueprintCallable, Category = "GeometryScript|Booleans", meta=(ScriptMethod))
+	UFUNCTION(BlueprintCallable, Category = "GeometryScript|Booleans", meta=(ScriptMethod, HidePin = "Debug"))
 	static UE_API UPARAM(DisplayName = "Target Mesh") UDynamicMesh* 
 	ApplyMeshMirror( 
 		UDynamicMesh* TargetMesh, 
@@ -245,7 +245,7 @@ public:
 	* @param PerVertexValue Values to interpolate, must be one per vertex
 	* @param IsoValue Value at which to insert the new curve
 	*/
-	UFUNCTION(BlueprintCallable, Category = "GeometryScript|Booleans", meta = (ScriptMethod))
+	UFUNCTION(BlueprintCallable, Category = "GeometryScript|Booleans", meta = (ScriptMethod, HidePin = "Debug"))
 	static UE_API UPARAM(DisplayName = "Target Mesh") UDynamicMesh*
 	ApplyMeshIsoCurves(
 		UDynamicMesh* TargetMesh,
diff --git a/Engine/Plugins/Runtime/GeometryScripting/Source/GeometryScriptingCore/Public/GeometryScript/MeshComparisonFunctions.h b/Engine/Plugins/Runtime/GeometryScripting/Source/GeometryScriptingCore/Public/GeometryScript/MeshComparisonFunctions.h
--- a/Engine/Plugins/Runtime/GeometryScripting/Source/GeometryScriptingCore/Public/GeometryScript/MeshComparisonFunctions.h
+++ b/Engine/Plugins/Runtime/GeometryScripting/Source/GeometryScriptingCore/Public/GeometryScript/MeshComparisonFunctions.h
@@ -104,7 +104,7 @@ public:
 	 * Returns true if the two input meshes are equivalent under the comparisons defined by the input options. If false, DifferenceInfo provides info on the first difference found.
 	 * @param DifferenceInfo If the meshes are different, provides info on the first difference found.
 	 */
-	UFUNCTION(BlueprintCallable, Category = "GeometryScript|Comparison", meta = (ScriptMethod))
+	UFUNCTION(BlueprintCallable, Category = "GeometryScript|Comparison", meta = (ScriptMethod, HidePin = "Debug"))
 	static UE_API UPARAM(DisplayName = "Target Mesh") UDynamicMesh*
 	IsSameMeshAs(
 		UDynamicMesh* TargetMesh,
@@ -126,7 +126,7 @@ public:
 	/**
 	* Measures the min/max and average closest-point distances between two meshes.
 	*/
-	UFUNCTION(BlueprintCallable, Category = "GeometryScript|Comparison", meta = (ScriptMethod))
+	UFUNCTION(BlueprintCallable, Category = "GeometryScript|Comparison", meta = (ScriptMethod, HidePin = "Debug"))
 	static UE_API UPARAM(DisplayName = "Target Mesh") UDynamicMesh*
 	MeasureDistancesBetweenMeshes(
 		UDynamicMesh* TargetMesh,
@@ -141,7 +141,7 @@ public:
 	/**
 	* Returns true if the two input meshes (with optional transforms) are geometrically intersecting.
 	*/
-	UFUNCTION(BlueprintCallable, Category = "GeometryScript|Comparison", meta = (ScriptMethod))
+	UFUNCTION(BlueprintCallable, Category = "GeometryScript|Comparison", meta = (ScriptMethod, HidePin = "Debug"))
 	static UE_API UPARAM(DisplayName = "Target Mesh") UDynamicMesh*
 	IsIntersectingMesh(
 		UDynamicMesh* TargetMesh,
diff --git a/Engine/Plugins/Runtime/GeometryScripting/Source/GeometryScriptingCore/Public/GeometryScript/MeshDecompositionFunctions.h b/Engine/Plugins/Runtime/GeometryScripting/Source/GeometryScriptingCore/Public/GeometryScript/MeshDecompositionFunctions.h
--- a/Engine/Plugins/Runtime/GeometryScripting/Source/GeometryScriptingCore/Public/GeometryScript/MeshDecompositionFunctions.h
+++ b/Engine/Plugins/Runtime/GeometryScripting/Source/GeometryScriptingCore/Public/GeometryScript/MeshDecompositionFunctions.h
@@ -26,7 +26,7 @@ public:
 	 * @param ComponentMeshes New List of meshes is returned here
 	 * @param MeshPool New meshes in ComponentMeshes output list are allocated from this pool if it is provided (highly recommended!!)
 	 */
-	UFUNCTION(BlueprintCallable, Category = "GeometryScript|Decomposition", meta = (ScriptMethod, Keywords = "Connected Islands"))
+	UFUNCTION(BlueprintCallable, Category = "GeometryScript|Decomposition", meta = (ScriptMethod, HidePin = "Debug", Keywords = "Connected Islands"))
 	static UE_API UPARAM(DisplayName = "Target Mesh") UDynamicMesh* 
 	SplitMeshByComponents(  
 		UDynamicMesh* TargetMesh, 
@@ -41,7 +41,7 @@ public:
 	 * @param MeshPool New meshes in ComponentMeshes output list are allocated from this pool if it is provided (highly recommended!!)
 	 * @param ConnectVerticesThreshold Vertices closer than this distance will be classified as part of the same component, even if they aren't connected by the mesh triangulation
 	 */
-	UFUNCTION(BlueprintCallable, Category = "GeometryScript|Decomposition", meta = (ScriptMethod, Keywords = "Islands Connected Components"))
+	UFUNCTION(BlueprintCallable, Category = "GeometryScript|Decomposition", meta = (ScriptMethod, HidePin = "Debug", Keywords = "Islands Connected Components"))
 	static UE_API UPARAM(DisplayName = "Target Mesh") UDynamicMesh*
 	SplitMeshByVertexOverlap(
 		UDynamicMesh* TargetMesh,
@@ -57,7 +57,7 @@ public:
 	 * @param ComponentMaterialIDs MaterialID for each Mesh in ComponentMeshes is returned here
 	 * @param MeshPool New meshes in ComponentMeshes output list are allocated from this pool if it is provided (highly recommended!!)
 	 */
-	UFUNCTION(BlueprintCallable, Category = "GeometryScript|Decomposition", meta=(ScriptMethod))
+	UFUNCTION(BlueprintCallable, Category = "GeometryScript|Decomposition", meta=(ScriptMethod, HidePin = "Debug"))
 	static UE_API UPARAM(DisplayName = "Target Mesh") UDynamicMesh* 
 	SplitMeshByMaterialIDs(  
 		UDynamicMesh* TargetMesh, 
@@ -73,7 +73,7 @@ public:
 	 * @param ComponentPolygroups Original Polygroup for each Mesh in ComponentMeshes is returned here
 	 * @param MeshPool New meshes in ComponentMeshes output list are allocated from this pool if it is provided (highly recommended!!)
 	 */
-	UFUNCTION(BlueprintCallable, Category = "GeometryScript|Decomposition", meta = (ScriptMethod, DisplayName = "Split Mesh By PolyGroups"))
+	UFUNCTION(BlueprintCallable, Category = "GeometryScript|Decomposition", meta = (ScriptMethod, HidePin = "Debug", DisplayName = "Split Mesh By PolyGroups"))
 	static UE_API UPARAM(DisplayName = "Target Mesh") UDynamicMesh* 
 	SplitMeshByPolygroups(  
 		UDynamicMesh* TargetMesh, 
@@ -147,7 +147,7 @@ public:
 	/**
 	 * CopyMeshSelectionToMesh should be used instead of this function
 	 */
-	UFUNCTION(BlueprintCallable, Category = "GeometryScript|Decomposition", meta=(ScriptMethod))
+	UFUNCTION(BlueprintCallable, Category = "GeometryScript|Decomposition", meta=(ScriptMethod, HidePin = "Debug"))
 	static UE_API UPARAM(DisplayName = "Target Mesh") UDynamicMesh* 
 	GetSubMeshFromMesh(  
 		UDynamicMesh* TargetMesh, 
@@ -161,7 +161,7 @@ public:
 	 * @param bAppendToExisting if false (default), StoreToSubmesh is cleared, otherwise selected triangles are appended
 	 * @param bPreserveGroupIDs if true, GroupIDs of triangles on TargetMesh are preserved in StoreToSubmesh. Otherwise new GroupIDs are allocated.
 	 */
-	UFUNCTION(BlueprintCallable, Category = "GeometryScript|Decomposition", meta=(ScriptMethod))
+	UFUNCTION(BlueprintCallable, Category = "GeometryScript|Decomposition", meta=(ScriptMethod, HidePin = "Debug"))
 	static UE_API UPARAM(DisplayName = "Target Mesh") UDynamicMesh* 
 	CopyMeshSelectionToMesh(  
 		UDynamicMesh* TargetMesh, 
@@ -175,7 +175,7 @@ public:
 	/**
 	 * Set CopyToMesh to be the same mesh as CopyFromMesh
 	 */
-	UFUNCTION(BlueprintCallable, Category = "GeometryScript|Decomposition", meta=(ScriptMethod))
+	UFUNCTION(BlueprintCallable, Category = "GeometryScript|Decomposition", meta=(ScriptMethod, HidePin = "Debug"))
 	static UE_API UPARAM(DisplayName = "Copy From Mesh") UDynamicMesh* 
 	CopyMeshToMesh(  
 		UDynamicMesh* CopyFromMesh, 
diff --git a/Engine/Plugins/Runtime/GeometryScripting/Source/GeometryScriptingCore/Public/GeometryScript/MeshDeformFunctions.h b/Engine/Plugins/Runtime/GeometryScripting/Source/GeometryScriptingCore/Public/GeometryScript/MeshDeformFunctions.h
--- a/Engine/Plugins/Runtime/GeometryScripting/Source/GeometryScriptingCore/Public/GeometryScript/MeshDeformFunctions.h
+++ b/Engine/Plugins/Runtime/GeometryScripting/Source/GeometryScriptingCore/Public/GeometryScript/MeshDeformFunctions.h
@@ -206,7 +206,7 @@ public:
 	* Applies a Bend Warp around an axis defined by the Bend Orientation transform.  
 	* The extents of the affected region can be controlled by the Options.
 	*/
-	UFUNCTION(BlueprintCallable, Category = "GeometryScript|Deformations", meta=(ScriptMethod))
+	UFUNCTION(BlueprintCallable, Category = "GeometryScript|Deformations", meta=(ScriptMethod, HidePin = "Debug"))
 	static UE_API UPARAM(DisplayName = "Target Mesh") UDynamicMesh* 
 	ApplyBendWarpToMesh(  
 		UDynamicMesh* TargetMesh, 
@@ -220,7 +220,7 @@ public:
 	 * Applies a twist warp around an axis defined by the Twist Orientation transform.
 	 * The extents of the affected region can be controlled by the Options.
 	 */
-	UFUNCTION(BlueprintCallable, Category = "GeometryScript|Deformations", meta=(ScriptMethod))
+	UFUNCTION(BlueprintCallable, Category = "GeometryScript|Deformations", meta=(ScriptMethod, HidePin = "Debug"))
 	static UE_API UPARAM(DisplayName = "Target Mesh") UDynamicMesh* 
 	ApplyTwistWarpToMesh(  
 		UDynamicMesh* TargetMesh, 
@@ -235,7 +235,7 @@ public:
 	* The amount of flare in the perpendicular directions can be controlled by FlarePercentX and FlarePercentY
 	* and the extents of the affected region can be controlled by the Options.
 	*/
-	UFUNCTION(BlueprintCallable, Category = "GeometryScript|Deformations", meta=(ScriptMethod))
+	UFUNCTION(BlueprintCallable, Category = "GeometryScript|Deformations", meta=(ScriptMethod, HidePin = "Debug"))
 	static UE_API UPARAM(DisplayName = "Target Mesh") UDynamicMesh* 
 	ApplyFlareWarpToMesh(  
 		UDynamicMesh* TargetMesh, 
@@ -251,7 +251,7 @@ public:
 	* Applies various simple math-function-based warps around an axis defined by the Warp Orientation transform, 
 	* currently a 1D or 2D sine-wave with arbitrary orientation may be selected by WarpType.
 	*/
-	UFUNCTION(BlueprintCallable, Category = "GeometryScript|Deformations", meta=(ScriptMethod))
+	UFUNCTION(BlueprintCallable, Category = "GeometryScript|Deformations", meta=(ScriptMethod, HidePin = "Debug"))
 	static UE_API UPARAM(DisplayName = "Target Mesh") UDynamicMesh* 
 	ApplyMathWarpToMesh(  
 		UDynamicMesh* TargetMesh, 
@@ -265,7 +265,7 @@ public:
 	* @param Selection if non-empty, only the vertices identified by the selection will be displaced, otherwise
 	* the Options' EmptyBehavior will be followed.
 	*/
-	UFUNCTION(BlueprintCallable, Category = "GeometryScript|Deformations", meta=(ScriptMethod))
+	UFUNCTION(BlueprintCallable, Category = "GeometryScript|Deformations", meta=(ScriptMethod, HidePin = "Debug"))
 	static UE_API UPARAM(DisplayName = "Target Mesh") UDynamicMesh* 
 	ApplyPerlinNoiseToMesh(  
 		UDynamicMesh* TargetMesh, 
@@ -278,7 +278,7 @@ public:
 	* @param Selection if non-empty, only the vertices identified by the selection will be subject to smoothing,  otherwise
 	* the Options' EmptyBehavior will be followed.
 	*/
-	UFUNCTION(BlueprintCallable, Category = "GeometryScript|Deformations", meta=(ScriptMethod))
+	UFUNCTION(BlueprintCallable, Category = "GeometryScript|Deformations", meta=(ScriptMethod, HidePin = "Debug"))
 	static UE_API UPARAM(DisplayName = "Target Mesh") UDynamicMesh* 
 	ApplyIterativeSmoothingToMesh(  
 		UDynamicMesh* TargetMesh, 
@@ -291,7 +291,7 @@ public:
 	 * @param Selection if non-empty, only the vertices identified by the selection will be subject to displacement,  otherwise
 	 * the Options' EmptyBehavior will be followed.
 	 */
-	UFUNCTION(BlueprintCallable, Category = "GeometryScript|Deformations", meta=(ScriptMethod))
+	UFUNCTION(BlueprintCallable, Category = "GeometryScript|Deformations", meta=(ScriptMethod, HidePin = "Debug"))
 	static UE_API UPARAM(DisplayName = "Target Mesh") UDynamicMesh* 
 	ApplyDisplaceFromTextureMap(  
 		UDynamicMesh* TargetMesh, 
@@ -306,7 +306,7 @@ public:
 	 * VectorList Length must be >= the MaxVertexID of the Target Mesh.
 	 * @param Selection if non-empty, only the vertices identified by the selection will be displaced. The VectorList must still be the same size as the whole mesh, this is just a filter on which vertices are updated.
 	 */
-	UFUNCTION(BlueprintCallable, Category = "GeometryScript|Deformations", meta=(ScriptMethod))
+	UFUNCTION(BlueprintCallable, Category = "GeometryScript|Deformations", meta=(ScriptMethod, HidePin = "Debug"))
 	static UE_API UPARAM(DisplayName = "Target Mesh") UDynamicMesh* 
 	ApplyDisplaceFromPerVertexVectors(  
 		UDynamicMesh* TargetMesh,
diff --git a/Engine/Plugins/Runtime/GeometryScripting/Source/GeometryScriptingCore/Public/GeometryScript/MeshGeodesicFunctions.h b/Engine/Plugins/Runtime/GeometryScripting/Source/GeometryScriptingCore/Public/GeometryScript/MeshGeodesicFunctions.h
--- a/Engine/Plugins/Runtime/GeometryScripting/Source/GeometryScriptingCore/Public/GeometryScript/MeshGeodesicFunctions.h
+++ b/Engine/Plugins/Runtime/GeometryScripting/Source/GeometryScriptingCore/Public/GeometryScript/MeshGeodesicFunctions.h
@@ -29,7 +29,7 @@ public:
 	* @param VertexIDList, if found this will hold on return a list of mesh vertex IDs that define the path from StartVertexID to EndVertexID.
 	* @param bFoundErrors will be false on success.  
 	*/
-	UFUNCTION(BlueprintCallable, Category = "GeometryScript|MeshGeodesic", meta=(ScriptMethod))
+	UFUNCTION(BlueprintCallable, Category = "GeometryScript|MeshGeodesic", meta=(ScriptMethod, HidePin = "Debug"))
 	static UE_API UPARAM(DisplayName = "Target Mesh") UDynamicMesh* 
 	GetShortestVertexPath(
 		UDynamicMesh* TargetMesh,
@@ -52,7 +52,7 @@ public:
 	* Note, Barycentric coordinates are of the form (a,b,c) where each entry is positive and a + b + c = 1.  
 	*       If the provided coordinates are invalid, the value (1/3, 1/3, 1/3) will be used.
 	*/
-	UFUNCTION(BlueprintCallable, Category = "GeometryScript|MeshGeodesic", meta=(ScriptMethod))
+	UFUNCTION(BlueprintCallable, Category = "GeometryScript|MeshGeodesic", meta=(ScriptMethod, HidePin = "Debug"))
 	static UE_API UPARAM(DisplayName = "Target Mesh") UDynamicMesh* 
 	GetShortestSurfacePath(
 		UDynamicMesh* TargetMesh,
@@ -78,7 +78,7 @@ public:
 	*       If the provided coordinate is invalid, the value (1/3, 1/3, 1/3) will be used.
 	*       Also, if the direction vector is nearly zero, the up-vector will be used. 
 	*/
-	UFUNCTION(BlueprintCallable, Category = "GeometryScript|MeshGeodesic", meta=(ScriptMethod))
+	UFUNCTION(BlueprintCallable, Category = "GeometryScript|MeshGeodesic", meta=(ScriptMethod, HidePin = "Debug"))
 	static UE_API UPARAM(DisplayName = "Target Mesh") UDynamicMesh* 
 	CreateSurfacePath(
 		UDynamicMesh* TargetMesh,
diff --git a/Engine/Plugins/Runtime/GeometryScripting/Source/GeometryScriptingCore/Public/GeometryScript/MeshMaterialFunctions.h b/Engine/Plugins/Runtime/GeometryScripting/Source/GeometryScriptingCore/Public/GeometryScript/MeshMaterialFunctions.h
--- a/Engine/Plugins/Runtime/GeometryScripting/Source/GeometryScriptingCore/Public/GeometryScript/MeshMaterialFunctions.h
+++ b/Engine/Plugins/Runtime/GeometryScripting/Source/GeometryScriptingCore/Public/GeometryScript/MeshMaterialFunctions.h
@@ -26,7 +26,7 @@ public:
 	* Enables per-triangle Material IDs on a mesh and initializes the values to 0. 
 	* If Target Mesh already has Material IDs, this function will do nothing.
 	*/
-	UFUNCTION(BlueprintCallable, Category = "GeometryScript|Materials", meta=(ScriptMethod))
+	UFUNCTION(BlueprintCallable, Category = "GeometryScript|Materials", meta=(ScriptMethod, HidePin = "Debug"))
 	static UE_API UPARAM(DisplayName = "Target Mesh") UDynamicMesh* 
 	EnableMaterialIDs( UDynamicMesh* TargetMesh, UGeometryScriptDebug* Debug = nullptr );
 
@@ -34,14 +34,14 @@ public:
 	* Resets all Material IDs on a mesh to the given ClearValue, or 0 if no ClearValue is provided.
 	* If Material IDs are not already enabled on the Target Mesh, this function will first enable them and then set the value.
 	*/
-	UFUNCTION(BlueprintCallable, Category = "GeometryScript|Materials", meta=(ScriptMethod))
+	UFUNCTION(BlueprintCallable, Category = "GeometryScript|Materials", meta=(ScriptMethod, HidePin = "Debug"))
 	static UE_API UPARAM(DisplayName = "Target Mesh") UDynamicMesh* 
 	ClearMaterialIDs( UDynamicMesh* TargetMesh, int ClearValue = 0, UGeometryScriptDebug* Debug = nullptr );
 
 	/**
 	* For all triangles with a Material ID matching the given value (From Material ID), update the Material ID to the new value (To Material ID).
 	*/
-	UFUNCTION(BlueprintCallable, Category = "GeometryScript|Materials", meta=(ScriptMethod))
+	UFUNCTION(BlueprintCallable, Category = "GeometryScript|Materials", meta=(ScriptMethod, HidePin = "Debug"))
 	static UE_API UPARAM(DisplayName = "Target Mesh") UDynamicMesh* 
 	RemapMaterialIDs( 
 		UDynamicMesh* TargetMesh, 
@@ -58,7 +58,7 @@ public:
 	 * If a Material cannot be found in ToMaterialList, a warning will be printed and the MaterialID left unmodified, 
 	 * unless MissingMaterialID is set to a value >= 0, in which case MissingMaterialID will be assigned
 	 */
-	UFUNCTION(BlueprintCallable, Category = "GeometryScript|Materials", meta=(ScriptMethod))
+	UFUNCTION(BlueprintCallable, Category = "GeometryScript|Materials", meta=(ScriptMethod, HidePin = "Debug"))
 	static UE_API UPARAM(DisplayName = "Target Mesh") UDynamicMesh* 
 	RemapToNewMaterialIDsByMaterial( 
 		UDynamicMesh* TargetMesh, 
@@ -82,7 +82,7 @@ public:
 	 * @param RemapInvalidMaterialID If >= 0, automatically remap invalid input material IDs to this value
 	 * @param bCompactDuplicateMaterials If true, materials from TargetMeshMaterials will only be added if they are not already in RequiredMaterials. If false, all TargetMeshMaterials are appended.
 	 */
-	UFUNCTION(BlueprintCallable, Category = "GeometryScript|Materials", meta = (ScriptMethod))
+	UFUNCTION(BlueprintCallable, Category = "GeometryScript|Materials", meta = (ScriptMethod, HidePin = "Debug"))
 	static UE_API UPARAM(DisplayName = "Target Mesh") UDynamicMesh*
 	RemapAndCombineMaterials(
 		UDynamicMesh* TargetMesh,
@@ -107,7 +107,7 @@ public:
 	* If a triangle is not present in the Target Mesh the number -1 will be used for the corresponding Material ID.  
 	* If Material IDs are not enabled on the TargetMesh no Material IDs will be added to the result list.
 	*/
-	UFUNCTION(BlueprintPure, Category = "GeometryScript|Materials", meta=(ScriptMethod))
+	UFUNCTION(BlueprintPure, Category = "GeometryScript|Materials", meta=(ScriptMethod, HidePin = "Debug"))
 	static UE_API UPARAM(DisplayName = "Target Mesh") UDynamicMesh* 
 	GetMaterialIDsOfTriangles(
 		UDynamicMesh* TargetMesh,
@@ -127,7 +127,7 @@ public:
 	/**
 	 * Populates Triangle ID List with the Triangle IDs of triangles that share the specified Material ID in the Target Mesh.
 	 */
-	UFUNCTION(BlueprintPure, Category = "GeometryScript|Selection", meta=(ScriptMethod))
+	UFUNCTION(BlueprintPure, Category = "GeometryScript|Selection", meta=(ScriptMethod, HidePin = "Debug"))
 	static UE_API UPARAM(DisplayName = "Target Mesh") UDynamicMesh* 
 	GetTrianglesByMaterialID( 
 		UDynamicMesh* TargetMesh, 
@@ -153,7 +153,7 @@ public:
 	* @param TriangleMateralIDList the list of Material IDs, the length of this list should be the same as the Max Triangle ID for this mesh.
 	* @param bDeferChangeNotifications if true, the UDynamicMesh does not emit a change event/signal for this modification.
 	*/
-	UFUNCTION(BlueprintCallable, Category = "GeometryScript|Materials", meta=(ScriptMethod, DisplayName="Set Material IDs On All Triangles"))
+	UFUNCTION(BlueprintCallable, Category = "GeometryScript|Materials", meta=(ScriptMethod, HidePin = "Debug", DisplayName="Set Material IDs On All Triangles"))
 	static UE_API UPARAM(DisplayName = "Target Mesh") UDynamicMesh*
 	SetAllTriangleMaterialIDs(
 		UDynamicMesh* TargetMesh,
@@ -167,7 +167,7 @@ public:
 	* @param MaterialID the ID to be assigned to each triangle in the input list.
 	* @param bDeferChangeNotifications if true, the UDynamicMesh does not emit a change event/signal for this modification.
 	*/
-	UFUNCTION(BlueprintCallable, Category = "GeometryScript|Materials", meta=(ScriptMethod, DisplayName="Set Material ID On Triangles"))
+	UFUNCTION(BlueprintCallable, Category = "GeometryScript|Materials", meta=(ScriptMethod, HidePin = "Debug", DisplayName="Set Material ID On Triangles"))
 	static UE_API UPARAM(DisplayName = "Target Mesh") UDynamicMesh* 
 	SetMaterialIDOnTriangles(
 		UDynamicMesh* TargetMesh,
@@ -181,7 +181,7 @@ public:
 	* @param MaterialID new Material ID to set
 	* @param bDeferChangeNotifications if true, the UDynamicMesh does not emit a change event/signal for this modification
 	*/
-	UFUNCTION(BlueprintCallable, Category = "GeometryScript|Materials", meta=(ScriptMethod, DisplayName="Set Material ID For Mesh Selection"))
+	UFUNCTION(BlueprintCallable, Category = "GeometryScript|Materials", meta=(ScriptMethod, HidePin = "Debug", DisplayName="Set Material ID For Mesh Selection"))
 	static UE_API UPARAM(DisplayName = "Target Mesh") UDynamicMesh* 
 	SetMaterialIDForMeshSelection( 
 		UDynamicMesh* TargetMesh, 
@@ -197,7 +197,7 @@ public:
 	* @param MaterialID explicit new MaterialID to set
 	* @param bDeferChangeNotifications if true, the UDynamicMesh does not emit a change event/signal for this modification
 	*/
-	UFUNCTION(BlueprintCallable, Category = "GeometryScript|Materials", meta=(ScriptMethod, DisplayName = "Set PolyGroup Material ID"))
+	UFUNCTION(BlueprintCallable, Category = "GeometryScript|Materials", meta=(ScriptMethod, HidePin = "Debug", DisplayName = "Set PolyGroup Material ID"))
 	static UE_API UPARAM(DisplayName = "Target Mesh") UDynamicMesh* 
 	SetPolygroupMaterialID( 
 		UDynamicMesh* TargetMesh, 
@@ -212,7 +212,7 @@ public:
 	 * Delete all triangles in TargetMesh with the given MaterialID
 	 * @param NumDeleted number of deleted triangles is returned here
 	 */
-	UFUNCTION(BlueprintCallable, Category = "GeometryScript|Materials", meta=(ScriptMethod))
+	UFUNCTION(BlueprintCallable, Category = "GeometryScript|Materials", meta=(ScriptMethod, HidePin = "Debug"))
 	static UE_API UPARAM(DisplayName = "Target Mesh") UDynamicMesh* 
 	DeleteTrianglesByMaterialID( 
 		UDynamicMesh* TargetMesh, 
@@ -228,7 +228,7 @@ public:
 	 * @param CompactedMaterialList new Compacted Material list, one-to-one with new compacted MaterialIDs
 	 * @param bRemoveDuplicateMaterials Whether to also remove duplicate materials from the compacted list
 	 */
-	UFUNCTION(BlueprintCallable, Category = "GeometryScript|Materials", meta=(ScriptMethod))
+	UFUNCTION(BlueprintCallable, Category = "GeometryScript|Materials", meta=(ScriptMethod, HidePin = "Debug"))
 	static UE_API UPARAM(DisplayName = "Target Mesh") UDynamicMesh* 
 	CompactMaterialIDs( 
 		UDynamicMesh* TargetMesh, 
diff --git a/Engine/Plugins/Runtime/GeometryScripting/Source/GeometryScriptingCore/Public/GeometryScript/MeshModelingFunctions.h b/Engine/Plugins/Runtime/GeometryScripting/Source/GeometryScriptingCore/Public/GeometryScript/MeshModelingFunctions.h
--- a/Engine/Plugins/Runtime/GeometryScripting/Source/GeometryScriptingCore/Public/GeometryScript/MeshModelingFunctions.h
+++ b/Engine/Plugins/Runtime/GeometryScripting/Source/GeometryScriptingCore/Public/GeometryScript/MeshModelingFunctions.h
@@ -300,7 +300,7 @@ public:
 	 * The input Selection will still identify the same geometric elements after Disconnecting.
 	 * @param bAllowBowtiesInOutput if false, any bowtie vertices created by the operation will be automatically split
 	 */
-	UFUNCTION(BlueprintCallable, Category = "GeometryScript|Modeling", meta=(ScriptMethod))
+	UFUNCTION(BlueprintCallable, Category = "GeometryScript|Modeling", meta=(ScriptMethod, HidePin = "Debug"))
 	static UE_API UPARAM(DisplayName = "Target Mesh") UDynamicMesh* 
 	ApplyMeshDisconnectFaces(
 		UDynamicMesh* TargetMesh,
@@ -315,7 +315,7 @@ public:
 	 * @param TargetMesh Mesh to operate on
 	 * @param Selection Which edges to operate on. Non-edge selections will be interpreted as edge selections -- i.e., all selected triangles' edges, or all selected vertices' one-ring edges.
 	 */
-	UFUNCTION(BlueprintCallable, Category = "GeometryScript|Modeling", meta = (ScriptMethod))
+	UFUNCTION(BlueprintCallable, Category = "GeometryScript|Modeling", meta = (ScriptMethod, HidePin = "Debug"))
 	static UE_API UPARAM(DisplayName = "Target Mesh") UDynamicMesh*
 	ApplyMeshDisconnectFacesAlongEdges(
 		UDynamicMesh* TargetMesh,
@@ -327,7 +327,7 @@ public:
 	 * @param NewTriangles a Mesh Selection of the duplicate triangles is returned here (with type Triangles)
 	 * @param bAllowBowtiesInOutput if false, any bowtie vertices resulting created in the Duplicate area will be disconnected into unique vertices
 	 */
-	UFUNCTION(BlueprintCallable, Category = "GeometryScript|Modeling", meta=(ScriptMethod))
+	UFUNCTION(BlueprintCallable, Category = "GeometryScript|Modeling", meta=(ScriptMethod, HidePin = "Debug"))
 	static UE_API UPARAM(DisplayName = "Target Mesh") UDynamicMesh* 
 	ApplyMeshDuplicateFaces(
 		UDynamicMesh* TargetMesh,
@@ -341,7 +341,7 @@ public:
 	 * Offset the vertices of TargetMesh from their initial positions based on averaged vertex normals.
 	 * This function is intended for high-res meshes, for polymodeling-style offsets, ApplyMeshOffsetFaces will produce better results.
 	 */
-	UFUNCTION(BlueprintCallable, Category = "GeometryScript|Modeling", meta=(ScriptMethod))
+	UFUNCTION(BlueprintCallable, Category = "GeometryScript|Modeling", meta=(ScriptMethod, HidePin = "Debug"))
 	static UE_API UPARAM(DisplayName = "Target Mesh") UDynamicMesh* 
 	ApplyMeshOffset(
 		UDynamicMesh* TargetMesh,
@@ -352,7 +352,7 @@ public:
 	 * Create a thickened shell from TargetMesh by offsetting the vertex positions along averaged vertex normals, inwards or outwards.
 	 * Similar to ApplyMeshOffset but also includes the initial mesh (possibly flipped, if the offset is positive)
 	 */
-	UFUNCTION(BlueprintCallable, Category = "GeometryScript|Modeling", meta=(ScriptMethod))
+	UFUNCTION(BlueprintCallable, Category = "GeometryScript|Modeling", meta=(ScriptMethod, HidePin = "Debug"))
 	static UE_API UPARAM(DisplayName = "Target Mesh") UDynamicMesh* 
 	ApplyMeshShell(
 		UDynamicMesh* TargetMesh,
@@ -364,7 +364,7 @@ public:
 	 * Apply Linear Extrusion (ie extrusion in a single direction) to the triangles of TargetMesh identified by the Selection.
 	 * The input Selection will still identify the same geometric elements after the Extrusion
 	 */
-	UFUNCTION(BlueprintCallable, Category = "GeometryScript|Modeling", meta=(ScriptMethod))
+	UFUNCTION(BlueprintCallable, Category = "GeometryScript|Modeling", meta=(ScriptMethod, HidePin = "Debug"))
 	static UE_API UPARAM(DisplayName = "Target Mesh") UDynamicMesh* 
 	ApplyMeshLinearExtrudeFaces(
 		UDynamicMesh* TargetMesh,
@@ -376,7 +376,7 @@ public:
 	 * Apply an Offset to the faces of TargetMesh identified by the Selection, or all faces if the Selection is empty.
 	 * The Offset direction at each vertex can be derived from the averaged vertex normals or per-triangle normals.
 	 */
-	UFUNCTION(BlueprintCallable, Category = "GeometryScript|Modeling", meta=(ScriptMethod))
+	UFUNCTION(BlueprintCallable, Category = "GeometryScript|Modeling", meta=(ScriptMethod, HidePin = "Debug"))
 	static UE_API UPARAM(DisplayName = "Target Mesh") UDynamicMesh* 
 	ApplyMeshOffsetFaces(
 		UDynamicMesh* TargetMesh,
@@ -387,7 +387,7 @@ public:
 	/**
 	 * Apply an Inset or Outset to the faces of TargetMesh identified by the Selection, or all faces if the Selection is empty.
 	 */
-	UFUNCTION(BlueprintCallable, Category = "GeometryScript|Modeling", meta=(ScriptMethod))
+	UFUNCTION(BlueprintCallable, Category = "GeometryScript|Modeling", meta=(ScriptMethod, HidePin = "Debug"))
 	static UE_API UPARAM(DisplayName = "Target Mesh") UDynamicMesh* 
 	ApplyMeshInsetOutsetFaces(
 		UDynamicMesh* TargetMesh,
@@ -400,7 +400,7 @@ public:
 	 * @param Selection specifies which mesh edges to Bevel
 	 * @param BevelOptions settings for the Bevel Operation
 	 */
-	UFUNCTION(BlueprintCallable, Category = "GeometryScript|Modeling", meta = (ScriptMethod))
+	UFUNCTION(BlueprintCallable, Category = "GeometryScript|Modeling", meta = (ScriptMethod, HidePin = "Debug"))
 	static UE_API UPARAM(DisplayName = "Target Mesh") UDynamicMesh*
 	ApplyMeshBevelEdgeSelection(
 		UDynamicMesh* TargetMesh,
@@ -414,7 +414,7 @@ public:
 	 * @param BevelMode specifies whether Selection should be optionally converted to a Triangle Region or set of PolyGroup Edges
 	 * @param BevelOptions settings for the Bevel Operation
 	 */
-	UFUNCTION(BlueprintCallable, Category = "GeometryScript|Modeling", meta=(ScriptMethod, DisplayName = "Apply Mesh Bevel Region Selection"))
+	UFUNCTION(BlueprintCallable, Category = "GeometryScript|Modeling", meta=(ScriptMethod, HidePin = "Debug", DisplayName = "Apply Mesh Bevel Region Selection"))
 	static UE_API UPARAM(DisplayName = "Target Mesh") UDynamicMesh* 
 	ApplyMeshBevelSelection(
 		UDynamicMesh* TargetMesh,
@@ -427,7 +427,7 @@ public:
 	/**
 	 * Apply a Mesh Bevel operation to all PolyGroup Edges.
      */
-	UFUNCTION(BlueprintCallable, Category = "GeometryScript|Modeling", meta=(ScriptMethod, DisplayName = "Apply Mesh PolyGroup Bevel"))
+	UFUNCTION(BlueprintCallable, Category = "GeometryScript|Modeling", meta=(ScriptMethod, HidePin = "Debug", DisplayName = "Apply Mesh PolyGroup Bevel"))
 	static UE_API UPARAM(DisplayName = "Target Mesh") UDynamicMesh* 
 	ApplyMeshPolygroupBevel(
 		UDynamicMesh* TargetMesh,
@@ -449,7 +449,7 @@ public:
 	//
 
 
-	UFUNCTION(BlueprintCallable, Category = "GeometryScript|Compatibility", meta=(ScriptMethod))
+	UFUNCTION(BlueprintCallable, Category = "GeometryScript|Compatibility", meta=(ScriptMethod, HidePin = "Debug"))
 	static UE_API UPARAM(DisplayName = "Target Mesh") UDynamicMesh* 
 	ApplyMeshExtrude_Compatibility_5p0(
 		UDynamicMesh* TargetMesh,
diff --git a/Engine/Plugins/Runtime/GeometryScripting/Source/GeometryScriptingCore/Public/GeometryScript/MeshNormalsFunctions.h b/Engine/Plugins/Runtime/GeometryScripting/Source/GeometryScriptingCore/Public/GeometryScript/MeshNormalsFunctions.h
--- a/Engine/Plugins/Runtime/GeometryScripting/Source/GeometryScriptingCore/Public/GeometryScript/MeshNormalsFunctions.h
+++ b/Engine/Plugins/Runtime/GeometryScripting/Source/GeometryScriptingCore/Public/GeometryScript/MeshNormalsFunctions.h
@@ -76,7 +76,7 @@ public:
 	/**
 	 * Flip/Invert the normal vectors of TargetMesh by multiplying them by -1, as well as reversing the mesh triangle orientations, ie triangle (a,b,c) becomes (b,a,c)
 	 */
-	UFUNCTION(BlueprintCallable, Category = "GeometryScript|Normals", meta=(ScriptMethod))
+	UFUNCTION(BlueprintCallable, Category = "GeometryScript|Normals", meta=(ScriptMethod, HidePin = "Debug"))
 	static UE_API UPARAM(DisplayName = "Target Mesh") UDynamicMesh* 
 	FlipNormals( 
 		UDynamicMesh* TargetMesh, 
@@ -88,7 +88,7 @@ public:
 	 * @param bFlipTriangleOrientation If true, triangles are flipped by reversing the mesh triangle orientations, ie triangle (a,b,c) becomes (b,a,c)
 	 * @param bFlipNormalDirection If true, normals of vertices on the selected triangles are flipped by multiplying them by -1
 	 */
-	UFUNCTION(BlueprintCallable, Category = "GeometryScript|Normals", meta = (ScriptMethod))
+	UFUNCTION(BlueprintCallable, Category = "GeometryScript|Normals", meta = (ScriptMethod, HidePin = "Debug"))
 	static UE_API UPARAM(DisplayName = "Target Mesh") UDynamicMesh*
 	FlipTriangleSelectionNormals(
 		UDynamicMesh* TargetMesh,
@@ -104,7 +104,7 @@ public:
 	 * In the second pass, the "global" orientation is detected by casting rays from outside the mesh. This may produce incorrect results for
 	 * meshes that are not closed. 
 	 */
-	UFUNCTION(BlueprintCallable, Category = "GeometryScript|Normals", meta=(ScriptMethod))
+	UFUNCTION(BlueprintCallable, Category = "GeometryScript|Normals", meta=(ScriptMethod, HidePin = "Debug"))
 	static UE_API UPARAM(DisplayName = "Target Mesh") UDynamicMesh* 
 	AutoRepairNormals( 
 		UDynamicMesh* TargetMesh, 
@@ -114,7 +114,7 @@ public:
 	 * Recompute the normals of TargetMesh by averaging the triangle/face normals around each vertex, using combined area and angle weighting.
 	 * Each vertex will have a single normal, ie there will be no hard edges.
 	 */
-	UFUNCTION(BlueprintCallable, Category = "GeometryScript|Normals", meta=(ScriptMethod, DisplayName="Set Mesh To Per Vertex Normals (Computed)"))
+	UFUNCTION(BlueprintCallable, Category = "GeometryScript|Normals", meta=(ScriptMethod, HidePin = "Debug", DisplayName="Set Mesh To Per Vertex Normals (Computed)"))
 	static UE_API UPARAM(DisplayName = "Target Mesh") UDynamicMesh* 
 	SetPerVertexNormals( 
 		UDynamicMesh* TargetMesh, 
@@ -124,7 +124,7 @@ public:
 	 * Recompute the normals of TargetMesh by setting the normals of each triangle vertex to the triangle/face normal.
 	 * Each vertex will have a unique normal in each triangle, ie there will be hard edges / split normals at every mesh edge
 	 */
-	UFUNCTION(BlueprintCallable, Category = "GeometryScript|Normals", meta=(ScriptMethod, DisplayName = "Set Mesh To Facet Normals"))
+	UFUNCTION(BlueprintCallable, Category = "GeometryScript|Normals", meta=(ScriptMethod, HidePin = "Debug", DisplayName = "Set Mesh To Facet Normals"))
 	static UE_API UPARAM(DisplayName = "Target Mesh") UDynamicMesh* 
 	SetPerFaceNormals( 
 		UDynamicMesh* TargetMesh, 
@@ -135,7 +135,7 @@ public:
 	 * edges, ie each shared triangle-vertex normal is recomputed by averaging the face normals of triangles that reference
 	 * that shared triangle-vertex normal
 	 */
-	UFUNCTION(BlueprintCallable, Category = "GeometryScript|Normals", meta=(ScriptMethod))
+	UFUNCTION(BlueprintCallable, Category = "GeometryScript|Normals", meta=(ScriptMethod, HidePin = "Debug"))
 	static UE_API UPARAM(DisplayName = "Target Mesh") UDynamicMesh* 
 	RecomputeNormals(  
 		UDynamicMesh* TargetMesh, 
@@ -148,7 +148,7 @@ public:
 	 * This method will preserve any existing hard edges, ie each shared triangle-vertex normal is recomputed by averaging 
 	 * the face normals of triangles that reference that shared triangle-vertex normal
 	 */
-	UFUNCTION(BlueprintCallable, Category = "GeometryScript|Normals", meta=(ScriptMethod))
+	UFUNCTION(BlueprintCallable, Category = "GeometryScript|Normals", meta=(ScriptMethod, HidePin = "Debug"))
 	static UE_API UPARAM(DisplayName = "Target Mesh") UDynamicMesh* 
 	RecomputeNormalsForMeshSelection(  
 		UDynamicMesh* TargetMesh, 
@@ -162,7 +162,7 @@ public:
 	 * recompute the new shared triangle-vertex normals using the given CalculateOptions. 
 	 * The normal recomputation is identical to calling RecomputeNormals.
 	 */
-	UFUNCTION(BlueprintCallable, Category = "GeometryScript|Normals", meta=(ScriptMethod))
+	UFUNCTION(BlueprintCallable, Category = "GeometryScript|Normals", meta=(ScriptMethod, HidePin = "Debug"))
 	static UE_API UPARAM(DisplayName = "Target Mesh") UDynamicMesh* 
 	ComputeSplitNormals( 
 		UDynamicMesh* TargetMesh, 
@@ -179,7 +179,7 @@ public:
 	 * @param CalculateOptions Options for computing the normals, if bRecalculateNormals is true
 	 * @param bDeferChangeNotifications If true, no mesh change notification will be sent. Set to true if performing many changes in a loop.
 	 */
-	UFUNCTION(BlueprintCallable, Category = "GeometryScript|Normals", meta = (ScriptMethod))
+	UFUNCTION(BlueprintCallable, Category = "GeometryScript|Normals", meta = (ScriptMethod, HidePin = "Debug"))
 	static UE_API UPARAM(DisplayName = "Target Mesh") UDynamicMesh*
 	SetSplitNormalsAlongSelectedEdges(
 		UDynamicMesh* TargetMesh,
@@ -210,7 +210,7 @@ public:
 	 * Set all normals in the TargetMesh Normals Overlay to the specified per-vertex normals
 	 * @param VertexNormalList per-vertex normals. Size must be equal to the MaxVertexID of TargetMesh  (ie non-compact TargetMesh is supported)
 	 */
-	UFUNCTION(BlueprintCallable, Category = "GeometryScript|Normals", meta=(ScriptMethod, DisplayName="Set Mesh To Per Vertex Normals (From List)"))
+	UFUNCTION(BlueprintCallable, Category = "GeometryScript|Normals", meta=(ScriptMethod, HidePin = "Debug", DisplayName="Set Mesh To Per Vertex Normals (From List)"))
 	static UE_API UPARAM(DisplayName = "Target Mesh") UDynamicMesh*
 	SetMeshPerVertexNormals(
 		UDynamicMesh* TargetMesh,
@@ -241,7 +241,7 @@ public:
 	/**
 	 * Check if the TargetMesh has a Tangents Attribute Layer enabled
 	 */
-	UFUNCTION(BlueprintCallable, Category = "GeometryScript|Normals", meta=(ScriptMethod))
+	UFUNCTION(BlueprintCallable, Category = "GeometryScript|Normals", meta=(ScriptMethod, HidePin = "Debug"))
 	static UE_API UPARAM(DisplayName = "Target Mesh") UDynamicMesh* 
 	GetMeshHasTangents( 
 		UDynamicMesh* TargetMesh,
@@ -251,7 +251,7 @@ public:
 	/**
 	 * Remove any existing Tangents Attribute Layer from the TargetMesh
 	 */
-	UFUNCTION(BlueprintCallable, Category = "GeometryScript|Normals", meta=(ScriptMethod))
+	UFUNCTION(BlueprintCallable, Category = "GeometryScript|Normals", meta=(ScriptMethod, HidePin = "Debug"))
 	static UE_API UPARAM(DisplayName = "Target Mesh") UDynamicMesh* 
 	DiscardTangents( 
 		UDynamicMesh* TargetMesh,
@@ -261,7 +261,7 @@ public:
 	 * Recompute Tangents for the TargetMesh, using the method and settings specified by FGeometryScriptTangentsOptions
 	 * @note If recomputing Tangents for use with a DynamicMeshComponent, it is also necessary to set the Tangents Type on the Component to "Externally Provided"
 	 */
-	UFUNCTION(BlueprintCallable, Category = "GeometryScript|Normals", meta=(ScriptMethod))
+	UFUNCTION(BlueprintCallable, Category = "GeometryScript|Normals", meta=(ScriptMethod, HidePin = "Debug"))
 	static UE_API UPARAM(DisplayName = "Target Mesh") UDynamicMesh* 
 	ComputeTangents( 
 		UDynamicMesh* TargetMesh, 
@@ -275,7 +275,7 @@ public:
 	 * @param TangentYList per-vertex bitangent/binormal vectors. Size must be equal to TangentXList
 	 * @note If setting Tangents for use with a DynamicMeshComponent, it is also necessary to set the Tangents Type on the Component to "Externally Provided"
 	 */
-	UFUNCTION(BlueprintCallable, Category = "GeometryScript|Normals", meta=(ScriptMethod))
+	UFUNCTION(BlueprintCallable, Category = "GeometryScript|Normals", meta=(ScriptMethod, HidePin = "Debug"))
 	static UE_API UPARAM(DisplayName = "Target Mesh") UDynamicMesh*
 	SetMeshPerVertexTangents(
 		UDynamicMesh* TargetMesh,
diff --git a/Engine/Plugins/Runtime/GeometryScripting/Source/GeometryScriptingCore/Public/GeometryScript/MeshPolygroupFunctions.h b/Engine/Plugins/Runtime/GeometryScripting/Source/GeometryScriptingCore/Public/GeometryScript/MeshPolygroupFunctions.h
--- a/Engine/Plugins/Runtime/GeometryScripting/Source/GeometryScriptingCore/Public/GeometryScript/MeshPolygroupFunctions.h
+++ b/Engine/Plugins/Runtime/GeometryScripting/Source/GeometryScriptingCore/Public/GeometryScript/MeshPolygroupFunctions.h
@@ -21,14 +21,14 @@ public:
 	/**
 	* Enables the standard PolyGroup Layer on the Target Mesh.
 	*/
-	UFUNCTION(BlueprintCallable, Category = "GeometryScript|PolyGroups", meta = (ScriptMethod, DisplayName = "Enable PolyGroups"))
+	UFUNCTION(BlueprintCallable, Category = "GeometryScript|PolyGroups", meta = (ScriptMethod, HidePin = "Debug", DisplayName = "Enable PolyGroups"))
 	static UE_API UPARAM(DisplayName = "Target Mesh") UDynamicMesh* 
 	EnablePolygroups( UDynamicMesh* TargetMesh, UGeometryScriptDebug* Debug = nullptr );
 
 	/**
 	* Sets the number of extended PolyGroup Layers on a Mesh.
 	*/
-	UFUNCTION(BlueprintCallable, Category = "GeometryScript|PolyGroups", meta = (ScriptMethod, DisplayName = "Set Num Extended PolyGroup Layers"))
+	UFUNCTION(BlueprintCallable, Category = "GeometryScript|PolyGroups", meta = (ScriptMethod, HidePin = "Debug", DisplayName = "Set Num Extended PolyGroup Layers"))
 	static UE_API UPARAM(DisplayName = "Target Mesh") UDynamicMesh* 
 	SetNumExtendedPolygroupLayers( 
 		UDynamicMesh* TargetMesh, 
@@ -43,7 +43,7 @@ public:
 	 * @param GroupLayer PolyGroup layer with the requested name
 	 * @param bAlreadyExisted If true, the layer already existed and did not need to be added (and the output GroupLayer will refer to that existing layer)
 	 */
-	UFUNCTION(BlueprintCallable, Category = "GeometryScript|PolyGroups", meta = (ScriptMethod, DisplayName = "Add Named PolyGroup Layer"))
+	UFUNCTION(BlueprintCallable, Category = "GeometryScript|PolyGroups", meta = (ScriptMethod, HidePin = "Debug", DisplayName = "Add Named PolyGroup Layer"))
 	static UE_API UPARAM(DisplayName = "Target Mesh") UDynamicMesh* 
 	AddNamedPolygroupLayer( 
 		UDynamicMesh* TargetMesh, 
@@ -60,7 +60,7 @@ public:
 	 * @param GroupLayer PolyGroup layer with the requested name (if Outcome is Success)
 	 * @param Outcome Whether the requested layer was found
 	 */
-	UFUNCTION(BlueprintCallable, Category = "GeometryScript|PolyGroups", meta = (ScriptMethod, DisplayName = "Find Extended PolyGroup Layer By Name", ExpandEnumAsExecs = "Outcome"))
+	UFUNCTION(BlueprintCallable, Category = "GeometryScript|PolyGroups", meta = (ScriptMethod, HidePin = "Debug", DisplayName = "Find Extended PolyGroup Layer By Name", ExpandEnumAsExecs = "Outcome"))
 	static UE_API UPARAM(DisplayName = "Target Mesh") const UDynamicMesh*
 	FindExtendedPolygroupLayerByName(
 		const UDynamicMesh* TargetMesh,
@@ -75,7 +75,7 @@ public:
 	* Resets the triangle PolyGroup assignments within a PolyGroup Layer to the given Clear Value (or 0 if no Clear Value is specified).
 	* Note, this will have no effect if PolyGroups have not been enabled on the mesh, or if the requested Group Layer does not exist. 
 	*/
-	UFUNCTION(BlueprintCallable, Category = "GeometryScript|PolyGroups", meta = (ScriptMethod, DisplayName = "Clear PolyGroups"))
+	UFUNCTION(BlueprintCallable, Category = "GeometryScript|PolyGroups", meta = (ScriptMethod, HidePin = "Debug", DisplayName = "Clear PolyGroups"))
 	static UE_API UPARAM(DisplayName = "Target Mesh") UDynamicMesh* 
 	ClearPolygroups( 
 		UDynamicMesh* TargetMesh, 
@@ -87,7 +87,7 @@ public:
 	* Copies the triangle PolyGroup assignments from one layer on the Target Mesh to another.
 	* Note, this will have no effect if PolyGroups have not been enabled on the mesh, or if one of the requested Group Layers does not exist.
 	*/
-	UFUNCTION(BlueprintCallable, Category = "GeometryScript|PolyGroups", meta = (ScriptMethod, DisplayName = "Copy PolyGroups Layer"))
+	UFUNCTION(BlueprintCallable, Category = "GeometryScript|PolyGroups", meta = (ScriptMethod, HidePin = "Debug", DisplayName = "Copy PolyGroups Layer"))
 	static UE_API UPARAM(DisplayName = "Target Mesh") UDynamicMesh* 
 	CopyPolygroupsLayer( 
 		UDynamicMesh* TargetMesh, 
@@ -101,7 +101,7 @@ public:
 	* @param GroupLayer indicates PolyGroup Layer that will be populated with unique values for each UV island.
 	* @param UVLayer specifies the UV Layer used to construct the PolyGroups.
 	*/
-	UFUNCTION(BlueprintCallable, Category = "GeometryScript|PolyGroups", meta = (ScriptMethod, DisplayName = "Convert UV Islands To PolyGroups"))
+	UFUNCTION(BlueprintCallable, Category = "GeometryScript|PolyGroups", meta = (ScriptMethod, HidePin = "Debug", DisplayName = "Convert UV Islands To PolyGroups"))
 	static UE_API UPARAM(DisplayName = "Target Mesh") UDynamicMesh* 
 	ConvertUVIslandsToPolygroups( 
 		UDynamicMesh* TargetMesh, 
@@ -113,7 +113,7 @@ public:
 	* Creates and assigns a new PolyGroup for each disconnected component of a Mesh. Regions of a mesh are disconnected they do not have a triangle in common.
 	* Note, this will have no effect if the Group Layer does not exist.
 	*/
-	UFUNCTION(BlueprintCallable, Category = "GeometryScript|PolyGroups", meta = (ScriptMethod, DisplayName = "Convert Components To PolyGroups"))
+	UFUNCTION(BlueprintCallable, Category = "GeometryScript|PolyGroups", meta = (ScriptMethod, HidePin = "Debug", DisplayName = "Convert Components To PolyGroups"))
 	static UE_API UPARAM(DisplayName = "Target Mesh") UDynamicMesh* 
 	ConvertComponentsToPolygroups( 
 		UDynamicMesh* TargetMesh, 
@@ -127,7 +127,7 @@ public:
 	* @param CreaseAngle measured in degrees and used when comparing adjacent faces.
 	* @param MinGroupSize the minimum number of triangles in each PolyGroup.
 	*/
-	UFUNCTION(BlueprintCallable, Category = "GeometryScript|PolyGroups", meta = (ScriptMethod, DisplayName = "Compute PolyGroups From Angle Threshold"))
+	UFUNCTION(BlueprintCallable, Category = "GeometryScript|PolyGroups", meta = (ScriptMethod, HidePin = "Debug", DisplayName = "Compute PolyGroups From Angle Threshold"))
 	static UE_API UPARAM(DisplayName = "Target Mesh") UDynamicMesh* 
 	ComputePolygroupsFromAngleThreshold( 
 		UDynamicMesh* TargetMesh, 
@@ -139,7 +139,7 @@ public:
 	/**
 	* Sets PolyGroups by identifying adjacent triangles that form reasonable quads. Note any triangles that do not neatly pair to form quads will receive their own PolyGroup.  
 	*/
-	UFUNCTION(BlueprintCallable, Category = "GeometryScript|PolyGroups", meta = (ScriptMethod, DisplayName = "Compute PolyGroups From Polygon Detection"))
+	UFUNCTION(BlueprintCallable, Category = "GeometryScript|PolyGroups", meta = (ScriptMethod, HidePin = "Debug", DisplayName = "Compute PolyGroups From Polygon Detection"))
 	static UE_API UPARAM(DisplayName = "Target Mesh") UDynamicMesh* 
 	ComputePolygroupsFromPolygonDetection( 
 		UDynamicMesh* TargetMesh, 
@@ -172,7 +172,7 @@ public:
 	 * @param PolyGroup ID identifies the triangles in the Target Mesh to delete.
 	 * @param NumDeleted on return will contain the number of triangles deleted from the Target Mesh.
 	 */
-	UFUNCTION(BlueprintCallable, Category = "GeometryScript|PolyGroups", meta = (ScriptMethod, DisplayName = "Delete Triangles In PolyGroup"))
+	UFUNCTION(BlueprintCallable, Category = "GeometryScript|PolyGroups", meta = (ScriptMethod, HidePin = "Debug", DisplayName = "Delete Triangles In PolyGroup"))
 	static UE_API UPARAM(DisplayName = "Target Mesh") UDynamicMesh* 
 	DeleteTrianglesInPolygroup( 
 		UDynamicMesh* TargetMesh, 
diff --git a/Engine/Plugins/Runtime/GeometryScripting/Source/GeometryScriptingCore/Public/GeometryScript/MeshPrimitiveFunctions.h b/Engine/Plugins/Runtime/GeometryScripting/Source/GeometryScriptingCore/Public/GeometryScript/MeshPrimitiveFunctions.h
--- a/Engine/Plugins/Runtime/GeometryScripting/Source/GeometryScriptingCore/Public/GeometryScript/MeshPrimitiveFunctions.h
+++ b/Engine/Plugins/Runtime/GeometryScripting/Source/GeometryScriptingCore/Public/GeometryScript/MeshPrimitiveFunctions.h
@@ -163,7 +163,7 @@ public:
 	/**
 	* Appends a 3D box to the Target Mesh.
 	*/
-	UFUNCTION(BlueprintCallable, Category = "GeometryScript|Primitives", meta=(ScriptMethod))
+	UFUNCTION(BlueprintCallable, Category = "GeometryScript|Primitives", meta=(ScriptMethod, HidePin = "Debug"))
 	static UE_API UPARAM(DisplayName = "Target Mesh") UDynamicMesh* 
 	AppendBox( 
 		UDynamicMesh* TargetMesh, 
@@ -182,7 +182,7 @@ public:
 	 * Appends a 3D box to the Target Mesh
 	 * Also creates matching simple collision
 	 */
-	UFUNCTION(BlueprintCallable, Category = "GeometryScript|Primitives", meta = (ScriptMethod))
+	UFUNCTION(BlueprintCallable, Category = "GeometryScript|Primitives", meta = (ScriptMethod, HidePin = "Debug"))
 	static UE_API UPARAM(DisplayName = "Target Mesh") UDynamicMesh*
 	AppendBoxWithCollision(
 		UDynamicMesh* TargetMesh,
@@ -201,7 +201,7 @@ public:
 	/**
 	 * Appends a 3D box to the Target Mesh with dimensions and origin taken from the input Box
 	 */
-	UFUNCTION(BlueprintCallable, Category = "GeometryScript|Primitives", meta=(ScriptMethod))
+	UFUNCTION(BlueprintCallable, Category = "GeometryScript|Primitives", meta=(ScriptMethod, HidePin = "Debug"))
 	static UE_API UPARAM(DisplayName = "Target Mesh") UDynamicMesh* 
 	AppendBoundingBox( 
 		UDynamicMesh* TargetMesh, 
@@ -217,7 +217,7 @@ public:
 	 * Appends a 3D box to the Target Mesh with dimensions and origin taken from the input Box
 	 * Also creates matching simple collision
 	 */
-	UFUNCTION(BlueprintCallable, Category = "GeometryScript|Primitives", meta=(ScriptMethod))
+	UFUNCTION(BlueprintCallable, Category = "GeometryScript|Primitives", meta=(ScriptMethod, HidePin = "Debug"))
 	static UE_API UPARAM(DisplayName = "Target Mesh") UDynamicMesh* 
 	AppendBoundingBoxWithCollision( 
 		UDynamicMesh* TargetMesh, 
@@ -234,7 +234,7 @@ public:
 	/**
 	* Appends a 3D Sphere triangulated using latitude/longitude topology to the Target Mesh.
 	*/ 
-	UFUNCTION(BlueprintCallable, Category = "GeometryScript|Primitives", meta=(ScriptMethod))
+	UFUNCTION(BlueprintCallable, Category = "GeometryScript|Primitives", meta=(ScriptMethod, HidePin = "Debug"))
 	static UE_API UPARAM(DisplayName = "Target Mesh") UDynamicMesh* 
 	AppendSphereLatLong( 
 		UDynamicMesh* TargetMesh, 
@@ -251,7 +251,7 @@ public:
 	* Also creates matching simple collision. Note: If transform has non-uniform scale, collision shape will be a uniform-scale approximation
 	* matching the behavior of the physics system -- specifically, it will scale the sphere radius by the smallest axis scale.
 	*/
-	UFUNCTION(BlueprintCallable, Category = "GeometryScript|Primitives", meta=(ScriptMethod))
+	UFUNCTION(BlueprintCallable, Category = "GeometryScript|Primitives", meta=(ScriptMethod, HidePin = "Debug"))
 	static UE_API UPARAM(DisplayName = "Target Mesh") UDynamicMesh* 
 	AppendSphereLatLongWithCollision( 
 		UDynamicMesh* TargetMesh, 
@@ -267,7 +267,7 @@ public:
 	/**
 	* Appends a 3D sphere triangulated using box topology to the Target Mesh.
 	*/ 
-	UFUNCTION(BlueprintCallable, Category = "GeometryScript|Primitives", meta=(ScriptMethod))
+	UFUNCTION(BlueprintCallable, Category = "GeometryScript|Primitives", meta=(ScriptMethod, HidePin = "Debug"))
 	static UE_API UPARAM(DisplayName = "Target Mesh") UDynamicMesh* 
 	AppendSphereBox( 
 		UDynamicMesh* TargetMesh, 
@@ -285,7 +285,7 @@ public:
 	 * Also creates matching simple collision. Note: If transform has non-uniform scale, collision shape will be a uniform-scale approximation
 	 * matching the behavior of the physics system -- specifically, it will scale the sphere radius by the smallest axis scale.
 	 */
-	UFUNCTION(BlueprintCallable, Category = "GeometryScript|Primitives", meta = (ScriptMethod))
+	UFUNCTION(BlueprintCallable, Category = "GeometryScript|Primitives", meta = (ScriptMethod, HidePin = "Debug"))
 	static UE_API UPARAM(DisplayName = "Target Mesh") UDynamicMesh*
 	AppendSphereBoxWithCollision(
 		UDynamicMesh* TargetMesh,
@@ -302,7 +302,7 @@ public:
 	/**
 	* Appends a 3D Capsule to the Target Mesh.
 	*/ 
-	UFUNCTION(BlueprintCallable, Category = "GeometryScript|Primitives", meta=(ScriptMethod))
+	UFUNCTION(BlueprintCallable, Category = "GeometryScript|Primitives", meta=(ScriptMethod, HidePin = "Debug"))
 	static UE_API UPARAM(DisplayName = "Target Mesh") UDynamicMesh* 
 	AppendCapsule( 
 		UDynamicMesh* TargetMesh, 
@@ -321,7 +321,7 @@ public:
 	 * Also creates matching simple collision. Note: If transform has non-uniform scale, collision shape will be a uniform-scale approximation
 	 * matching the behavior of the physics system -- specifically, it will scale the radius by the larger of the X, Y axis scales, and the length by the Z axis scale
 	*/ 
-	UFUNCTION(BlueprintCallable, Category = "GeometryScript|Primitives", meta=(ScriptMethod))
+	UFUNCTION(BlueprintCallable, Category = "GeometryScript|Primitives", meta=(ScriptMethod, HidePin = "Debug"))
 	static UE_API UPARAM(DisplayName = "Target Mesh") UDynamicMesh* 
 	AppendCapsuleWithCollision( 
 		UDynamicMesh* TargetMesh, 
@@ -355,7 +355,7 @@ public:
 	/**
 	* Appends a 3D Cylinder (with optional end caps) to the Target Mesh.
 	*/
-	UFUNCTION(BlueprintCallable, Category = "GeometryScript|Primitives", meta=(ScriptMethod))
+	UFUNCTION(BlueprintCallable, Category = "GeometryScript|Primitives", meta=(ScriptMethod, HidePin = "Debug"))
 	static UE_API UPARAM(DisplayName = "Target Mesh") UDynamicMesh* 
 	AppendCylinder( 
 		UDynamicMesh* TargetMesh, 
@@ -372,7 +372,7 @@ public:
 	/**
 	* Appends a 3D cone to the Target Mesh.
 	*/ 
-	UFUNCTION(BlueprintCallable, Category = "GeometryScript|Primitives", meta=(ScriptMethod))
+	UFUNCTION(BlueprintCallable, Category = "GeometryScript|Primitives", meta=(ScriptMethod, HidePin = "Debug"))
 	static UE_API UPARAM(DisplayName = "Target Mesh") UDynamicMesh* 
 	AppendCone( 
 		UDynamicMesh* TargetMesh, 
@@ -390,7 +390,7 @@ public:
 	/**
 	* Appends a 3D torus (donut) or partial torus to the Target Mesh.
 	*/
-	UFUNCTION(BlueprintCallable, Category = "GeometryScript|Primitives", meta=(ScriptMethod))
+	UFUNCTION(BlueprintCallable, Category = "GeometryScript|Primitives", meta=(ScriptMethod, HidePin = "Debug"))
 	static UE_API UPARAM(DisplayName = "Target Mesh") UDynamicMesh* 
 	AppendTorus( 
 		UDynamicMesh* TargetMesh, 
@@ -409,7 +409,7 @@ public:
 	 * Polygon should be oriented counter-clockwise to produce a correctly-oriented shape, otherwise it will be inside-out
 	 * Polygon endpoint is not repeated.
 	 */
-	UFUNCTION(BlueprintCallable, Category = "GeometryScript|Primitives", meta=(ScriptMethod))
+	UFUNCTION(BlueprintCallable, Category = "GeometryScript|Primitives", meta=(ScriptMethod, HidePin = "Debug"))
 	static UE_API UPARAM(DisplayName = "Target Mesh") UDynamicMesh* 
 	AppendRevolvePolygon( 
 		UDynamicMesh* TargetMesh, 
@@ -424,7 +424,7 @@ public:
     /**
 	* Revolves a 2D polygon on a helical path, like one used to create a vertical spiral, appending the result to the Target Mesh.
 	*/ 
-	UFUNCTION(BlueprintCallable, Category = "GeometryScript|Primitives", meta = (ScriptMethod))
+	UFUNCTION(BlueprintCallable, Category = "GeometryScript|Primitives", meta = (ScriptMethod, HidePin = "Debug"))
 	static UE_API UPARAM(DisplayName = "Target Mesh") UDynamicMesh*
 	AppendSpiralRevolvePolygon(
 		UDynamicMesh* TargetMesh,
@@ -440,7 +440,7 @@ public:
 	/**
 	* Revolves an open 2D path, with optional top and bottom end caps, appending the result to the Target Mesh.
 	*/ 
-	UFUNCTION(BlueprintCallable, Category = "GeometryScript|Primitives", meta=(ScriptMethod))
+	UFUNCTION(BlueprintCallable, Category = "GeometryScript|Primitives", meta=(ScriptMethod, HidePin = "Debug"))
 	static UE_API UPARAM(DisplayName = "Target Mesh") UDynamicMesh* 
 	AppendRevolvePath( 
 		UDynamicMesh* TargetMesh, 
@@ -466,7 +466,7 @@ public:
 	 * @param RotationAngleDeg Rotation applied to the 2D Polygon. Positive rotation rotates clockwise, ie Up/+Z/+V towards Right/+Y/+U. This Rotation is applied before any rotation in the SweepPath Transforms.
 	 * @param MiterLimit If > 1, maximum scaling to apply at sharp turns in the curve path to keep the apparent swept profile from shrinking, and sweep path frames will be aligned to the path direction
 	 */
-	UFUNCTION(BlueprintCallable, Category = "GeometryScript|Primitives", meta=(ScriptMethod, AutoCreateRefTerm="PolylineTexParamU, SweepPathTexParamV"))
+	UFUNCTION(BlueprintCallable, Category = "GeometryScript|Primitives", meta=(ScriptMethod, HidePin = "Debug", AutoCreateRefTerm="PolylineTexParamU, SweepPathTexParamV"))
 	static UE_API UPARAM(DisplayName = "Target Mesh") UDynamicMesh* 
 	AppendSweepPolyline( 
 		UDynamicMesh* TargetMesh, 
@@ -487,7 +487,7 @@ public:
 	 * Polygon should be oriented counter-clockwise to produce a correctly-oriented shape, otherwise it will be inside-out
 	 * Polygon endpoint is not repeated.
 	 */
-	UFUNCTION(BlueprintCallable, Category = "GeometryScript|Primitives", meta=(ScriptMethod))
+	UFUNCTION(BlueprintCallable, Category = "GeometryScript|Primitives", meta=(ScriptMethod, HidePin = "Debug"))
 	static UE_API UPARAM(DisplayName = "Target Mesh") UDynamicMesh* 
 	AppendSimpleExtrudePolygon( 
 		UDynamicMesh* TargetMesh, 
@@ -511,7 +511,7 @@ public:
 	 * @param RotationAngleDeg Rotation applied to the 2D Polygon. Positive rotation rotates clockwise, ie Up/+Z/+V towards Right/+Y/+U
 	 * @param MiterLimit If > 1, maximum scaling to apply at sharp turns in the curve path to keep the apparent swept profile from shrinking
 	 */
-	UFUNCTION(BlueprintCallable, Category = "GeometryScript|Primitives", meta=(ScriptMethod))
+	UFUNCTION(BlueprintCallable, Category = "GeometryScript|Primitives", meta=(ScriptMethod, HidePin = "Debug"))
 	static UE_API UPARAM(DisplayName = "Target Mesh") UDynamicMesh* 
 	AppendSimpleSweptPolygon( 
 		UDynamicMesh* TargetMesh, 
@@ -540,7 +540,7 @@ public:
 	 * @param RotationAngleDeg Rotation applied to the 2D Polygon. Positive rotation rotates clockwise, ie Up/+Z/+V towards Right/+Y/+U. This Rotation is applied before any rotation in the SweepPath Transforms.
 	 * @param MiterLimit If > 1, maximum scaling to apply at sharp turns in the curve path to keep the apparent swept profile from shrinking, and sweep path frames will be aligned to the path direction
 	 */
-	UFUNCTION(BlueprintCallable, Category = "GeometryScript|Primitives", meta=(ScriptMethod))
+	UFUNCTION(BlueprintCallable, Category = "GeometryScript|Primitives", meta=(ScriptMethod, HidePin = "Debug"))
 	static UE_API UPARAM(DisplayName = "Target Mesh") UDynamicMesh* 
 	AppendSweepPolygon( 
 		UDynamicMesh* TargetMesh, 
@@ -560,7 +560,7 @@ public:
 	/**
 	* Appends a planar Rectangle to a Dynamic Mesh.
 	*/ 
-	UFUNCTION(BlueprintCallable, Category = "GeometryScript|Primitives", meta=(ScriptMethod))
+	UFUNCTION(BlueprintCallable, Category = "GeometryScript|Primitives", meta=(ScriptMethod, HidePin = "Debug"))
 	static UE_API UPARAM(DisplayName = "Target Mesh") UDynamicMesh* 
 	AppendRectangleXY( 
 		UDynamicMesh* TargetMesh, 
@@ -575,7 +575,7 @@ public:
 	/**
 	* Appends a planar Rectangle with Rounded Corners (RoundRect) to the Target Mesh.
 	*/
-	UFUNCTION(BlueprintCallable, Category = "GeometryScript|Primitives", meta=(ScriptMethod))
+	UFUNCTION(BlueprintCallable, Category = "GeometryScript|Primitives", meta=(ScriptMethod, HidePin = "Debug"))
 	static UE_API UPARAM(DisplayName = "Target Mesh") UDynamicMesh* 
 	AppendRoundRectangleXY( 
 		UDynamicMesh* TargetMesh, 
@@ -592,7 +592,7 @@ public:
 	/**
 	* Appends a planar disc to the Target Mesh.
 	*/ 
-	UFUNCTION(BlueprintCallable, Category = "GeometryScript|Primitives", meta=(ScriptMethod))
+	UFUNCTION(BlueprintCallable, Category = "GeometryScript|Primitives", meta=(ScriptMethod, HidePin = "Debug"))
 	static UE_API UPARAM(DisplayName = "Target Mesh") UDynamicMesh* 
 	AppendDisc( 
 		UDynamicMesh* TargetMesh, 
@@ -611,7 +611,7 @@ public:
 	* Polygon should be oriented counter-clockwise to produce a correctly-oriented shape, otherwise it will be inside-out
 	* Polygon endpoint is not repeated.
 	*/
-	UFUNCTION(BlueprintCallable, Category = "GeometryScript|Primitives", meta=(ScriptMethod))
+	UFUNCTION(BlueprintCallable, Category = "GeometryScript|Primitives", meta=(ScriptMethod, HidePin = "Debug"))
 	static UE_API UPARAM(DisplayName = "Target Mesh") UDynamicMesh* 
 	AppendTriangulatedPolygon( 
 		UDynamicMesh* TargetMesh, 
@@ -627,7 +627,7 @@ public:
 	* Uses Ear Clipping-based triangulation. Output vertices will always be 1:1 with input vertices.
 	* Polygon endpoint is not repeated.
 	*/
-	UFUNCTION(BlueprintCallable, Category = "GeometryScript|Primitives", meta = (ScriptMethod))
+	UFUNCTION(BlueprintCallable, Category = "GeometryScript|Primitives", meta = (ScriptMethod, HidePin = "Debug"))
 	static UE_API UPARAM(DisplayName = "Target Mesh") UDynamicMesh*
 	AppendTriangulatedPolygon3D(
 		UDynamicMesh* TargetMesh,
@@ -639,7 +639,7 @@ public:
 	/**
 	* Appends a linear staircase to the Target Mesh.
 	*/
-	UFUNCTION(BlueprintCallable, Category = "GeometryScript|Primitives", meta=(ScriptMethod))
+	UFUNCTION(BlueprintCallable, Category = "GeometryScript|Primitives", meta=(ScriptMethod, HidePin = "Debug"))
 	static UE_API UPARAM(DisplayName = "Target Mesh") UDynamicMesh* 
 	AppendLinearStairs( 
 		UDynamicMesh* TargetMesh, 
@@ -655,7 +655,7 @@ public:
 	/**
 	* Appends a rising circular staircase to the Target Mesh.
 	*/ 
-	UFUNCTION(BlueprintCallable, Category = "GeometryScript|Primitives", meta=(ScriptMethod))
+	UFUNCTION(BlueprintCallable, Category = "GeometryScript|Primitives", meta=(ScriptMethod, HidePin = "Debug"))
 	static UE_API UPARAM(DisplayName = "Target Mesh") UDynamicMesh* 
 	AppendCurvedStairs( 
 		UDynamicMesh* TargetMesh, 
@@ -672,7 +672,7 @@ public:
 	/**
 	* Generates triangulated Voronoi Cells from the provided Voronoi Sites, identifying each with PolyGroups, and appends to the Target Mesh.
 	*/ 
-	UFUNCTION(BlueprintCallable, Category = "GeometryScript|Primitives", meta = (ScriptMethod))
+	UFUNCTION(BlueprintCallable, Category = "GeometryScript|Primitives", meta = (ScriptMethod, HidePin = "Debug"))
 	static UE_API UPARAM(DisplayName = "Target Mesh") UDynamicMesh*
 	AppendVoronoiDiagram2D(
 		UDynamicMesh* TargetMesh,
@@ -688,7 +688,7 @@ public:
 	* On success, all vertices are always appended to the output mesh, though duplicate vertices will not be used in any triangles and may optionally be removed.
 	* Use PositionsToVertexIDs to map indices in the input VertexPositions array to vertex IDs in the Dynamic Mesh.
 	*/
-	UFUNCTION(BlueprintCallable, Category = "GeometryScript|Primitives", meta = (ScriptMethod, AutoCreateRefTerm = "ConstrainedEdges"))
+	UFUNCTION(BlueprintCallable, Category = "GeometryScript|Primitives", meta = (ScriptMethod, HidePin = "Debug", AutoCreateRefTerm = "ConstrainedEdges"))
 	static UE_API UPARAM(DisplayName = "Target Mesh") UDynamicMesh*
 	AppendDelaunayTriangulation2D(
 		UDynamicMesh* TargetMesh,
@@ -723,7 +723,7 @@ public:
 	* @param bTriangulationError Reports whether the triangulation contains errors, typically due to intersecting edges in the input. Consider pre-processing the PolygonsList with PolygonsUnion to resolve intersections and prevent this error.
 	* @param bStillAppendOnError Whether to still append a best-effort triangulation in error cases. Often this will be a triangulation that does not quite match the polygon shape near intersecting edges in the input, but otherwise is as-expected.
 	*/
-	UFUNCTION(BlueprintCallable, Category = "GeometryScript|Primitives")
+	UFUNCTION(BlueprintCallable, Category = "GeometryScript|Primitives", meta = (HidePin = "Debug"))
 	static UE_API UPARAM(DisplayName = "Target Mesh") UDynamicMesh*
 	AppendPolygonListTriangulation(
 		UDynamicMesh* TargetMesh,
@@ -739,7 +739,7 @@ public:
 	 * Appends Simple Collision shapes to the Target Mesh, triangulated as specified by Triangulation Options
 	 * @param Transform Transform to be applied to simple collision shapes, following the method by which simple collision shapes are transformed -- so, e.g., spheres will not be non-uniformly scaled
 	 */
-	UFUNCTION(BlueprintCallable, DisplayName = "Append Simple Collision Shapes to Mesh", Category = "GeometryScript|Primitives")
+	UFUNCTION(BlueprintCallable, DisplayName = "Append Simple Collision Shapes to Mesh", Category = "GeometryScript|Primitives", meta = (HidePin = "Debug"))
 	static UE_API UPARAM(DisplayName = "Target Mesh") UDynamicMesh*
 	AppendSimpleCollisionShapes(
 		UDynamicMesh* TargetMesh,
@@ -753,7 +753,7 @@ public:
 	/**
 	 * Appends the spheres in the Sphere Covering to the Target Mesh
 	 */
-	UFUNCTION(BlueprintCallable, Category = "GeometryScript|Primitives")
+	UFUNCTION(BlueprintCallable, Category = "GeometryScript|Primitives", meta = (HidePin = "Debug"))
 	static UE_API UPARAM(DisplayName = "Target Mesh") UDynamicMesh*
 	AppendSphereCovering(
 		UDynamicMesh* TargetMesh,
@@ -784,7 +784,7 @@ public:
 	 * Incorrectly divides the input DimensionX and DimensionY by 2.
 	 * @warning It is strongly recommended that callers of this function update to the current AppendRectangleXY function!
 	 */
-	UFUNCTION(BlueprintCallable, Category = "GeometryScript|Compatibility", meta=(ScriptMethod))
+	UFUNCTION(BlueprintCallable, Category = "GeometryScript|Compatibility", meta=(ScriptMethod, HidePin = "Debug"))
 	static UE_API UPARAM(DisplayName = "Target Mesh") UDynamicMesh* 
 	AppendRectangle_Compatibility_5_0( 
 		UDynamicMesh* TargetMesh, 
@@ -802,7 +802,7 @@ public:
 	 * Incorrectly divides the input DimensionX and DimensionY by 2.
 	 * @warning It is strongly recommended that callers of this function update to the current AppendRoundRectangleXY function!
 	 */
-	UFUNCTION(BlueprintCallable, Category = "GeometryScript|Compatibility", meta=(ScriptMethod))
+	UFUNCTION(BlueprintCallable, Category = "GeometryScript|Compatibility", meta=(ScriptMethod, HidePin = "Debug"))
 	static UE_API UPARAM(DisplayName = "Target Mesh") UDynamicMesh* 
 	AppendRoundRectangle_Compatibility_5_0( 
 		UDynamicMesh* TargetMesh, 
diff --git a/Engine/Plugins/Runtime/GeometryScripting/Source/GeometryScriptingCore/Public/GeometryScript/MeshRemeshFunctions.h b/Engine/Plugins/Runtime/GeometryScripting/Source/GeometryScriptingCore/Public/GeometryScript/MeshRemeshFunctions.h
--- a/Engine/Plugins/Runtime/GeometryScripting/Source/GeometryScriptingCore/Public/GeometryScript/MeshRemeshFunctions.h
+++ b/Engine/Plugins/Runtime/GeometryScripting/Source/GeometryScriptingCore/Public/GeometryScript/MeshRemeshFunctions.h
@@ -153,7 +153,7 @@ public:
 	 * Apply Uniform Remeshing to the TargetMesh. 
 	 * @warning this function can be quite expensive. The results may be non-deterministic, and are expected to change in future versions.
 	 */
-	UFUNCTION(BlueprintCallable, Category = "GeometryScript|Simplification", meta=(ScriptMethod))
+	UFUNCTION(BlueprintCallable, Category = "GeometryScript|Simplification", meta=(ScriptMethod, HidePin = "Debug"))
 	static UE_API UPARAM(DisplayName = "Target Mesh") UDynamicMesh* 
 	ApplyUniformRemesh(  
 		UDynamicMesh* TargetMesh, 
diff --git a/Engine/Plugins/Runtime/GeometryScripting/Source/GeometryScriptingCore/Public/GeometryScript/MeshRepairFunctions.h b/Engine/Plugins/Runtime/GeometryScripting/Source/GeometryScriptingCore/Public/GeometryScript/MeshRepairFunctions.h
--- a/Engine/Plugins/Runtime/GeometryScripting/Source/GeometryScriptingCore/Public/GeometryScript/MeshRepairFunctions.h
+++ b/Engine/Plugins/Runtime/GeometryScripting/Source/GeometryScriptingCore/Public/GeometryScript/MeshRepairFunctions.h
@@ -220,7 +220,7 @@ public:
 	/**
 	* Compacts the mesh's vertices and triangles to remove any "holes" in the Vertex ID or Triangle ID lists.
 	*/
-	UFUNCTION(BlueprintCallable, Category = "GeometryScript|Repair", meta=(ScriptMethod))
+	UFUNCTION(BlueprintCallable, Category = "GeometryScript|Repair", meta=(ScriptMethod, HidePin = "Debug"))
 	static UE_API UPARAM(DisplayName = "Target Mesh") UDynamicMesh* 
 	CompactMesh(  
 		UDynamicMesh* TargetMesh, 
@@ -229,7 +229,7 @@ public:
 	/**
 	 * Remove vertices that are not used by any triangles. Note: Does not update the IDs of any remaining vertices; use CompactMesh to do so.
 	 */
-	UFUNCTION(BlueprintCallable, Category = "GeometryScript|Repair", meta = (ScriptMethod))
+	UFUNCTION(BlueprintCallable, Category = "GeometryScript|Repair", meta = (ScriptMethod, HidePin = "Debug"))
 	static UE_API UPARAM(DisplayName = "Target Mesh") UDynamicMesh*
 	RemoveUnusedVertices(
 		UDynamicMesh* TargetMesh,
@@ -239,7 +239,7 @@ public:
 	 * Snap vertices on open edges to the closest compatible open boundary, if found within the tolerance distance
 	 * Unlike ResolveMeshTJunctions, does not introduce new vertices to the mesh
 	 */
-	UFUNCTION(BlueprintCallable, Category = "GeometryScript|Repair", meta = (ScriptMethod))
+	UFUNCTION(BlueprintCallable, Category = "GeometryScript|Repair", meta = (ScriptMethod, HidePin = "Debug"))
 	static UE_API UPARAM(DisplayName = "Target Mesh") UDynamicMesh*
 	SnapMeshOpenBoundaries(
 		UDynamicMesh* TargetMesh,
@@ -250,7 +250,7 @@ public:
 	/**
 	* Attempts to resolve T-Junctions in the mesh by addition of vertices and welding.
 	*/
-	UFUNCTION(BlueprintCallable, Category = "GeometryScript|Repair", meta=(ScriptMethod))
+	UFUNCTION(BlueprintCallable, Category = "GeometryScript|Repair", meta=(ScriptMethod, HidePin = "Debug"))
 	static UE_API UPARAM(DisplayName = "Target Mesh") UDynamicMesh* 
 	ResolveMeshTJunctions(  
 		UDynamicMesh* TargetMesh, 
@@ -260,7 +260,7 @@ public:
 	/**
 	* Welds any open boundary edges of the mesh together if possible in order to remove "cracks."
 	*/
-	UFUNCTION(BlueprintCallable, Category = "GeometryScript|Repair", meta=(ScriptMethod))
+	UFUNCTION(BlueprintCallable, Category = "GeometryScript|Repair", meta=(ScriptMethod, HidePin = "Debug"))
 	static UE_API UPARAM(DisplayName = "Target Mesh") UDynamicMesh* 
 	WeldMeshEdges(  
 		UDynamicMesh* TargetMesh, 
@@ -273,7 +273,7 @@ public:
 	* @param NumFilledHoles reports the number of holes filled by the function.
 	* @param NumFailedHolesFills reports the detected holes that were unable to be filled.
 	*/
-	UFUNCTION(BlueprintCallable, Category = "GeometryScript|Repair", meta=(ScriptMethod))
+	UFUNCTION(BlueprintCallable, Category = "GeometryScript|Repair", meta=(ScriptMethod, HidePin = "Debug"))
 	static UE_API UPARAM(DisplayName = "Target Mesh") UDynamicMesh* 
 	FillAllMeshHoles(  
 		UDynamicMesh* TargetMesh, 
@@ -285,7 +285,7 @@ public:
 	/*
 	* Removes connected components of the mesh that have a volume, area, or triangle count below a threshold as specified by the Options.
 	*/
-	UFUNCTION(BlueprintCallable, Category = "GeometryScript|Repair", meta=(ScriptMethod))
+	UFUNCTION(BlueprintCallable, Category = "GeometryScript|Repair", meta=(ScriptMethod, HidePin = "Debug"))
 	static UE_API UPARAM(DisplayName = "Target Mesh") UDynamicMesh* 
 	RemoveSmallComponents(  
 		UDynamicMesh* TargetMesh, 
@@ -296,7 +296,7 @@ public:
 	* Removes any triangles in the mesh that are not visible from the exterior view, under various definitions of "visible" and "outside"
 	* as specified by the Options.
 	*/
-	UFUNCTION(BlueprintCallable, Category = "GeometryScript|Repair", meta=(ScriptMethod, KeyWords = "Jacket Occluded"))
+	UFUNCTION(BlueprintCallable, Category = "GeometryScript|Repair", meta=(ScriptMethod, HidePin = "Debug", KeyWords = "Jacket Occluded"))
 	static UE_API UPARAM(DisplayName = "Target Mesh") UDynamicMesh* 
 	RemoveHiddenTriangles(  
 		UDynamicMesh* TargetMesh, 
@@ -310,7 +310,7 @@ public:
 	* @param TransparentSelection Optionally indicate parts of the mesh that should not be considered occluding, e.g. transparent triangles
 	* @param ResultSelection The fully occluded triangles of the target mesh (based on the options / transparencies, and excluding any IgnoreSelection triangles)
 	*/
-	UFUNCTION(BlueprintCallable, Category = "GeometryScript|Repair", meta = (ScriptMethod, KeyWords = "Remove Jacket Occluded"))
+	UFUNCTION(BlueprintCallable, Category = "GeometryScript|Repair", meta = (ScriptMethod, HidePin = "Debug", KeyWords = "Remove Jacket Occluded"))
 	static UE_API UPARAM(DisplayName = "Target Mesh") UDynamicMesh* 
 	SelectHiddenTrianglesFromOutside(
 		UDynamicMesh* TargetMesh, 
@@ -324,7 +324,7 @@ public:
 	* Splits Bowties in the mesh and/or the attributes.  A Bowtie is formed when a single vertex is connected to more than two boundary edges, 
 	* and splitting duplicates the shared vertex so each triangle will have a unique copy.
 	*/
-	UFUNCTION(BlueprintCallable, Category = "GeometryScript|Repair", meta=(ScriptMethod))
+	UFUNCTION(BlueprintCallable, Category = "GeometryScript|Repair", meta=(ScriptMethod, HidePin = "Debug"))
 	static UE_API UPARAM(DisplayName = "Target Mesh") UDynamicMesh* 
 	SplitMeshBowties(  
 		UDynamicMesh* TargetMesh, 
@@ -335,7 +335,7 @@ public:
 	/**
 	* Removes triangles that have area or edge length below specified amounts depending on the Options requested.
 	*/
-	UFUNCTION(BlueprintCallable, Category = "GeometryScript|Repair", meta=(ScriptMethod))
+	UFUNCTION(BlueprintCallable, Category = "GeometryScript|Repair", meta=(ScriptMethod, HidePin = "Debug"))
 	static UE_API UPARAM(DisplayName = "Target Mesh") UDynamicMesh* 
 	RepairMeshDegenerateGeometry(  
 		UDynamicMesh* TargetMesh, 
diff --git a/Engine/Plugins/Runtime/GeometryScripting/Source/GeometryScriptingCore/Public/GeometryScript/MeshSamplingFunctions.h b/Engine/Plugins/Runtime/GeometryScripting/Source/GeometryScriptingCore/Public/GeometryScript/MeshSamplingFunctions.h
--- a/Engine/Plugins/Runtime/GeometryScripting/Source/GeometryScriptingCore/Public/GeometryScript/MeshSamplingFunctions.h
+++ b/Engine/Plugins/Runtime/GeometryScripting/Source/GeometryScriptingCore/Public/GeometryScript/MeshSamplingFunctions.h
@@ -114,7 +114,7 @@ public:
 	/**
 	 * Compute a 'uniform random' (not 'uniform spacing') point sampling over the mesh surface
 	 */
-	UFUNCTION(BlueprintCallable, Category = "GeometryScript|MeshSampling", meta = (ScriptMethod))
+	UFUNCTION(BlueprintCallable, Category = "GeometryScript|MeshSampling", meta = (ScriptMethod, HidePin = "Debug"))
 	static UE_API UPARAM(DisplayName = "Target Mesh") const UDynamicMesh*
 	ComputeUniformRandomPointSampling(
 		const UDynamicMesh* TargetMesh,
@@ -131,7 +131,7 @@ public:
 	 * @param Samples output list of sample points. Transform Location is sample position, Rotation orients Z with the triangle normal
 	 * @param TriangleIDs TriangleID that contains each sample point. Length is the same as Samples array.
 	 */
-	UFUNCTION(BlueprintCallable, Category = "GeometryScript|MeshSampling", meta=(ScriptMethod))
+	UFUNCTION(BlueprintCallable, Category = "GeometryScript|MeshSampling", meta=(ScriptMethod, HidePin = "Debug"))
 	static UE_API UPARAM(DisplayName = "Target Mesh") UDynamicMesh* 
 	ComputePointSampling( 
 		UDynamicMesh* TargetMesh, 
@@ -145,7 +145,7 @@ public:
 	 * The sample points have radii in the range [Options.SamplingRadius, NonUniformOptions.MaxSamplingRadius], and
 	 * are non-overlapping, ie the distance between two points is always larger than the sum of their respective radii.
 	 */
-	UFUNCTION(BlueprintCallable, Category = "GeometryScript|MeshSampling", meta=(ScriptMethod))
+	UFUNCTION(BlueprintCallable, Category = "GeometryScript|MeshSampling", meta=(ScriptMethod, HidePin = "Debug"))
 	static UE_API UPARAM(DisplayName = "Target Mesh") UDynamicMesh* 
 	ComputeNonUniformPointSampling( 
 		UDynamicMesh* TargetMesh, 
@@ -162,7 +162,7 @@ public:
 	 * are non-overlapping, ie the distance between two points is always larger than the sum of their respective radii.
 	 * @param VertexWeights defines a per-vertex weight in range [0,1], these are interpolated to create a scalar field over the mesh triangles which is used to weight the sampling radii
 	 */
-	UFUNCTION(BlueprintCallable, Category = "GeometryScript|MeshSampling", meta=(ScriptMethod))
+	UFUNCTION(BlueprintCallable, Category = "GeometryScript|MeshSampling", meta=(ScriptMethod, HidePin = "Debug"))
 	static UE_API UPARAM(DisplayName = "Target Mesh") UDynamicMesh* 
 	ComputeVertexWeightedPointSampling( 
 		UDynamicMesh* TargetMesh, 
@@ -180,7 +180,7 @@ public:
 	 * @param Box     Bounding Box containing the scene to be captured
 	 * @param Options Defines the Camera viewing directions into the box and other Camera parameters
 	 */
-	UFUNCTION(BlueprintCallable, Category = "GeometryScript|MeshSampling", meta = (AutoCreateRefTerm = "Options"))
+	UFUNCTION(BlueprintCallable, Category = "GeometryScript|MeshSampling", meta = (HidePin = "Debug", AutoCreateRefTerm = "Options"))
 	static UE_API void
 	ComputeRenderCaptureCamerasForBox(
 		TArray<FGeometryScriptRenderCaptureCamera>& Cameras,
@@ -192,7 +192,7 @@ public:
 	 * Compute oriented sample points on the visible surfaces of the given Actors
 	 * The Samples are computed using Render Capture from the given virtual Cameras
 	 */
-	UFUNCTION(BlueprintCallable, Category = "GeometryScript|MeshSampling")
+	UFUNCTION(BlueprintCallable, Category = "GeometryScript|MeshSampling", meta = (HidePin = "Debug"))
 	static UE_API void
 	ComputeRenderCapturePointSampling(
 		TArray<FTransform>& Samples,
diff --git a/Engine/Plugins/Runtime/GeometryScripting/Source/GeometryScriptingCore/Public/GeometryScript/MeshSculptLayersFunctions.h b/Engine/Plugins/Runtime/GeometryScripting/Source/GeometryScriptingCore/Public/GeometryScript/MeshSculptLayersFunctions.h
--- a/Engine/Plugins/Runtime/GeometryScripting/Source/GeometryScriptingCore/Public/GeometryScript/MeshSculptLayersFunctions.h
+++ b/Engine/Plugins/Runtime/GeometryScripting/Source/GeometryScriptingCore/Public/GeometryScript/MeshSculptLayersFunctions.h
@@ -31,7 +31,7 @@ public:
 
 	/// Enable Sculpt Layers on the Target Mesh, if not already enabled, with at least the requested number of layers. Note: Will never decrease the number of enabled layers.
 	/// @param NumLayers The amount of layers that must be enabled on the mesh. Note: If this many layers (or more) are already enabled, the mesh will not be changed.
-	UFUNCTION(BlueprintCallable, Category = "GeometryScript|SculptLayers", meta = (ScriptMethod))
+	UFUNCTION(BlueprintCallable, Category = "GeometryScript|SculptLayers", meta = (ScriptMethod, HidePin = "Debug"))
 	static UE_API UPARAM(DisplayName = "Target Mesh") UDynamicMesh*
 	EnableSculptLayers(
 		UDynamicMesh* TargetMesh,
@@ -39,7 +39,7 @@ public:
 		UGeometryScriptDebug* Debug = nullptr);
 
 	/// Set the requested LayerIndex as the current active sculpt layer, if possible.
-	UFUNCTION(BlueprintCallable, Category = "GeometryScript|SculptLayers", meta = (ScriptMethod))
+	UFUNCTION(BlueprintCallable, Category = "GeometryScript|SculptLayers", meta = (ScriptMethod, HidePin = "Debug"))
 	static UE_API UPARAM(DisplayName = "Target Mesh") UDynamicMesh*
 	SetActiveSculptLayer(
 		UDynamicMesh* TargetMesh,
@@ -47,7 +47,7 @@ public:
 		UGeometryScriptDebug* Debug = nullptr);
 
 	/// Set the weight of the layer at LayerIndex to the requested Weight
-	UFUNCTION(BlueprintCallable, Category = "GeometryScript|SculptLayers", meta = (ScriptMethod))
+	UFUNCTION(BlueprintCallable, Category = "GeometryScript|SculptLayers", meta = (ScriptMethod, HidePin = "Debug"))
 	static UE_API UPARAM(DisplayName = "Target Mesh") UDynamicMesh*
 	SetSculptLayerWeight(
 		UDynamicMesh* TargetMesh,
@@ -58,7 +58,7 @@ public:
 
 	/// Set the weights of multiple layers to match the given Weights array.
 	/// Note: If the Weights array length is larger than the number of layers, will just set the weights of the existing layers. Will not add or remove layers.
-	UFUNCTION(BlueprintCallable, Category = "GeometryScript|SculptLayers", meta = (ScriptMethod))
+	UFUNCTION(BlueprintCallable, Category = "GeometryScript|SculptLayers", meta = (ScriptMethod, HidePin = "Debug"))
 	static UE_API UPARAM(DisplayName = "Target Mesh") UDynamicMesh*
 	SetSculptLayerWeightsArray(
 		UDynamicMesh* TargetMesh,
@@ -80,7 +80,7 @@ public:
 	static UE_API UPARAM(DisplayName = "Layer Index") int32 GetActiveSculptLayer(const UDynamicMesh* TargetMesh);
 
 	/// Discard all sculpt layer data, leaving current vertex positions unchanged
-	UFUNCTION(BlueprintCallable, Category = "GeometryScript|SculptLayers", meta = (ScriptMethod, Keywords="Delete Remove"))
+	UFUNCTION(BlueprintCallable, Category = "GeometryScript|SculptLayers", meta = (ScriptMethod, HidePin = "Debug", Keywords="Delete Remove"))
 	static UE_API UPARAM(DisplayName = "Target Mesh") UDynamicMesh*
 	DiscardSculptLayers(
 		UDynamicMesh* TargetMesh,
@@ -91,7 +91,7 @@ public:
 	/// @param MergeLayerStart The first layer to merge
 	/// @param MergeLayerNum The number of layers to merge. Note: If larger than the number of existing layers after MergeLayerStart, will be clamped as needed.
 	/// @param bUseWeights If true, layers will be merged based on their current weights -- keeping the mesh shape unchanged. Otherwise, layers will be merged with weight 1, which may change the shape of the mesh.
-	UFUNCTION(BlueprintCallable, Category = "GeometryScript|SculptLayers", meta = (ScriptMethod, Keywords = "Combine"))
+	UFUNCTION(BlueprintCallable, Category = "GeometryScript|SculptLayers", meta = (ScriptMethod, HidePin = "Debug", Keywords = "Combine"))
 	static UE_API UPARAM(DisplayName = "Target Mesh") UDynamicMesh*
 	MergeSculptLayers(
 		UDynamicMesh* TargetMesh,
diff --git a/Engine/Plugins/Runtime/GeometryScripting/Source/GeometryScriptingCore/Public/GeometryScript/MeshSelectionQueryFunctions.h b/Engine/Plugins/Runtime/GeometryScripting/Source/GeometryScriptingCore/Public/GeometryScript/MeshSelectionQueryFunctions.h
--- a/Engine/Plugins/Runtime/GeometryScripting/Source/GeometryScriptingCore/Public/GeometryScript/MeshSelectionQueryFunctions.h
+++ b/Engine/Plugins/Runtime/GeometryScripting/Source/GeometryScriptingCore/Public/GeometryScript/MeshSelectionQueryFunctions.h
@@ -20,7 +20,7 @@ public:
 	 * Get the 3D Bounding Box of a Mesh Selection, ie bounding box of vertices contained in the Selection
 	 * @param bIsEmpty will return as true if the selection was empty (the box will be initialized to 0 in this case)
 	 */
-	UFUNCTION(BlueprintCallable, Category = "GeometryScript|SelectionQueries", meta=(ScriptMethod))
+	UFUNCTION(BlueprintCallable, Category = "GeometryScript|SelectionQueries", meta=(ScriptMethod, HidePin = "Debug"))
 	static UE_API UPARAM(DisplayName = "Target Mesh") UDynamicMesh* 
 	GetMeshSelectionBoundingBox(
 		UDynamicMesh* TargetMesh,
@@ -38,7 +38,7 @@ public:
 	 * @param NumLoops number of loops found is returned here
 	 * @param bFoundErrors true is returned here if topological errors were found during loop computation. In this case the Loop set may be incomplete.
 	 */
-	UFUNCTION(BlueprintCallable, Category = "GeometryScript|SelectionQueries", meta=(ScriptMethod))
+	UFUNCTION(BlueprintCallable, Category = "GeometryScript|SelectionQueries", meta=(ScriptMethod, HidePin = "Debug"))
 	static UE_API UPARAM(DisplayName = "Target Mesh") UDynamicMesh* 
 	GetMeshSelectionBoundaryLoops(
 		UDynamicMesh* TargetMesh,
diff --git a/Engine/Plugins/Runtime/GeometryScripting/Source/GeometryScriptingCore/Public/GeometryScript/MeshSimplifyFunctions.h b/Engine/Plugins/Runtime/GeometryScripting/Source/GeometryScriptingCore/Public/GeometryScript/MeshSimplifyFunctions.h
--- a/Engine/Plugins/Runtime/GeometryScripting/Source/GeometryScriptingCore/Public/GeometryScript/MeshSimplifyFunctions.h
+++ b/Engine/Plugins/Runtime/GeometryScripting/Source/GeometryScriptingCore/Public/GeometryScript/MeshSimplifyFunctions.h
@@ -91,7 +91,7 @@ public:
 	* Simplifies planar areas of the mesh that have more triangles than necessary. Note that it does not change the 3D shape of the mesh.
 	* Planar regions are identified by comparison of face normals using a Angle Threshold in the Options.
 	*/
-	UFUNCTION(BlueprintCallable, Category = "GeometryScript|Simplification", meta=(ScriptMethod))
+	UFUNCTION(BlueprintCallable, Category = "GeometryScript|Simplification", meta=(ScriptMethod, HidePin = "Debug"))
 	static UE_API UPARAM(DisplayName = "Target Mesh") UDynamicMesh* 
 	ApplySimplifyToPlanar(  
 		UDynamicMesh* TargetMesh, 
@@ -102,7 +102,7 @@ public:
 	* Simplifies the mesh down to the PolyGroup Topology. For example, the high-level faces of the mesh PolyGroups. 
 	* Another example would be on a default Box-Sphere where simplifying to PolyGroup topology produces a box.
 	*/
-	UFUNCTION(BlueprintCallable, Category = "GeometryScript|Simplification", meta=(ScriptMethod, DisplayName = "Apply Simplify To PolyGroup Topology"))
+	UFUNCTION(BlueprintCallable, Category = "GeometryScript|Simplification", meta=(ScriptMethod, HidePin = "Debug", DisplayName = "Apply Simplify To PolyGroup Topology"))
 	static UE_API UPARAM(DisplayName = "Target Mesh") UDynamicMesh* 
 	ApplySimplifyToPolygroupTopology(  
 		UDynamicMesh* TargetMesh, 
@@ -113,7 +113,7 @@ public:
 	/**
 	* Simplifies the mesh until a target triangle count is reached. Behavior can be additionally controlled with the Options. 
 	*/
-	UFUNCTION(BlueprintCallable, Category = "GeometryScript|Simplification", meta=(ScriptMethod))
+	UFUNCTION(BlueprintCallable, Category = "GeometryScript|Simplification", meta=(ScriptMethod, HidePin = "Debug"))
 	static UE_API UPARAM(DisplayName = "Target Mesh") UDynamicMesh* 
 	ApplySimplifyToTriangleCount(  
 		UDynamicMesh* TargetMesh, 
@@ -124,7 +124,7 @@ public:
 	/**
 	* Simplifies the mesh until a target vertex count is reached. Behavior can be additionally controlled with the Options. 
 	*/
-	UFUNCTION(BlueprintCallable, Category = "GeometryScript|Simplification", meta=(ScriptMethod))
+	UFUNCTION(BlueprintCallable, Category = "GeometryScript|Simplification", meta=(ScriptMethod, HidePin = "Debug"))
 	static UE_API UPARAM(DisplayName = "Target Mesh") UDynamicMesh* 
 	ApplySimplifyToVertexCount(  
 		UDynamicMesh* TargetMesh, 
@@ -135,7 +135,7 @@ public:
 	/**
 	* Simplifies the mesh until a target triangle count is reached, using the UE Editor's standard mesh simplifier. Editor only.
 	*/
-	UFUNCTION(BlueprintCallable, Category = "GeometryScript|Simplification", meta=(ScriptMethod))
+	UFUNCTION(BlueprintCallable, Category = "GeometryScript|Simplification", meta = (ScriptMethod, HidePin = "Debug"))
 	static UE_API UPARAM(DisplayName = "Target Mesh") UDynamicMesh* 
 	ApplyEditorSimplifyToTriangleCount(  
 		UDynamicMesh* TargetMesh, 
@@ -145,7 +145,7 @@ public:
 	/**
 	* Simplifies the mesh until a target vertex count is reached, using the UE Editor's standard mesh simplifier. Editor only.
 	*/
-	UFUNCTION(BlueprintCallable, Category = "GeometryScript|Simplification", meta=(ScriptMethod))
+	UFUNCTION(BlueprintCallable, Category = "GeometryScript|Simplification", meta=(ScriptMethod, HidePin = "Debug"))
 	static UE_API UPARAM(DisplayName = "Target Mesh") UDynamicMesh* 
 	ApplyEditorSimplifyToVertexCount(  
 		UDynamicMesh* TargetMesh, 
@@ -156,7 +156,7 @@ public:
 	* Simplifies the mesh to a target geometric tolerance. Stops when any further simplification would result in a deviation from the input mesh larger than the tolerance.
 	* Behavior can be additionally controlled with the Options. 
 	*/
-	UFUNCTION(BlueprintCallable, Category = "GeometryScript|Simplification", meta=(ScriptMethod))
+	UFUNCTION(BlueprintCallable, Category = "GeometryScript|Simplification", meta=(ScriptMethod, HidePin = "Debug"))
 	static UE_API UPARAM(DisplayName = "Target Mesh") UDynamicMesh* 
 	ApplySimplifyToTolerance(  
 		UDynamicMesh* TargetMesh, 
diff --git a/Engine/Plugins/Runtime/GeometryScripting/Source/GeometryScriptingCore/Public/GeometryScript/MeshSpatialFunctions.h b/Engine/Plugins/Runtime/GeometryScripting/Source/GeometryScriptingCore/Public/GeometryScript/MeshSpatialFunctions.h
--- a/Engine/Plugins/Runtime/GeometryScripting/Source/GeometryScriptingCore/Public/GeometryScript/MeshSpatialFunctions.h
+++ b/Engine/Plugins/Runtime/GeometryScripting/Source/GeometryScriptingCore/Public/GeometryScript/MeshSpatialFunctions.h
@@ -71,7 +71,7 @@ public:
 	/**
 	* Builds a Bounding Volume Hierarchy (BVH) object for a mesh that can be used with multiple spatial queries.
 	*/
-	UFUNCTION(BlueprintCallable, Category = "GeometryScript|Spatial", meta=(ScriptMethod, DisplayName="Build BVH For Mesh"))
+	UFUNCTION(BlueprintCallable, Category = "GeometryScript|Spatial", meta=(ScriptMethod, HidePin = "Debug", DisplayName="Build BVH For Mesh"))
 	static UE_API UPARAM(DisplayName = "Target Mesh") UDynamicMesh* 
 	BuildBVHForMesh(
 		UDynamicMesh* TargetMesh,
@@ -81,7 +81,7 @@ public:
 	/**
 	* Checks if the provided Bounding Volume Hierarchy (BVH) can still be used with the Mesh — it generally returns false if the mesh has been changed.
 	*/
-	UFUNCTION(BlueprintCallable, Category = "GeometryScript|Spatial", meta=(ScriptMethod, DisplayName="Is BVH Valid For Mesh"))
+	UFUNCTION(BlueprintCallable, Category = "GeometryScript|Spatial", meta=(ScriptMethod, HidePin = "Debug", DisplayName="Is BVH Valid For Mesh"))
 	static UE_API UPARAM(DisplayName = "Target Mesh") UDynamicMesh* 
 	IsBVHValidForMesh(
 		UDynamicMesh* TargetMesh,
@@ -92,7 +92,7 @@ public:
 	/**
 	* Rebuilds the Bounding Volume Hierarchy (BVH) for the mesh in-place, which can reduce memory allocations, compared to building a new BVH.
 	*/
-	UFUNCTION(BlueprintCallable, Category = "GeometryScript|Spatial", meta=(ScriptMethod, DisplayName="Rebuild BVH For Mesh"))
+	UFUNCTION(BlueprintCallable, Category = "GeometryScript|Spatial", meta=(ScriptMethod, HidePin = "Debug", DisplayName="Rebuild BVH For Mesh"))
 	static UE_API UPARAM(DisplayName = "Target Mesh") UDynamicMesh* 
 	RebuildBVHForMesh(
 		UDynamicMesh* TargetMesh,
@@ -108,7 +108,7 @@ public:
 	* @param Outcome will be either Found or Not Found depending on the success of the query.
 	* Note NearestResult.bValid will be false if the query failed. 
 	*/
-	UFUNCTION(BlueprintCallable, Category = "GeometryScript|Spatial", meta=(ScriptMethod, ExpandEnumAsExecs = "Outcome"))
+	UFUNCTION(BlueprintCallable, Category = "GeometryScript|Spatial", meta=(ScriptMethod, HidePin = "Debug", ExpandEnumAsExecs = "Outcome"))
 	static UE_API UPARAM(DisplayName = "Target Mesh") UDynamicMesh* 
 	FindNearestPointOnMesh(
 		UDynamicMesh* TargetMesh,
@@ -124,7 +124,7 @@ public:
 	* Note, depending on the Ray Origin and Ray Direction, there is the possibility that the ray might not intersect with the Target Mesh.  
 	* Should the ray miss, the HitResult.bHit will be false and the Outcome  will be Not Found.
 	*/
-	UFUNCTION(BlueprintCallable, Category = "GeometryScript|Spatial", meta=(ScriptMethod, ExpandEnumAsExecs = "Outcome"))
+	UFUNCTION(BlueprintCallable, Category = "GeometryScript|Spatial", meta=(ScriptMethod, HidePin = "Debug", ExpandEnumAsExecs = "Outcome"))
 	static UE_API UPARAM(DisplayName = "Target Mesh") UDynamicMesh* 
 	FindNearestRayIntersectionWithMesh(
 		UDynamicMesh* TargetMesh,
@@ -142,7 +142,7 @@ public:
 	* @param QueryPoint the point in the mesh's 3D local space.
 	* @param Options control the fast winding number threshold 
 	*/
-	UFUNCTION(BlueprintCallable, Category = "GeometryScript|Spatial", meta=(ScriptMethod, ExpandEnumAsExecs = "Outcome"))
+	UFUNCTION(BlueprintCallable, Category = "GeometryScript|Spatial", meta=(ScriptMethod, HidePin = "Debug", ExpandEnumAsExecs = "Outcome"))
 	static UE_API UPARAM(DisplayName = "Target Mesh") UDynamicMesh* 
 	IsPointInsideMesh(
 		UDynamicMesh* TargetMesh,
@@ -165,7 +165,7 @@ public:
 	* @param QueryPoint the point in the mesh's 3D local space.
 	* @param Options control the fast winding number threshold 
 	*/
-	UFUNCTION(BlueprintCallable, Category = "GeometryScript|Spatial", meta=(ScriptMethod, AdvancedDisplay = "MinNumTrianglePoints"))
+	UFUNCTION(BlueprintCallable, Category = "GeometryScript|Spatial", meta=(ScriptMethod, HidePin = "Debug", AdvancedDisplay = "MinNumTrianglePoints"))
 	static UE_API UPARAM(DisplayName = "Target Mesh") UDynamicMesh* 
 	SelectMeshElementsInBoxWithBVH(
 		UDynamicMesh* TargetMesh,
diff --git a/Engine/Plugins/Runtime/GeometryScripting/Source/GeometryScriptingCore/Public/GeometryScript/MeshSubdivideFunctions.h b/Engine/Plugins/Runtime/GeometryScripting/Source/GeometryScriptingCore/Public/GeometryScript/MeshSubdivideFunctions.h
--- a/Engine/Plugins/Runtime/GeometryScripting/Source/GeometryScriptingCore/Public/GeometryScript/MeshSubdivideFunctions.h
+++ b/Engine/Plugins/Runtime/GeometryScripting/Source/GeometryScriptingCore/Public/GeometryScript/MeshSubdivideFunctions.h
@@ -60,7 +60,7 @@ public:
 	/**
 	 * Apply PN Tessellation to the Target Mesh as controlled by the Tessellation Level and the Options.
 	 */
-	UFUNCTION(BlueprintCallable, Category = "GeometryScript|Subdivide", meta=(ScriptMethod, DisplayName = "Apply PN Tessellation"))
+	UFUNCTION(BlueprintCallable, Category = "GeometryScript|Subdivide", meta=(ScriptMethod, HidePin = "Debug", DisplayName = "Apply PN Tessellation"))
 	static UE_API UPARAM(DisplayName = "Target Mesh") UDynamicMesh* 
 	ApplyPNTessellation(
 		UDynamicMesh* TargetMesh,
@@ -71,7 +71,7 @@ public:
 	/**
 	* Apply Uniform Tessellation to the Target Mesh as controlled by the Tessellation Level and the Options.
 	*/
-	UFUNCTION(BlueprintCallable, Category = "GeometryScript|Subdivide", meta=(ScriptMethod))
+	UFUNCTION(BlueprintCallable, Category = "GeometryScript|Subdivide", meta=(ScriptMethod, HidePin = "Debug"))
 	static UE_API UPARAM(DisplayName = "Target Mesh") UDynamicMesh* 
 	ApplyUniformTessellation(
 		UDynamicMesh* TargetMesh,
@@ -85,7 +85,7 @@ public:
 	* @param Options controls the behavior of the tessellation if the Selection is empty.
 	* @param TessellationLevel determines the amount of tessellation
 	*/
-	UFUNCTION(BlueprintCallable, Category = "GeometryScript|Subdivide", meta=(ScriptMethod))
+	UFUNCTION(BlueprintCallable, Category = "GeometryScript|Subdivide", meta=(ScriptMethod, HidePin = "Debug"))
 	static UE_API UPARAM(DisplayName = "Target Mesh") UDynamicMesh* 
 	ApplySelectiveTessellation(
 		UDynamicMesh* TargetMesh,
diff --git a/Engine/Plugins/Runtime/GeometryScripting/Source/GeometryScriptingCore/Public/GeometryScript/MeshTransformFunctions.h b/Engine/Plugins/Runtime/GeometryScripting/Source/GeometryScriptingCore/Public/GeometryScript/MeshTransformFunctions.h
--- a/Engine/Plugins/Runtime/GeometryScripting/Source/GeometryScriptingCore/Public/GeometryScript/MeshTransformFunctions.h
+++ b/Engine/Plugins/Runtime/GeometryScripting/Source/GeometryScriptingCore/Public/GeometryScript/MeshTransformFunctions.h
@@ -22,7 +22,7 @@ public:
 	/**
 	* Applies the provided Transform to the vertices of a Mesh.
 	*/
-	UFUNCTION(BlueprintCallable, Category = "GeometryScript|Transforms", meta=(ScriptMethod, AdvancedDisplay = "bFixOrientationForNegativeScale"))
+	UFUNCTION(BlueprintCallable, Category = "GeometryScript|Transforms", meta=(ScriptMethod, HidePin = "Debug", AdvancedDisplay = "bFixOrientationForNegativeScale"))
 	static UE_API UPARAM(DisplayName = "Target Mesh") UDynamicMesh*
 	TransformMesh(
 		UDynamicMesh* TargetMesh,
@@ -33,7 +33,7 @@ public:
 	/**
 	 * Applies the inverse of the provided Transform to the TargetMesh
 	 */
-	UFUNCTION(BlueprintCallable, Category = "GeometryScript|Transforms", meta = (ScriptMethod, AdvancedDisplay = "bFixOrientationForNegativeScale"))
+	UFUNCTION(BlueprintCallable, Category = "GeometryScript|Transforms", meta = (ScriptMethod, HidePin = "Debug", AdvancedDisplay = "bFixOrientationForNegativeScale"))
 	static UE_API UPARAM(DisplayName = "Target Mesh") UDynamicMesh*
 	InverseTransformMesh(
 		UDynamicMesh* TargetMesh,
@@ -44,7 +44,7 @@ public:
 	/**
 	* Applies the provided Translation to the vertices of a Mesh.
 	*/
-	UFUNCTION(BlueprintCallable, Category = "GeometryScript|Transforms", meta=(ScriptMethod))
+	UFUNCTION(BlueprintCallable, Category = "GeometryScript|Transforms", meta=(ScriptMethod, HidePin = "Debug"))
 	static UE_API UPARAM(DisplayName = "Target Mesh") UDynamicMesh* 
 	TranslateMesh(
 		UDynamicMesh* TargetMesh,
@@ -54,7 +54,7 @@ public:
 	/**
 	* Rotates the mesh relative to the specified Rotation Origin.
 	*/
-	UFUNCTION(BlueprintCallable, Category = "GeometryScript|Transforms", meta=(ScriptMethod))
+	UFUNCTION(BlueprintCallable, Category = "GeometryScript|Transforms", meta=(ScriptMethod, HidePin = "Debug"))
 	static UE_API UPARAM(DisplayName = "Target Mesh") UDynamicMesh* 
 	RotateMesh(
 		UDynamicMesh* TargetMesh,
@@ -65,7 +65,7 @@ public:
 	/**
 	* Applies the provided Scale transformation relative to the Scale Origin to the vertices of a Mesh.
 	*/
-	UFUNCTION(BlueprintCallable, Category = "GeometryScript|Transforms", meta=(ScriptMethod, AdvancedDisplay = "bFixOrientationForNegativeScale"))
+	UFUNCTION(BlueprintCallable, Category = "GeometryScript|Transforms", meta=(ScriptMethod, HidePin = "Debug", AdvancedDisplay = "bFixOrientationForNegativeScale"))
 	static UE_API UPARAM(DisplayName = "Target Mesh") UDynamicMesh* 
 	ScaleMesh(
 		UDynamicMesh* TargetMesh,
@@ -79,7 +79,7 @@ public:
 	/**
 	* Applies the given Transform to the vertices identified by the Selection of the mesh.
 	*/
-	UFUNCTION(BlueprintCallable, Category = "GeometryScript|Transforms", meta=(ScriptMethod, AdvancedDisplay = "bFixOrientationForNegativeScale"))
+	UFUNCTION(BlueprintCallable, Category = "GeometryScript|Transforms", meta=(ScriptMethod, HidePin = "Debug", AdvancedDisplay = "bFixOrientationForNegativeScale"))
 	static UE_API UPARAM(DisplayName = "Target Mesh") UDynamicMesh*
 	TransformMeshSelection(
 		UDynamicMesh* TargetMesh,
@@ -90,7 +90,7 @@ public:
 	/**
 	 * Applies the inverse of the given Transform to the vertices identified by the Selection of the TargetMesh.
 	 */
-	UFUNCTION(BlueprintCallable, Category = "GeometryScript|Transforms", meta=(ScriptMethod, AdvancedDisplay = "bFixOrientationForNegativeScale"))
+	UFUNCTION(BlueprintCallable, Category = "GeometryScript|Transforms", meta=(ScriptMethod, HidePin = "Debug", AdvancedDisplay = "bFixOrientationForNegativeScale"))
 	static UE_API UPARAM(DisplayName = "Target Mesh") UDynamicMesh*
 	InverseTransformMeshSelection(
 		UDynamicMesh* TargetMesh,
@@ -101,7 +101,7 @@ public:
 	/**
 	* Applies the given Translation to the vertices identified by the Selection of the mesh.
 	*/ 
-	UFUNCTION(BlueprintCallable, Category = "GeometryScript|Transforms", meta=(ScriptMethod))
+	UFUNCTION(BlueprintCallable, Category = "GeometryScript|Transforms", meta=(ScriptMethod, HidePin = "Debug"))
 	static UE_API UPARAM(DisplayName = "Target Mesh") UDynamicMesh* 
 	TranslateMeshSelection(
 		UDynamicMesh* TargetMesh,
@@ -112,7 +112,7 @@ public:
 	/**
 	* Rotates the selected part of the mesh relative to the specified Rotation Origin.
 	*/
-	UFUNCTION(BlueprintCallable, Category = "GeometryScript|Transforms", meta=(ScriptMethod))
+	UFUNCTION(BlueprintCallable, Category = "GeometryScript|Transforms", meta=(ScriptMethod, HidePin = "Debug"))
 	static UE_API UPARAM(DisplayName = "Target Mesh") UDynamicMesh* 
 	RotateMeshSelection(
 		UDynamicMesh* TargetMesh,
@@ -124,7 +124,7 @@ public:
 	/**
 	* Applies the given Scale transform relative to the Scale Origin to the selection part of the mesh.
 	*/
-	UFUNCTION(BlueprintCallable, Category = "GeometryScript|Transforms", meta=(ScriptMethod, AdvancedDisplay = "bFixOrientationForNegativeScale"))
+	UFUNCTION(BlueprintCallable, Category = "GeometryScript|Transforms", meta=(ScriptMethod, HidePin = "Debug", AdvancedDisplay = "bFixOrientationForNegativeScale"))
 	static UE_API UPARAM(DisplayName = "Target Mesh") UDynamicMesh* 
 	ScaleMeshSelection(
 		UDynamicMesh* TargetMesh,
@@ -138,7 +138,7 @@ public:
 	 * Set the Pivot Location for the Mesh. Since the Pivot of a Mesh object is always the point at (0,0,0),
 	 * this function simply translates the mesh by -PivotLocation.
 	 */
-	UFUNCTION(BlueprintCallable, Category = "GeometryScript|Transforms", meta=(ScriptMethod))
+	UFUNCTION(BlueprintCallable, Category = "GeometryScript|Transforms", meta=(ScriptMethod, HidePin = "Debug"))
 	static UPARAM(DisplayName = "Target Mesh") UDynamicMesh* 
 	TranslatePivotToLocation(
 		UDynamicMesh* TargetMesh,
diff --git a/Engine/Plugins/Runtime/GeometryScripting/Source/GeometryScriptingCore/Public/GeometryScript/MeshUVFunctions.h b/Engine/Plugins/Runtime/GeometryScripting/Source/GeometryScriptingCore/Public/GeometryScript/MeshUVFunctions.h
--- a/Engine/Plugins/Runtime/GeometryScripting/Source/GeometryScriptingCore/Public/GeometryScript/MeshUVFunctions.h
+++ b/Engine/Plugins/Runtime/GeometryScripting/Source/GeometryScriptingCore/Public/GeometryScript/MeshUVFunctions.h
@@ -263,14 +263,14 @@ public:
 	/**
 	* Set the number of UV Channels on the Target Mesh. If not already enabled, this will enable the mesh attributes.
 	*/
-	UFUNCTION(BlueprintCallable, Category = "GeometryScript|UVs", meta=(ScriptMethod), DisplayName = "Set Num UV Channels")
+	UFUNCTION(BlueprintCallable, Category = "GeometryScript|UVs", meta=(ScriptMethod, HidePin = "Debug"), DisplayName = "Set Num UV Channels")
 	static UE_API UPARAM(DisplayName = "Target Mesh") UDynamicMesh* 
 	SetNumUVSets( 
 		UDynamicMesh* TargetMesh, 
 		UPARAM(DisplayName = "Num UV Channels") int NumUVSets,
 		UGeometryScriptDebug* Debug = nullptr );
 
-	UFUNCTION(BlueprintCallable, Category = "GeometryScript|UVs", meta = (ScriptMethod), DisplayName = "Clear UV Channel")
+	UFUNCTION(BlueprintCallable, Category = "GeometryScript|UVs", meta = (ScriptMethod, HidePin = "Debug"), DisplayName = "Clear UV Channel")
 	static UE_API UPARAM(DisplayName = "Target Mesh") UDynamicMesh*
 	ClearUVChannel(
 		UDynamicMesh* TargetMesh,
@@ -280,7 +280,7 @@ public:
 	/**
 	* Copy the data in one UV Channel to another UV Channel on the same Target Mesh.
 	*/
-	UFUNCTION(BlueprintCallable, Category = "GeometryScript|UVs", meta=(ScriptMethod), DisplayName = "Copy UV Channel")
+	UFUNCTION(BlueprintCallable, Category = "GeometryScript|UVs", meta=(ScriptMethod, HidePin = "Debug"), DisplayName = "Copy UV Channel")
 	static UE_API UPARAM(DisplayName = "Target Mesh") UDynamicMesh* 
 	CopyUVSet( 
 		UDynamicMesh* TargetMesh, 
@@ -358,7 +358,7 @@ public:
 	 * @param bInsertSeams Whether to insert new seams. If false, removes existing seams instead.
 	 * @param bDeferChangeNotifications If true, no mesh change notification will be sent. Set to true if performing many changes in a loop.
 	 */
-	UFUNCTION(BlueprintCallable, Category = "GeometryScript|UVs", meta = (ScriptMethod, DisplayName = "Set UV Seams Along Selected Edges"))
+	UFUNCTION(BlueprintCallable, Category = "GeometryScript|UVs", meta = (ScriptMethod, HidePin = "Debug", DisplayName = "Set UV Seams Along Selected Edges"))
 	static UE_API UPARAM(DisplayName = "Target Mesh") UDynamicMesh*
 	SetUVSeamsAlongSelectedEdges(
 		UDynamicMesh* TargetMesh,
@@ -399,7 +399,7 @@ public:
 	* Update all selected UV values in the specified UV Channel by adding the Translation value to each.
 	* If the provided Selection is empty, the Translation is applied to the entire UV Channel.
 	*/
-	UFUNCTION(BlueprintCallable, Category = "GeometryScript|UVs", meta=(ScriptMethod))
+	UFUNCTION(BlueprintCallable, Category = "GeometryScript|UVs", meta=(ScriptMethod, HidePin = "Debug"))
 	static UE_API UPARAM(DisplayName = "Target Mesh") UDynamicMesh* 
 	TranslateMeshUVs( 
 		UDynamicMesh* TargetMesh, 
@@ -412,7 +412,7 @@ public:
 	* Update all selected UV values in the specified UV Channel by Scale, mathematically the new value is given by (UV - ScaleOrigin) * Scale + ScaleOrigin
 	* If the provided Selection is empty, the update is applied to the entire UV Channel.
 	*/
-	UFUNCTION(BlueprintCallable, Category = "GeometryScript|UVs", meta=(ScriptMethod))
+	UFUNCTION(BlueprintCallable, Category = "GeometryScript|UVs", meta=(ScriptMethod, HidePin = "Debug"))
 	static UE_API UPARAM(DisplayName = "Target Mesh") UDynamicMesh* 
 	ScaleMeshUVs( 
 		UDynamicMesh* TargetMesh, 
@@ -426,7 +426,7 @@ public:
 	* Update all the selected UV values in the specified UV Channel by a rotation of Rotation Angle (in degrees) relative to the Rotation Origin.
 	* If the provided Selection is empty, the update is applied to the entire UV Channel.
 	*/
-	UFUNCTION(BlueprintCallable, Category = "GeometryScript|UVs", meta=(ScriptMethod))
+	UFUNCTION(BlueprintCallable, Category = "GeometryScript|UVs", meta=(ScriptMethod, HidePin = "Debug"))
 	static UE_API UPARAM(DisplayName = "Target Mesh") UDynamicMesh* 
 	RotateMeshUVs( 
 		UDynamicMesh* TargetMesh, 
@@ -440,7 +440,7 @@ public:
 	/**
 	 * Scale of PlaneTransform defines world-space dimension that maps to 1 UV dimension
 	 */
-	UFUNCTION(BlueprintCallable, Category = "GeometryScript|UVs", meta=(ScriptMethod))
+	UFUNCTION(BlueprintCallable, Category = "GeometryScript|UVs", meta=(ScriptMethod, HidePin = "Debug"))
 	static UE_API UPARAM(DisplayName = "Target Mesh") UDynamicMesh* 
 	SetMeshUVsFromPlanarProjection( 
 		UDynamicMesh* TargetMesh, 
@@ -453,7 +453,7 @@ public:
 	/**
 	* Using Box Projection, update the UVs in the UV Channel for an entire mesh or a subset defined by a non-empty Selection.
 	*/
-	UFUNCTION(BlueprintCallable, Category = "GeometryScript|UVs", meta=(ScriptMethod))
+	UFUNCTION(BlueprintCallable, Category = "GeometryScript|UVs", meta=(ScriptMethod, HidePin = "Debug"))
 	static UE_API UPARAM(DisplayName = "Target Mesh") UDynamicMesh* 
 	SetMeshUVsFromBoxProjection( 
 		UDynamicMesh* TargetMesh, 
@@ -466,7 +466,7 @@ public:
 	/**
 	* Using Cylinder Projection, update the UVs in the UV Channel for an entire mesh or a subset defined by a non-empty Selection.
 	*/
-	UFUNCTION(BlueprintCallable, Category = "GeometryScript|UVs", meta=(ScriptMethod))
+	UFUNCTION(BlueprintCallable, Category = "GeometryScript|UVs", meta=(ScriptMethod, HidePin = "Debug"))
 	static UE_API UPARAM(DisplayName = "Target Mesh") UDynamicMesh* 
 	SetMeshUVsFromCylinderProjection( 
 		UDynamicMesh* TargetMesh, 
@@ -495,7 +495,7 @@ public:
 	 * @param ProjectionDirection Direction to project (in the space where TargetMesh is transformed by TargetTransform, and SourceMesh is transformed by SourceTransform)
 	 * @param ProjectionOffset Projection will start offset by this amount from the TargetMesh vertices
 	 */
-	UFUNCTION(BlueprintCallable, Category = "GeometryScript|UVs", meta = (ScriptMethod, Keywords = "Set From Mesh"))
+	UFUNCTION(BlueprintCallable, Category = "GeometryScript|UVs", meta = (ScriptMethod, HidePin = "Debug", Keywords = "Set From Mesh"))
 	static UE_API UPARAM(DisplayName = "Target Mesh") UDynamicMesh*
 	TransferMeshUVsByProjection(
 		UDynamicMesh* TargetMesh,
@@ -517,7 +517,7 @@ public:
 	* Rescales UVs in the UV Channel for a Mesh to match a specified texel density, described by the options passed in. Supports
 	* processing on a subset of UVs via a non-empty Selection.
 	*/
-	UFUNCTION(BlueprintCallable, Category = "GeometryScript|UVs", meta=(ScriptMethod, AdvancedDisplay = "UDIMResolutions"))
+	UFUNCTION(BlueprintCallable, Category = "GeometryScript|UVs", meta=(ScriptMethod, HidePin = "Debug", AdvancedDisplay = "UDIMResolutions"))
 	static UE_API UPARAM(DisplayName = "Target Mesh") UDynamicMesh* 
 	ApplyTexelDensityUVScaling( 
 		UDynamicMesh* TargetMesh, 
@@ -530,7 +530,7 @@ public:
 	* Recomputes UVs in the UV Channel for a Mesh based on different types of well-defined UV islands, such as existing UV islands, PolyGroups, 
 	* or a subset of the mesh based on a non-empty Selection.
 	*/
-	UFUNCTION(BlueprintCallable, Category = "GeometryScript|UVs", meta=(ScriptMethod))
+	UFUNCTION(BlueprintCallable, Category = "GeometryScript|UVs", meta=(ScriptMethod, HidePin = "Debug"))
 	static UE_API UPARAM(DisplayName = "Target Mesh") UDynamicMesh* 
 	RecomputeMeshUVs( 
 		UDynamicMesh* TargetMesh, 
@@ -542,7 +542,7 @@ public:
 	/**
 	* Packs the existing UV islands in the specified UV Channel into standard UV space based on the Repack Options.
 	*/
-	UFUNCTION(BlueprintCallable, Category = "GeometryScript|UVs", meta=(ScriptMethod))
+	UFUNCTION(BlueprintCallable, Category = "GeometryScript|UVs", meta=(ScriptMethod, HidePin = "Debug"))
 	static UE_API UPARAM(DisplayName = "Target Mesh") UDynamicMesh* 
 	RepackMeshUVs( 
 		UDynamicMesh* TargetMesh, 
@@ -553,7 +553,7 @@ public:
 	/**
 	* Packs the existing UV islands in the specified UV Channel into standard UV space based on the Repack Options.
 	*/
-	UFUNCTION(BlueprintCallable, Category = "GeometryScript|UVs", meta = (ScriptMethod))
+	UFUNCTION(BlueprintCallable, Category = "GeometryScript|UVs", meta = (ScriptMethod, HidePin = "Debug"))
 	static UE_API UPARAM(DisplayName = "Target Mesh") UDynamicMesh*
 	LayoutMeshUVs(
 		UDynamicMesh* TargetMesh,
@@ -565,7 +565,7 @@ public:
 	/**
 	* Computes new UVs for the specified UV Channel using PatchBuilder method in the Options, and optionally packs.
 	*/
-	UFUNCTION(BlueprintCallable, Category = "GeometryScript|UVs", meta=(ScriptMethod))
+	UFUNCTION(BlueprintCallable, Category = "GeometryScript|UVs", meta=(ScriptMethod, HidePin = "Debug"))
 	static UE_API UPARAM(DisplayName = "Target Mesh") UDynamicMesh* 
 	AutoGeneratePatchBuilderMeshUVs( 
 		UDynamicMesh* TargetMesh, 
@@ -576,7 +576,7 @@ public:
 	/**
 	* Computes new UVs for the specified UV Channel using XAtlas, and optionally packs.
 	*/
-	UFUNCTION(BlueprintCallable, Category = "GeometryScript|UVs", meta=(ScriptMethod))
+	UFUNCTION(BlueprintCallable, Category = "GeometryScript|UVs", meta=(ScriptMethod, HidePin = "Debug"))
 	static UE_API UPARAM(DisplayName = "Target Mesh") UDynamicMesh* 
 	AutoGenerateXAtlasMeshUVs( 
 		UDynamicMesh* TargetMesh, 
@@ -596,7 +596,7 @@ public:
 	 * @param bFoundUnsetUVs output flag set to true if any of the queried triangles do not have valid UVs set
 	 * @param bOnlyIncludeValidUVTris if true, only triangles with valid UVs are included in 3D Mesh Area/Bounds
 	 */
-	UFUNCTION(BlueprintCallable, Category = "GeometryScript|UVs", meta=(ScriptMethod))
+	UFUNCTION(BlueprintCallable, Category = "GeometryScript|UVs", meta=(ScriptMethod, HidePin = "Debug"))
 	static UE_API UPARAM(DisplayName = "Copy From Mesh") UDynamicMesh* 
 	GetMeshUVSizeInfo(  
 		UDynamicMesh* TargetMesh, 
@@ -622,7 +622,7 @@ public:
 	 * @param bHasVertexIDGaps will be set to true if some vertex indices in TargetMesh were invalid, ie MaxVertexID > VertexCount 
 	 * @param bHasSplitUVs will be set to true if there were split UVs in the UV Channel
 	 */
-	UFUNCTION(BlueprintCallable, Category = "GeometryScript|UVs", meta=(ScriptMethod))
+	UFUNCTION(BlueprintCallable, Category = "GeometryScript|UVs", meta=(ScriptMethod, HidePin = "Debug"))
 	static UE_API UPARAM(DisplayName = "Target Mesh") UDynamicMesh* 
 	GetMeshPerVertexUVs( 
 		UDynamicMesh* TargetMesh, 
@@ -647,7 +647,7 @@ public:
 	 * @param bInvalidTopology will be returned true if any topological issues were found
 	 * @param bIsValidUVSet will be returned false if UVSetIndex is not available
 	 */
-	UFUNCTION(BlueprintCallable, Category = "GeometryScript|UVs", meta=(ScriptMethod), DisplayName="Copy Mesh UV Channel To Mesh")
+	UFUNCTION(BlueprintCallable, Category = "GeometryScript|UVs", meta=(ScriptMethod, HidePin = "Debug"), DisplayName="Copy Mesh UV Channel To Mesh")
 	static UE_API UPARAM(DisplayName = "Copy From Mesh") UDynamicMesh* 
 	CopyMeshUVLayerToMesh(  
 		UDynamicMesh* CopyFromMesh, 
@@ -670,7 +670,7 @@ public:
 	 * @param bIsValidUVSet will be returned false if To UV Channel is not available
 	 * @param bOnlyUVPositions if true, only (valid, matching) UV positions are updated, a full new UV topology is created
 	 */
-	UFUNCTION(BlueprintCallable, Category = "GeometryScript|UVs", meta=(ScriptMethod), DisplayName="Copy Mesh To Mesh UV Channel")
+	UFUNCTION(BlueprintCallable, Category = "GeometryScript|UVs", meta=(ScriptMethod, HidePin = "Debug"), DisplayName="Copy Mesh To Mesh UV Channel")
 	static UE_API UPARAM(DisplayName = "Copy From Mesh") UDynamicMesh* 
 	CopyMeshToMeshUVLayer(  
 		UDynamicMesh* CopyFromUVMesh, 
@@ -695,7 +695,7 @@ public:
 	 * @param VertexIDs output list of VertexIDs that UVs have been computed for, ie are within geodesic distance Radius from the CenterPoint
 	 * @param VertexUVs output list of Vertex UVs that corresponds to VertexIDs
 	 */
-	UFUNCTION(BlueprintCallable, Category = "GeometryScript|UVs", meta=(ScriptMethod, AdvancedDisplay = "bUseInterpolatedNormal, TangentYDirection"))
+	UFUNCTION(BlueprintCallable, Category = "GeometryScript|UVs", meta=(ScriptMethod, HidePin = "Debug", AdvancedDisplay = "bUseInterpolatedNormal, TangentYDirection"))
 	static UE_API UPARAM(DisplayName = "Target Mesh") UDynamicMesh* 
 	ComputeMeshLocalUVParam( 
 		UDynamicMesh* TargetMesh, 
diff --git a/Engine/Plugins/Runtime/GeometryScripting/Source/GeometryScriptingCore/Public/GeometryScript/MeshVertexColorFunctions.h b/Engine/Plugins/Runtime/GeometryScripting/Source/GeometryScriptingCore/Public/GeometryScript/MeshVertexColorFunctions.h
--- a/Engine/Plugins/Runtime/GeometryScripting/Source/GeometryScriptingCore/Public/GeometryScript/MeshVertexColorFunctions.h
+++ b/Engine/Plugins/Runtime/GeometryScripting/Source/GeometryScriptingCore/Public/GeometryScript/MeshVertexColorFunctions.h
@@ -118,7 +118,7 @@ public:
 	 * @param Color the constant color to set
 	 * @param Flags specify which RGBA channels to set (default all channels)
 	 */
-	UFUNCTION(BlueprintCallable, Category = "GeometryScript|VertexColor", meta=(ScriptMethod))
+	UFUNCTION(BlueprintCallable, Category = "GeometryScript|VertexColor", meta=(ScriptMethod, HidePin = "Debug"))
 	static UE_API UPARAM(DisplayName = "Target Mesh") UDynamicMesh*
 	SetMeshConstantVertexColor(
 		UDynamicMesh* TargetMesh,
@@ -136,7 +136,7 @@ public:
 	 * @param Flags specify which RGBA channels to set (default all channels)
 	 * @param bCreateColorSeam if true, a "hard edge" in the vertex colors is created, by creating new Elements for all the triangles in the selection. If enabled, Vertex selections are converted to Triangle selections, and Flags is ignored. 
 	 */
-	UFUNCTION(BlueprintCallable, Category = "GeometryScript|VertexColor", meta=(ScriptMethod))
+	UFUNCTION(BlueprintCallable, Category = "GeometryScript|VertexColor", meta=(ScriptMethod, HidePin = "Debug"))
 	static UE_API UPARAM(DisplayName = "Target Mesh") UDynamicMesh*
 	SetMeshSelectionVertexColor(
 		UDynamicMesh* TargetMesh,
@@ -151,7 +151,7 @@ public:
 	 * Set all vertex colors in the TargetMesh VertexColor Overlay to the specified per-vertex colors
 	 * @param VertexColorList per-vertex colors. Size must be less than or equal to the MaxVertexID of TargetMesh  (ie gaps are supported)
 	 */
-	UFUNCTION(BlueprintCallable, Category = "GeometryScript|VertexColor", meta=(ScriptMethod))
+	UFUNCTION(BlueprintCallable, Category = "GeometryScript|VertexColor", meta=(ScriptMethod, HidePin = "Debug"))
 	static UE_API UPARAM(DisplayName = "Target Mesh") UDynamicMesh*
 	SetMeshPerVertexColors(
 		UDynamicMesh* TargetMesh,
@@ -180,7 +180,7 @@ public:
 	 * Apply a SRGB to Linear color transformation on all vertex colors
 	 * on the mesh.
 	 */
-	UFUNCTION(BlueprintCallable, Category = "GeometryScript|VertexColor", meta=(ScriptMethod))
+	UFUNCTION(BlueprintCallable, Category = "GeometryScript|VertexColor", meta=(ScriptMethod, HidePin = "Debug"))
 	static UE_API UPARAM(DisplayName = "Target Mesh") UDynamicMesh*
 	ConvertMeshVertexColorsSRGBToLinear(
 		UDynamicMesh* TargetMesh,
@@ -190,7 +190,7 @@ public:
 	 * Apply a Linear to SRGB color transformation on all vertex colors
 	 * on the mesh.
 	 */
-	UFUNCTION(BlueprintCallable, Category = "GeometryScript|VertexColor", meta=(ScriptMethod))
+	UFUNCTION(BlueprintCallable, Category = "GeometryScript|VertexColor", meta=(ScriptMethod, HidePin = "Debug"))
 	static UE_API UPARAM(DisplayName = "Target Mesh") UDynamicMesh*
 	ConvertMeshVertexColorsLinearToSRGB(
 		UDynamicMesh* TargetMesh,
@@ -206,7 +206,7 @@ public:
      * @param Strength Each iteration, we will blur between the vertex of the color at the previous iteration and its neighbors' average by Strength amount (expected to be in the zero to one range).
      * @param BlurMode Determines how neighbors are weighted when computing their average.
 	 */
-	UFUNCTION(BlueprintCallable, Category = "GeometryScript|VertexColor", meta=(ScriptMethod))
+	UFUNCTION(BlueprintCallable, Category = "GeometryScript|VertexColor", meta=(ScriptMethod, HidePin = "Debug"))
 	static UE_API UPARAM(DisplayName = "Target Mesh") UDynamicMesh*
 	BlurMeshVertexColors(
 		UDynamicMesh* TargetMesh,
@@ -227,7 +227,7 @@ public:
 	 * @param Selection Optional subset of target mesh vertices to transfer weights to.
 	 * If left empty, colors will be transferred to all target mesh vertices.
 	 */
-	UFUNCTION(BlueprintCallable, Category = "GeometryScript|VertexColor", meta = (ScriptMethod))
+	UFUNCTION(BlueprintCallable, Category = "GeometryScript|VertexColor", meta = (ScriptMethod, HidePin = "Debug"))
 	static UE_API UPARAM(DisplayName = "Target Mesh") UDynamicMesh*
 	TransferVertexColorsFromMesh(
 		UDynamicMesh* SourceMesh,
diff --git a/Engine/Plugins/Runtime/GeometryScripting/Source/GeometryScriptingCore/Public/GeometryScript/MeshVoxelFunctions.h b/Engine/Plugins/Runtime/GeometryScripting/Source/GeometryScriptingCore/Public/GeometryScript/MeshVoxelFunctions.h
--- a/Engine/Plugins/Runtime/GeometryScripting/Source/GeometryScriptingCore/Public/GeometryScript/MeshVoxelFunctions.h
+++ b/Engine/Plugins/Runtime/GeometryScripting/Source/GeometryScriptingCore/Public/GeometryScript/MeshVoxelFunctions.h
@@ -138,7 +138,7 @@ public:
 	/**
 	* Replaces the mesh with a voxelized-and-meshed approximation (VoxWrap operation).
 	*/
-	UFUNCTION(BlueprintCallable, Category = "GeometryScript|Voxel", meta = (ScriptMethod, Keywords = "Vox Wrap", DisplayName = "Apply Mesh Solidify (Voxel Wrap)"))
+	UFUNCTION(BlueprintCallable, Category = "GeometryScript|Voxel", meta = (ScriptMethod, HidePin = "Debug", Keywords = "Vox Wrap", DisplayName = "Apply Mesh Solidify (Voxel Wrap)"))
 	static UE_API UPARAM(DisplayName = "Target Mesh") UDynamicMesh*
 	ApplyMeshSolidify(
 		UDynamicMesh* TargetMesh,
@@ -148,7 +148,7 @@ public:
 	/**
 	* Replaces the mesh with an SDF-based offset mesh approximation (VoxOffset operation).
 	*/
-	UFUNCTION(BlueprintCallable, Category = "GeometryScript|Voxel", meta=(ScriptMethod, Keywords="Vox Offset Dilate Contract", DisplayName = "Apply Mesh Morphology (Voxel Offset)"))
+	UFUNCTION(BlueprintCallable, Category = "GeometryScript|Voxel", meta=(ScriptMethod, HidePin = "Debug", Keywords="Vox Offset Dilate Contract", DisplayName = "Apply Mesh Morphology (Voxel Offset)"))
 	static UE_API UPARAM(DisplayName = "Target Mesh") UDynamicMesh*
 	ApplyMeshMorphology(
 		UDynamicMesh* TargetMesh,
diff --git a/Engine/Plugins/Runtime/GeometryScripting/Source/GeometryScriptingCore/Public/GeometryScript/PointSetFunctions.h b/Engine/Plugins/Runtime/GeometryScripting/Source/GeometryScriptingCore/Public/GeometryScript/PointSetFunctions.h
--- a/Engine/Plugins/Runtime/GeometryScripting/Source/GeometryScriptingCore/Public/GeometryScript/PointSetFunctions.h
+++ b/Engine/Plugins/Runtime/GeometryScripting/Source/GeometryScriptingCore/Public/GeometryScript/PointSetFunctions.h
@@ -112,7 +112,7 @@ public:
 	 * (2) if uniform spacing is requested, points will be ordered to have an octree-style coverage --
 	 *     so the first 8 points will cover the 8 octants (where samples are available) and the subsequent points will progressively fill in the space
 	 */
-	UFUNCTION(BlueprintCallable, Category = "GeometryScript|PointSet", meta = (AutoCreateRefTerm = "Options"))
+	UFUNCTION(BlueprintCallable, Category = "GeometryScript|PointSet", meta = (HidePin = "Debug", AutoCreateRefTerm = "Options"))
 	static UE_API void DownsamplePoints(
 		const TArray<FVector>& Points,
 		const FGeometryScriptPointPriorityOptions& Options,
diff --git a/Engine/Plugins/Runtime/GeometryScripting/Source/GeometryScriptingCore/Public/GeometryScript/SceneUtilityFunctions.h b/Engine/Plugins/Runtime/GeometryScripting/Source/GeometryScriptingCore/Public/GeometryScript/SceneUtilityFunctions.h
--- a/Engine/Plugins/Runtime/GeometryScripting/Source/GeometryScriptingCore/Public/GeometryScript/SceneUtilityFunctions.h
+++ b/Engine/Plugins/Runtime/GeometryScripting/Source/GeometryScriptingCore/Public/GeometryScript/SceneUtilityFunctions.h
@@ -77,7 +77,7 @@ public:
 	 * @param bTransformToWorld if true, output mesh is in World space
 	 * @param LocalToWorld local-to-world transform is returned here (whether mesh is in local or world space)
 	 */
-	UFUNCTION(BlueprintCallable, Category = "GeometryScript|Scene", meta = (ExpandEnumAsExecs = "Outcome"))
+	UFUNCTION(BlueprintCallable, Category = "GeometryScript|Scene", meta = (HidePin = "Debug", ExpandEnumAsExecs = "Outcome"))
 	static UE_API UPARAM(DisplayName = "Dynamic Mesh") UDynamicMesh*
 	CopyMeshFromComponent(
 		USceneComponent* Component,
@@ -93,7 +93,7 @@ public:
 	 * Configure the Material set on a PrimitiveComponent, by repeatedly calling SetMaterial.
 	 * This is a simple utility function and it's behavior will depend on the specifics of the Component.
 	 */
-	UFUNCTION(BlueprintCallable, Category = "GeometryScript|Scene")
+	UFUNCTION(BlueprintCallable, Category = "GeometryScript|Scene", meta = (HidePin = "Debug"))
 	static UE_API void
 	SetComponentMaterialList(
 		UPrimitiveComponent* Component, 
@@ -110,7 +110,7 @@ public:
 	 * @param bUseComplexCollision if true, complex collision is extracted, otherwise Simple collision shapes are meshed
 	 * @param SphereResolution determines tessellation density of sphere and capsule shapes
 	 */
-	UFUNCTION(BlueprintCallable, Category = "GeometryScript|Scene", meta = (ExpandEnumAsExecs = "Outcome"))
+	UFUNCTION(BlueprintCallable, Category = "GeometryScript|Scene", meta = (HidePin = "Debug", ExpandEnumAsExecs = "Outcome"))
 	static UE_API UPARAM(DisplayName = "Dynamic Mesh") UDynamicMesh*
 	CopyCollisionMeshesFromObject(
 		UObject* FromObject,
@@ -136,7 +136,7 @@ public:
 	 * @param OccludeMeshTransforms	Array of transforms for each occlude mesh. Array must have the same length as OccludeMeshes.
 	 * @param Options				Settings to control how occlusion is tested
 	 */
-	UFUNCTION(BlueprintCallable, Category = "GeometryScript|Scene", meta = (AutoCreateRefTerm = "SourceMeshes, SourceMeshTransforms, TransparentMeshes, TransparentMeshTransforms, OccludeMeshes, OccludeMeshTransforms, OcclusionOptions"))
+	UFUNCTION(BlueprintCallable, Category = "GeometryScript|Scene", meta = (HidePin = "Debug", AutoCreateRefTerm = "SourceMeshes, SourceMeshTransforms, TransparentMeshes, TransparentMeshTransforms, OccludeMeshes, OccludeMeshTransforms, OcclusionOptions"))
 	static UE_API void DetermineMeshOcclusion(
 		const TArray<UDynamicMesh*>& SourceMeshes,
 		const TArray<FTransform>& SourceMeshTransforms,
diff --git a/Engine/Plugins/Runtime/GeometryScripting/Source/GeometryScriptingCore/Public/GeometryScript/TextureMapFunctions.h b/Engine/Plugins/Runtime/GeometryScripting/Source/GeometryScriptingCore/Public/GeometryScript/TextureMapFunctions.h
--- a/Engine/Plugins/Runtime/GeometryScripting/Source/GeometryScriptingCore/Public/GeometryScript/TextureMapFunctions.h
+++ b/Engine/Plugins/Runtime/GeometryScripting/Source/GeometryScriptingCore/Public/GeometryScript/TextureMapFunctions.h
@@ -50,7 +50,7 @@ public:
 	/**
 	 * Samples the given TextureMap at the list of UV positions and returns the color at each position in ColorList output.
 	 */
-	UFUNCTION(BlueprintCallable, Category = "GeometryScript|TextureUtils")
+	UFUNCTION(BlueprintCallable, Category = "GeometryScript|TextureUtils", meta = (HidePin = "Debug"))
 	static UE_API void 
 	SampleTexture2DAtUVPositions(  
 		FGeometryScriptUVList UVList, 
@@ -63,7 +63,7 @@ public:
 	 * Sample the the given TextureMap at the list of UV positions and return the color at each position in ColorList output.
 	 * This function fetches GPU data before sampling so, depending on your application, it can be inefficient and slow!
 	 */
-	UFUNCTION(BlueprintCallable, Category = "GeometryScript|TextureUtils")
+	UFUNCTION(BlueprintCallable, Category = "GeometryScript|TextureUtils", meta = (HidePin = "Debug"))
 	static UE_API void 
 	SampleTextureRenderTarget2DAtUVPositions(  
 		FGeometryScriptUVList UVList, 
diff --git a/Engine/Plugins/Runtime/GeometryScripting/Source/GeometryScriptingEditor/Public/GeometryScript/CreateNewAssetUtilityFunctions.h b/Engine/Plugins/Runtime/GeometryScripting/Source/GeometryScriptingEditor/Public/GeometryScript/CreateNewAssetUtilityFunctions.h
--- a/Engine/Plugins/Runtime/GeometryScripting/Source/GeometryScriptingEditor/Public/GeometryScript/CreateNewAssetUtilityFunctions.h
+++ b/Engine/Plugins/Runtime/GeometryScripting/Source/GeometryScriptingEditor/Public/GeometryScript/CreateNewAssetUtilityFunctions.h
@@ -127,7 +127,7 @@ class GEOMETRYSCRIPTINGEDITOR_API UGeometryScriptLibrary_CreateNewAssetFunctions
 	GENERATED_BODY()
 public:
 
-	UFUNCTION(BlueprintCallable, Category = "GeometryScript|AssetManagement", meta = (ExpandEnumAsExecs = "Outcome"))
+	UFUNCTION(BlueprintCallable, Category = "GeometryScript|AssetManagement", meta = (HidePin = "Debug", ExpandEnumAsExecs = "Outcome"))
 	static void
 	CreateUniqueNewAssetPathName(
 		FString AssetFolderPath,
@@ -142,7 +142,7 @@ public:
 	/**
 	 * Create a new Volume from a Dynamic Mesh, in the same world as the calling blueprint.
 	 */
-	UFUNCTION(BlueprintCallable, Category = "GeometryScript|AssetManagement", meta = (ExpandEnumAsExecs = "Outcome", DisplayName = "Create New Volume From Mesh", WorldContext = "WorldContextObject"))
+	UFUNCTION(BlueprintCallable, Category = "GeometryScript|AssetManagement", meta = (HidePin = "Debug", ExpandEnumAsExecs = "Outcome", DisplayName = "Create New Volume From Mesh", WorldContext = "WorldContextObject"))
 	static UPARAM(DisplayName = "Volume Actor") AVolume*
 	CreateNewVolumeFromMesh_WorldContext(
 		UObject* WorldContextObject,
@@ -156,7 +156,7 @@ public:
 	/**
 	 * Create a new Volume from a Dynamic Mesh, in the specified world.
 	 */
-	UFUNCTION(BlueprintCallable, Category = "GeometryScript|AssetManagement", meta = (ExpandEnumAsExecs = "Outcome", DisplayName = "Create New Volume From Mesh In Target World"))
+	UFUNCTION(BlueprintCallable, Category = "GeometryScript|AssetManagement", meta = (HidePin = "Debug", ExpandEnumAsExecs = "Outcome", DisplayName = "Create New Volume From Mesh In Target World"))
 	static UPARAM(DisplayName = "Volume Actor") AVolume* 
 	CreateNewVolumeFromMesh(
 		UDynamicMesh* FromDynamicMesh, 
@@ -170,7 +170,7 @@ public:
 	/**
 	 * Create a new StaticMesh asset from a DynamicMesh. Save the asset at the AssetPathAndName location.
 	 */
-	UFUNCTION(BlueprintCallable, Category = "GeometryScript|AssetManagement", meta = (ExpandEnumAsExecs = "Outcome"))
+	UFUNCTION(BlueprintCallable, Category = "GeometryScript|AssetManagement", meta = (HidePin = "Debug", ExpandEnumAsExecs = "Outcome"))
 	static UPARAM(DisplayName = "Static Mesh Asset") UStaticMesh* 
 	CreateNewStaticMeshAssetFromMesh(
 		UDynamicMesh* FromDynamicMesh, 
@@ -183,7 +183,7 @@ public:
 	 * Create a new StaticMesh asset from a collection of LODs represented by an array of DynamicMeshes.
 	 * Save the asset at the AssetPathAndName location.
 	 */
-	UFUNCTION(BlueprintCallable, Category = "GeometryScript|AssetManagement", meta = (ScriptName="CreateNewStaticMeshAssetFromMeshLods;CreateNewStaticMeshAssetFromMeshLODs", ExpandEnumAsExecs = "Outcome"))
+	UFUNCTION(BlueprintCallable, Category = "GeometryScript|AssetManagement", meta = (HidePin = "Debug", ScriptName="CreateNewStaticMeshAssetFromMeshLods;CreateNewStaticMeshAssetFromMeshLODs", ExpandEnumAsExecs = "Outcome"))
 	static UPARAM(DisplayName = "Static Mesh Asset") UStaticMesh* 
 	CreateNewStaticMeshAssetFromMeshLODs(
 		TArray<UDynamicMesh*> FromDynamicMesh, 
@@ -196,7 +196,7 @@ public:
 	 * Create a new SkeletalMesh asset from a DynamicMesh and a Skeleton.
 	 * Save the asset at the AssetPathAndName location.
 	 */
-	UFUNCTION(BlueprintCallable, Category = "GeometryScript|AssetManagement", meta = (ExpandEnumAsExecs = "Outcome"))
+	UFUNCTION(BlueprintCallable, Category = "GeometryScript|AssetManagement", meta = (HidePin = "Debug", ExpandEnumAsExecs = "Outcome"))
 	static UPARAM(DisplayName = "Skeletal Mesh Asset") USkeletalMesh* 
 	CreateNewSkeletalMeshAssetFromMesh(
 		UDynamicMesh* FromDynamicMesh, 
@@ -210,7 +210,7 @@ public:
 	 * Create a new SkeletalMesh asset from a collection of LODs represented by an array of DynamicMeshes and a Skeleton.
 	 * Save the asset at the AssetPathAndName location.
 	 */
-	UFUNCTION(BlueprintCallable, Category = "GeometryScript|AssetManagement", meta = (ScriptName="CreateNewSkeletalMeshAssetFromMeshLods;CreateNewSkeletalMeshAssetFromMeshLODs", ExpandEnumAsExecs = "Outcome"))
+	UFUNCTION(BlueprintCallable, Category = "GeometryScript|AssetManagement", meta = (HidePin = "Debug", ScriptName="CreateNewSkeletalMeshAssetFromMeshLods;CreateNewSkeletalMeshAssetFromMeshLODs", ExpandEnumAsExecs = "Outcome"))
 	static UPARAM(DisplayName = "Skeletal Mesh Asset") USkeletalMesh* 
 	CreateNewSkeletalMeshAssetFromMeshLODs(
 		TArray<UDynamicMesh*> FromDynamicMeshLODs, 
@@ -220,7 +220,7 @@ public:
 		EGeometryScriptOutcomePins& Outcome,
 		UGeometryScriptDebug* Debug = nullptr);
 
-	UFUNCTION(BlueprintCallable, Category = "GeometryScript|AssetManagement", meta = (ExpandEnumAsExecs = "Outcome"))
+	UFUNCTION(BlueprintCallable, Category = "GeometryScript|AssetManagement", meta = (HidePin = "Debug", ExpandEnumAsExecs = "Outcome"))
 	static UPARAM(DisplayName = "Texture 2D Asset") UTexture2D* 
 	CreateNewTexture2DAsset(
 		UTexture2D* FromTexture, 
diff --git a/Engine/Plugins/Runtime/GeometryScripting/Source/GeometryScriptingEditor/Public/GeometryScript/EditorTextureMapFunctions.h b/Engine/Plugins/Runtime/GeometryScripting/Source/GeometryScriptingEditor/Public/GeometryScript/EditorTextureMapFunctions.h
--- a/Engine/Plugins/Runtime/GeometryScripting/Source/GeometryScriptingEditor/Public/GeometryScript/EditorTextureMapFunctions.h
+++ b/Engine/Plugins/Runtime/GeometryScripting/Source/GeometryScriptingEditor/Public/GeometryScript/EditorTextureMapFunctions.h
@@ -70,7 +70,7 @@ class GEOMETRYSCRIPTINGEDITOR_API UGeometryScriptLibrary_EditorTextureMapFunctio
 
 public:
 
-	UFUNCTION(BlueprintCallable, Category = "GeometryScript|TextureUtils")
+	UFUNCTION(BlueprintCallable, Category = "GeometryScript|TextureUtils", meta = (HidePin = "Debug"))
 	static FGeometryScriptChannelPackResult ChannelPack(
 		FGeometryScriptChannelPackSource RChannelSource,
 		FGeometryScriptChannelPackSource GChannelSource,
diff --git a/Engine/Plugins/Runtime/GeometryScripting/Source/GeometryScriptingEditor/Public/GeometryScript/OpenSubdivUtilityFunctions.h b/Engine/Plugins/Runtime/GeometryScripting/Source/GeometryScriptingEditor/Public/GeometryScript/OpenSubdivUtilityFunctions.h
--- a/Engine/Plugins/Runtime/GeometryScripting/Source/GeometryScriptingEditor/Public/GeometryScript/OpenSubdivUtilityFunctions.h
+++ b/Engine/Plugins/Runtime/GeometryScripting/Source/GeometryScriptingEditor/Public/GeometryScript/OpenSubdivUtilityFunctions.h
@@ -14,7 +14,7 @@ class GEOMETRYSCRIPTINGEDITOR_API UGeometryScriptLibrary_OpenSubdivFunctions : p
 public:
 
 
-	UFUNCTION(BlueprintCallable, Category = "GeometryScript|OpenSubdiv", meta=(ScriptMethod, DisplayName = "Apply PolyGroup Catmull Clark SubD"))
+	UFUNCTION(BlueprintCallable, Category = "GeometryScript|OpenSubdiv", meta=(ScriptMethod, HidePin = "Debug", DisplayName = "Apply PolyGroup Catmull Clark SubD"))
 	static UPARAM(DisplayName = "Target Mesh") UDynamicMesh* 
 	ApplyPolygroupCatmullClarkSubD(
 		UDynamicMesh* FromDynamicMesh, 
@@ -23,7 +23,7 @@ public:
 		UGeometryScriptDebug* Debug = nullptr);
 
 
-	UFUNCTION(BlueprintCallable, Category = "GeometryScript|OpenSubdiv", meta=(ScriptMethod, DisplayName = "Apply Triangle Loop SubD"))
+	UFUNCTION(BlueprintCallable, Category = "GeometryScript|OpenSubdiv", meta=(ScriptMethod, HidePin = "Debug", DisplayName = "Apply Triangle Loop SubD"))
 	static UPARAM(DisplayName = "Target Mesh") UDynamicMesh* 
 	ApplyTriangleLoopSubD(
 		UDynamicMesh* FromDynamicMesh, 
