Improved geometry script's CopyCollisionMeshesFromObject support for complex collision: It now supports any type that implements the IInterface_CollisionDataProvider interface, and specifically should now work better with Skeletal Mesh Components and Spline Mesh Components

diff --git a/Engine/Plugins/Runtime/GeometryScripting/Source/GeometryScriptingCore/Private/SceneUtilityFunctions.cpp b/Engine/Plugins/Runtime/GeometryScripting/Source/GeometryScriptingCore/Private/SceneUtilityFunctions.cpp
--- a/Engine/Plugins/Runtime/GeometryScripting/Source/GeometryScriptingCore/Private/SceneUtilityFunctions.cpp
+++ b/Engine/Plugins/Runtime/GeometryScripting/Source/GeometryScriptingCore/Private/SceneUtilityFunctions.cpp
@@ -142,25 +142,30 @@ UDynamicMesh* UGeometryScriptLibrary_SceneUtilityFunctions::CopyCollisionMeshesF
 		AccumulatedMesh.EnableMatchingAttributes(ReadMesh, true);
 	});
 
+	// Get the transform from the component (if applicable)
+	LocalToWorld = FTransform::Identity;
+	if (USceneComponent* SceneComponent = Cast<USceneComponent>(FromObject))
+	{
+		LocalToWorld = SceneComponent->GetComponentTransform();
+	}
+
 	if (bUseComplexCollision)
 	{
 		// find the Complex Collision mesh interface
-		IInterface_CollisionDataProvider* CollisionProvider = nullptr;
-		if (UStaticMeshComponent* StaticMeshComp = Cast<UStaticMeshComponent>(FromObject))
+		IInterface_CollisionDataProvider* CollisionProvider = Cast<IInterface_CollisionDataProvider>(FromObject);
+		// For static and skeletal mesh cases, try pulling from the asset if the component didn't give us collision data
+		if (!CollisionProvider || !CollisionProvider->ContainsPhysicsTriMeshData(true))
 		{
-			LocalToWorld = StaticMeshComp->GetComponentTransform();
-			UStaticMesh* StaticMesh = StaticMeshComp->GetStaticMesh();
-			CollisionProvider = (IInterface_CollisionDataProvider*)StaticMesh;
-		}
-		else if (UStaticMesh* StaticMesh = Cast<UStaticMesh>(FromObject))
-		{
-			LocalToWorld = FTransform::Identity;
-			CollisionProvider = (IInterface_CollisionDataProvider*)StaticMesh;
-		}
-		else if (UPrimitiveComponent* DynamicMeshComp = Cast<UDynamicMeshComponent>(FromObject))
-		{
-			LocalToWorld = DynamicMeshComp->GetComponentTransform();
-			CollisionProvider = (IInterface_CollisionDataProvider*)DynamicMeshComp;
+			if (USkeletalMeshComponent* SkeletalMeshComp = Cast<USkeletalMeshComponent>(FromObject))
+			{
+				USkeletalMesh* SkeletalMesh = SkeletalMeshComp->GetSkeletalMeshAsset();
+				CollisionProvider = Cast<IInterface_CollisionDataProvider>(SkeletalMesh);
+			}
+			else if (UStaticMeshComponent* StaticMeshComp = Cast<UStaticMeshComponent>(FromObject))
+			{
+				UStaticMesh* StaticMesh = StaticMeshComp->GetStaticMesh();
+				CollisionProvider = Cast<IInterface_CollisionDataProvider>(StaticMesh);
+			}
 		}
 
 		if (CollisionProvider == nullptr)
@@ -184,12 +189,10 @@ UDynamicMesh* UGeometryScriptLibrary_SceneUtilityFunctions::CopyCollisionMeshesF
 		const UBodySetup* BodySetup = nullptr;
 		if (UPrimitiveComponent* AnyComponent = Cast<UPrimitiveComponent>(FromObject))
 		{
-			LocalToWorld = AnyComponent->GetComponentTransform();
 			BodySetup = AnyComponent->GetBodySetup();
 		}
 		else if (const UStaticMesh* StaticMesh = Cast<UStaticMesh>(FromObject))
 		{
-			LocalToWorld = FTransform::Identity;
 			BodySetup = StaticMesh->GetBodySetup();
 		}
 
diff --git a/Engine/Plugins/Runtime/GeometryScripting/Source/GeometryScriptingCore/Public/GeometryScript/SceneUtilityFunctions.h b/Engine/Plugins/Runtime/GeometryScripting/Source/GeometryScriptingCore/Public/GeometryScript/SceneUtilityFunctions.h
--- a/Engine/Plugins/Runtime/GeometryScripting/Source/GeometryScriptingCore/Public/GeometryScript/SceneUtilityFunctions.h
+++ b/Engine/Plugins/Runtime/GeometryScripting/Source/GeometryScriptingCore/Public/GeometryScript/SceneUtilityFunctions.h
@@ -104,7 +104,7 @@ public:
 	/**
 	 * Extract the Collision Geometry from FromObject and copy/approximate it with meshes stored in ToDynamicMesh.
 	 * For Simple Collision, FromObject can be a StaticMesh Asset or any PrimitiveComponent
-	 * For Complex Collision, FromObject can be a StaticMesh Asset, StaticMeshComponent, or DynamicMeshComponent
+	 * For Complex Collision, FromObject can be a Static/Skeletal Mesh Asset/Component, Dynamic Mesh Component, or other Component w/ complex collision support
 	 * @param bTransformToWorld if true, output mesh is in World space
 	 * @param LocalToWorld local-to-world transform is returned here (whether mesh is in local or world space)
 	 * @param bUseComplexCollision if true, complex collision is extracted, otherwise Simple collision shapes are meshed
