Deprecated incorrect geometry script perlin noise node, and added a corrected v2

diff --git a/Engine/Plugins/Runtime/GeometryScripting/Source/GeometryScriptingCore/Private/MeshDeformFunctions.cpp b/Engine/Plugins/Runtime/GeometryScripting/Source/GeometryScriptingCore/Private/MeshDeformFunctions.cpp
--- a/Engine/Plugins/Runtime/GeometryScripting/Source/GeometryScriptingCore/Private/MeshDeformFunctions.cpp
+++ b/Engine/Plugins/Runtime/GeometryScripting/Source/GeometryScriptingCore/Private/MeshDeformFunctions.cpp
@@ -224,7 +224,23 @@ UDynamicMesh* UGeometryScriptLibrary_MeshDeformFunctions::ApplyMathWarpToMesh(
 }
 
 
-UDynamicMesh* UGeometryScriptLibrary_MeshDeformFunctions::ApplyPerlinNoiseToMesh(
+PRAGMA_DISABLE_DEPRECATION_WARNINGS
+UDynamicMesh*
+UGeometryScriptLibrary_MeshDeformFunctions::ApplyPerlinNoiseToMesh(
+	UDynamicMesh* TargetMesh,
+	FGeometryScriptMeshSelection Selection,
+	FGeometryScriptPerlinNoiseOptions Options,
+	UGeometryScriptDebug* Debug)
+{
+	// The compatibility version squares the Frequency to replicate the previous incorrect behavior
+	// and is otherwise the same as v2
+	FGeometryScriptPerlinNoiseOptions CompatOptions = Options;
+	CompatOptions.BaseLayer.Frequency *= CompatOptions.BaseLayer.Frequency;
+	return ApplyPerlinNoiseToMesh2(TargetMesh, Selection, CompatOptions, Debug);
+}
+PRAGMA_ENABLE_DEPRECATION_WARNINGS
+
+UDynamicMesh* UGeometryScriptLibrary_MeshDeformFunctions::ApplyPerlinNoiseToMesh2(
 	UDynamicMesh* TargetMesh,
 	FGeometryScriptMeshSelection Selection,
 	FGeometryScriptPerlinNoiseOptions Options,
@@ -264,7 +280,7 @@ UDynamicMesh* UGeometryScriptLibrary_MeshDeformFunctions::ApplyPerlinNoiseToMesh
 			if (bAlignWithNormal)
 			{
 				FVector NoisePos = (FVector)((double)Options.BaseLayer.Frequency * (Pos + Offsets[0]));
-				float Displacement = Options.BaseLayer.Magnitude * FMath::PerlinNoise3D(Options.BaseLayer.Frequency * NoisePos);
+				float Displacement = Options.BaseLayer.Magnitude * FMath::PerlinNoise3D(NoisePos);
 				Pos += Displacement * Normals[VertexID];
 			}
 			else
@@ -273,7 +289,7 @@ UDynamicMesh* UGeometryScriptLibrary_MeshDeformFunctions::ApplyPerlinNoiseToMesh
 				for (int32 k = 0; k < 3; ++k)
 				{
 					FVector NoisePos = (FVector)((double)Options.BaseLayer.Frequency * (Pos + Offsets[k]));
-					Displacement[k] = Options.BaseLayer.Magnitude * FMath::PerlinNoise3D(Options.BaseLayer.Frequency * NoisePos);
+					Displacement[k] = Options.BaseLayer.Magnitude * FMath::PerlinNoise3D(NoisePos);
 				}
 				Pos += Displacement;
 			}
diff --git a/Engine/Plugins/Runtime/GeometryScripting/Source/GeometryScriptingCore/Public/GeometryScript/MeshDeformFunctions.h b/Engine/Plugins/Runtime/GeometryScripting/Source/GeometryScriptingCore/Public/GeometryScript/MeshDeformFunctions.h
--- a/Engine/Plugins/Runtime/GeometryScripting/Source/GeometryScriptingCore/Public/GeometryScript/MeshDeformFunctions.h
+++ b/Engine/Plugins/Runtime/GeometryScripting/Source/GeometryScriptingCore/Public/GeometryScript/MeshDeformFunctions.h
@@ -265,9 +265,9 @@ public:
 	* @param Selection if non-empty, only the vertices identified by the selection will be displaced, otherwise
 	* the Options' EmptyBehavior will be followed.
 	*/
-	UFUNCTION(BlueprintCallable, Category = "GeometryScript|Deformations", meta=(ScriptMethod, HidePin = "Debug"))
+	UFUNCTION(BlueprintCallable, Category = "GeometryScript|Deformations", meta = (ScriptMethod, DisplayName = "Apply Perlin Noise to Mesh", HidePin = "Debug"))
 	static UE_API UPARAM(DisplayName = "Target Mesh") UDynamicMesh* 
-	ApplyPerlinNoiseToMesh(  
+	ApplyPerlinNoiseToMesh2(
 		UDynamicMesh* TargetMesh, 
 		FGeometryScriptMeshSelection Selection,
 		FGeometryScriptPerlinNoiseOptions Options,
@@ -315,6 +315,31 @@ public:
 		float Magnitude = 5.0,
 		UGeometryScriptDebug* Debug = nullptr);
 
+	//---------------------------------------------
+	// Backwards-Compatibility implementations
+	//---------------------------------------------
+	// Deprecated versions of functions are kept below
+	// 
+
+	/**
+	* Applies 3D Perlin noise displacement to the Target Mesh, with Frequency parameter incorrectly squared.
+	* @param Selection if non-empty, only the vertices identified by the selection will be displaced, otherwise
+	* the Options' EmptyBehavior will be followed.
+	* @warning It is strongly recommended that callers of this function update to the current ApplyPerlinNoiseToMesh function!
+	*/
+	UE_DEPRECATED(5.7, "Use 'ApplyPerlinNoiseToMesh2' instead. This method is deprecated because it incorrectly squared the frequency parameter.")
+	UFUNCTION(BlueprintCallable, Category = "GeometryScript|Compatibility", 
+		meta = (ScriptMethod, DeprecatedFunction, HidePin = "Debug",
+			DeprecationMessage = "Use 'Apply Perlin Noise to Mesh' instead. This node is deprecated because it incorrectly squared the frequency parameter.", 
+			DisplayName="Apply Perlin Noise to Mesh Compatibility 5.6 (Deprecated)"))
+	static UE_API UPARAM(DisplayName = "Target Mesh") UDynamicMesh*
+	ApplyPerlinNoiseToMesh(
+		UDynamicMesh* TargetMesh,
+		FGeometryScriptMeshSelection Selection,
+		FGeometryScriptPerlinNoiseOptions Options,
+		UGeometryScriptDebug* Debug = nullptr);
+
+
 };
 
 #undef UE_API
