reassure PVS about UGeometryScriptLibrary_MeshNormalsFunctions::UpdateVertexNormal safety + add comment to explain why it's safe

diff --git a/Engine/Plugins/Runtime/GeometryScripting/Source/GeometryScriptingCore/Private/MeshNormalsFunctions.cpp b/Engine/Plugins/Runtime/GeometryScripting/Source/GeometryScriptingCore/Private/MeshNormalsFunctions.cpp
--- a/Engine/Plugins/Runtime/GeometryScripting/Source/GeometryScriptingCore/Private/MeshNormalsFunctions.cpp
+++ b/Engine/Plugins/Runtime/GeometryScripting/Source/GeometryScriptingCore/Private/MeshNormalsFunctions.cpp
@@ -871,11 +871,12 @@ UDynamicMesh* UGeometryScriptLibrary_MeshNormalsFunctions::UpdateVertexNormal(
 		}
 		else   // bMergeSplitNormals == true && UniqueElementIDs.Num() > 1
 		{
-			Overlay->SetElement(UniqueElementIDs[0], (FVector3f)NewValue);
+			// Note the UniqueElementIDs array must not be empty, since TriVtxElements was not empty
+			Overlay->SetElement(UniqueElementIDs[0], (FVector3f)NewValue); //-V557
 			for (FIndex3i TriInfo : TriVtxElements)
 			{
 				FIndex3i OverlayTri = Overlay->GetTriangle(TriInfo.A);
-				OverlayTri[TriInfo.B] = UniqueElementIDs[0];
+				OverlayTri[TriInfo.B] = UniqueElementIDs[0]; //-V557
 				Overlay->SetTriangle(TriInfo.A, OverlayTri);
 			}
 		}
