remove requirement that patch unwrap mesh be non-compact

diff --git a/Engine/Plugins/Runtime/GeometryScripting/Source/GeometryScriptingCore/Private/MeshUVFunctions.cpp b/Engine/Plugins/Runtime/GeometryScripting/Source/GeometryScriptingCore/Private/MeshUVFunctions.cpp
--- a/Engine/Plugins/Runtime/GeometryScripting/Source/GeometryScriptingCore/Private/MeshUVFunctions.cpp
+++ b/Engine/Plugins/Runtime/GeometryScripting/Source/GeometryScriptingCore/Private/MeshUVFunctions.cpp
@@ -1069,12 +1069,6 @@ UDynamicMesh* UGeometryScriptLibrary_MeshUVFunctions::AutoGeneratePatchBuilderMe
 	ApplyMeshUVEditorOperation(TargetMesh, UVSetIndex, bHasUVSet, Debug,
 		[&](FDynamicMesh3& EditMesh, FDynamicMeshUVOverlay* UVOverlay, FDynamicMeshUVEditor& UVEditor)
 	{
-		if (EditMesh.IsCompact() == false)
-		{
-			UE::Geometry::AppendError(Debug, EGeometryScriptErrorType::InvalidInputs, LOCTEXT("AutoGeneratePatchBuilderMeshUVs_NonCompact", "AutoGeneratePatchBuilderMeshUVs: TargetMesh is non-Compact, PatchBuilder cannot be run. Try calling CompactMesh to update TargetMesh"));
-			return;
-		}
-
 		FPatchBasedMeshUVGenerator UVGenerator;
 
 		TUniquePtr<FPolygroupSet> PolygroupConstraint;
