Fixed issue where Geometry Script would maximally simplify a mesh if "Simplify to Tolerance" was run with a tolerance of 0

diff --git a/Engine/Plugins/Runtime/GeometryScripting/Source/GeometryScriptingCore/Private/MeshSimplifyFunctions.cpp b/Engine/Plugins/Runtime/GeometryScripting/Source/GeometryScriptingCore/Private/MeshSimplifyFunctions.cpp
--- a/Engine/Plugins/Runtime/GeometryScripting/Source/GeometryScriptingCore/Private/MeshSimplifyFunctions.cpp
+++ b/Engine/Plugins/Runtime/GeometryScripting/Source/GeometryScriptingCore/Private/MeshSimplifyFunctions.cpp
@@ -151,7 +151,8 @@ void DoSimplifyMesh(
 	bool bTargetIsTriCount,
 	int32 TargetCount,
 	FMeshProjectionTarget* ProjectionTarget = nullptr,
-	double GeometricTolerance = 0)
+	double GeometricTolerance = -1 // Note: Values < 0 will skip using GeometricTolerance
+)
 {
 	SimplificationType Simplifier(&EditMesh);
 
@@ -189,7 +190,7 @@ void DoSimplifyMesh(
 		Simplifier.CollapseMode = SimplificationType::ESimplificationCollapseModes::MinimalExistingVertexError;
 	}
 
-	if ( ProjectionTarget != nullptr && GeometricTolerance > 0)
+	if ( ProjectionTarget != nullptr && GeometricTolerance >= 0)
 	{
 		Simplifier.GeometricErrorConstraint = SimplificationType::EGeometricErrorCriteria::PredictedPointToProjectionTarget;
 		Simplifier.GeometricErrorTolerance = GeometricTolerance;
