add an OBB containment function to geometry script (matching what simple collision fitting does, but exposing the found oriented box directly)

diff --git a/Engine/Plugins/Runtime/GeometryScripting/Source/GeometryScriptingCore/Private/ContainmentFunctions.cpp b/Engine/Plugins/Runtime/GeometryScripting/Source/GeometryScriptingCore/Private/ContainmentFunctions.cpp
--- a/Engine/Plugins/Runtime/GeometryScripting/Source/GeometryScriptingCore/Private/ContainmentFunctions.cpp
+++ b/Engine/Plugins/Runtime/GeometryScripting/Source/GeometryScriptingCore/Private/ContainmentFunctions.cpp
@@ -8,6 +8,7 @@
 #include "Operations/MeshConvexHull.h"
 #include "Operations/MeshProjectionHull.h"
 #include "ShapeApproximation/MeshSimpleShapeApproximation.h"
+#include "MinVolumeBox3.h"
 #include "UDynamicMesh.h"
 
 #include UE_INLINE_GENERATED_CPP_BY_NAME(ContainmentFunctions)
@@ -201,6 +202,81 @@ UDynamicMesh* UGeometryScriptLibrary_ContainmentFunctions::ComputeMeshConvexDeco
 
 
 
+UDynamicMesh* UGeometryScriptLibrary_ContainmentFunctions::ComputeMeshOrientedBox(
+	UDynamicMesh* TargetMesh,
+	FOrientedBox& OrientedBoxOut,
+	FGeometryScriptMeshSelection Selection,
+	FGeometryScriptFitOrientedBoxOptions Options,
+	UGeometryScriptDebug* Debug)
+{
+	if (TargetMesh == nullptr)
+	{
+		UE::Geometry::AppendError(Debug, EGeometryScriptErrorType::InvalidInputs, LOCTEXT("ComputeMeshOrientedBox_InvalidInput", "ComputeMeshOrientedBox: TargetMesh is Null"));
+		return TargetMesh;
+	}
+
+	TargetMesh->ProcessMesh([&OrientedBoxOut, &Selection, &Options, Debug](const FDynamicMesh3& ProcessMesh)
+		{
+			TArray<int32> SelectedVertices;
+
+			if (!Selection.IsEmpty())
+			{
+				Selection.ProcessByVertexID(ProcessMesh, 
+					[&SelectedVertices](int32 VertexID) { SelectedVertices.Add(VertexID); }
+				);
+			}
+			else
+			{
+				SelectedVertices.Reserve(ProcessMesh.VertexCount());
+				for (int32 VID : ProcessMesh.VertexIndicesItr())
+				{
+					SelectedVertices.Add(VID);
+				}
+			}
+
+			int32 VertexCount = SelectedVertices.Num();
+			if (VertexCount == 0)
+			{
+				UE::Geometry::AppendError(Debug, EGeometryScriptErrorType::OperationFailed, LOCTEXT("ComputeMeshOrientedBox_Failed", "ComputeMeshOrientedBox: Cannot find bounds of zero vertices"));
+				return;
+			}
+
+			FOrientedBox3d ResultBox;
+			
+			bool bUseExactComputationForBox = Options.FitMethod == FGeometryScriptFitOrientedBoxMethod::Precise;
+			double MinDimension = FMath::Max(0., Options.MinBoxDimension);
+			FMinVolumeBox3d MinBoxCalc;
+			bool bMinBoxOK = MinBoxCalc.Solve(SelectedVertices.Num(),
+				[&ProcessMesh, &SelectedVertices](int32 Index)
+				{ return ProcessMesh.GetVertex(SelectedVertices[Index]); },
+				bUseExactComputationForBox, nullptr
+			);
+				
+			if (bMinBoxOK && MinBoxCalc.IsSolutionAvailable())
+			{
+				MinBoxCalc.GetResult(ResultBox);
+				double MinHalfDimension = MinDimension * .5;
+				for (int32 SubIdx = 0; SubIdx < 3; ++SubIdx)
+				{
+					ResultBox.Extents[SubIdx] = FMath::Max(MinHalfDimension, ResultBox.Extents[SubIdx]);
+				}
+				ResultBox = ResultBox.ReparameterizedCloserToWorldFrame();
+			}
+			else
+			{
+				// Fitting failed; fall back to AABB
+				ResultBox = ProcessMesh.GetBounds();
+				UE::Geometry::AppendWarning(Debug, EGeometryScriptErrorType::OperationFailed, LOCTEXT("ComputeMeshOrientedBox_Failed", "ComputeMeshOrientedBox: Failed to fit oriented box; falling back to axis-aligned box."));
+			}
+			
+			// convert result to the BP-accessible oriented box type
+			OrientedBoxOut = (FOrientedBox)ResultBox;
+		}
+	);
+
+	return TargetMesh;
+}
+
 
 
 #undef LOCTEXT_NAMESPACE
\ No newline at end of file
diff --git a/Engine/Plugins/Runtime/GeometryScripting/Source/GeometryScriptingCore/Public/GeometryScript/ContainmentFunctions.h b/Engine/Plugins/Runtime/GeometryScripting/Source/GeometryScriptingCore/Public/GeometryScript/ContainmentFunctions.h
--- a/Engine/Plugins/Runtime/GeometryScripting/Source/GeometryScriptingCore/Public/GeometryScript/ContainmentFunctions.h
+++ b/Engine/Plugins/Runtime/GeometryScripting/Source/GeometryScriptingCore/Public/GeometryScript/ContainmentFunctions.h
@@ -6,6 +6,7 @@
 #include "Kismet/BlueprintFunctionLibrary.h"
 #include "GeometryScript/GeometryScriptTypes.h"
 #include "GeometryScript/GeometryScriptSelectionTypes.h"
+#include "OrientedBoxTypes.h"
 #include "ContainmentFunctions.generated.h"
 
 #define UE_API GEOMETRYSCRIPTINGCORE_API
@@ -83,7 +84,29 @@ public:
 	int SimplifyToFaceCount = 0;
 };
 
+UENUM()
+enum class FGeometryScriptFitOrientedBoxMethod : uint8
+{
+	// Use an iterative method to fit the oriented box, starting from a fast initial guess and locally refining
+	FastIterative,
+	// Use a precise method to fit the oriented box, testing a potentially large space of options for the best fit
+	Precise
+};
 
+USTRUCT(BlueprintType)
+struct FGeometryScriptFitOrientedBoxOptions
+{
+	GENERATED_BODY()
+public:
+	/* Minimum length of the resulting bounding box along any of its principal directions */
+	UPROPERTY(EditAnywhere, BlueprintReadWrite, Category = Options, meta=(ClampMin=0))
+	double MinBoxDimension = .01;
+
+	/* Method to use for fitting the bounding box */
+	UPROPERTY(EditAnywhere, BlueprintReadWrite, Category = Options)
+	FGeometryScriptFitOrientedBoxMethod FitMethod = FGeometryScriptFitOrientedBoxMethod::FastIterative;
+
+};
 
 
 UCLASS(MinimalAPI, meta = (ScriptName = "GeometryScript_Containment"))
@@ -142,6 +165,20 @@ public:
 		FGeometryScriptConvexDecompositionOptions Options,
 		UGeometryScriptDebug* Debug = nullptr);
 
+	/**
+	 * Compute an oriented box fitting a given Target Mesh, or optionally a selection of the mesh
+	 *
+	 * @param Selection	Selection of mesh faces/vertices to contain in the oriented box. If not provided, entire mesh is used.
+	 */
+	UFUNCTION(BlueprintCallable, Category = "GeometryScript|Containment", meta = (ScriptMethod, KeyWords = "Fit Bounding OBB"))
+	static UE_API UPARAM(DisplayName = "Target Mesh") UDynamicMesh*
+	ComputeMeshOrientedBox(
+		UDynamicMesh* TargetMesh,
+		FOrientedBox& OrientedBoxOut,
+		FGeometryScriptMeshSelection Selection,
+		FGeometryScriptFitOrientedBoxOptions Options,
+		UGeometryScriptDebug* Debug = nullptr);
+
 };
 
 #undef UE_API
