fix duplicate loctext key in geometryscript obb containment method

diff --git a/Engine/Plugins/Runtime/GeometryScripting/Source/GeometryScriptingCore/Private/ContainmentFunctions.cpp b/Engine/Plugins/Runtime/GeometryScripting/Source/GeometryScriptingCore/Private/ContainmentFunctions.cpp
--- a/Engine/Plugins/Runtime/GeometryScripting/Source/GeometryScriptingCore/Private/ContainmentFunctions.cpp
+++ b/Engine/Plugins/Runtime/GeometryScripting/Source/GeometryScriptingCore/Private/ContainmentFunctions.cpp
@@ -237,7 +237,7 @@ UDynamicMesh* UGeometryScriptLibrary_ContainmentFunctions::ComputeMeshOrientedBo
 			int32 VertexCount = SelectedVertices.Num();
 			if (VertexCount == 0)
 			{
-				UE::Geometry::AppendError(Debug, EGeometryScriptErrorType::OperationFailed, LOCTEXT("ComputeMeshOrientedBox_Failed", "ComputeMeshOrientedBox: Cannot find bounds of zero vertices"));
+				UE::Geometry::AppendError(Debug, EGeometryScriptErrorType::OperationFailed, LOCTEXT("ComputeMeshOrientedBox_FailedEmptyInput", "ComputeMeshOrientedBox: Cannot find bounds of zero vertices"));
 				return;
 			}
 
@@ -266,7 +266,7 @@ UDynamicMesh* UGeometryScriptLibrary_ContainmentFunctions::ComputeMeshOrientedBo
 			{
 				// Fitting failed; fall back to AABB
 				ResultBox = ProcessMesh.GetBounds();
-				UE::Geometry::AppendWarning(Debug, EGeometryScriptErrorType::OperationFailed, LOCTEXT("ComputeMeshOrientedBox_Failed", "ComputeMeshOrientedBox: Failed to fit oriented box; falling back to axis-aligned box."));
+				UE::Geometry::AppendWarning(Debug, EGeometryScriptErrorType::OperationFailed, LOCTEXT("ComputeMeshOrientedBox_FailedFittingStep", "ComputeMeshOrientedBox: Failed to fit oriented box; falling back to axis-aligned box."));
 			}
 			
 			// convert result to the BP-accessible oriented box type
