fix bug with "AppendXWithSimpleCollision" geometry script methods where repeated calls of a script could accumulate extra copies of the collision shape, because blueprint re-uses output reference objects across script invocations

+ add oriented boxes to the geometry script primitive append functions

diff --git a/Engine/Plugins/Runtime/GeometryScripting/Source/GeometryScriptingCore/Private/MeshPrimitiveFunctions.cpp b/Engine/Plugins/Runtime/GeometryScripting/Source/GeometryScriptingCore/Private/MeshPrimitiveFunctions.cpp
--- a/Engine/Plugins/Runtime/GeometryScripting/Source/GeometryScriptingCore/Private/MeshPrimitiveFunctions.cpp
+++ b/Engine/Plugins/Runtime/GeometryScripting/Source/GeometryScriptingCore/Private/MeshPrimitiveFunctions.cpp
@@ -216,14 +216,11 @@ namespace UE::Private::AppendMeshHelper
 	static UDynamicMesh* AppendBox(UDynamicMesh* TargetMesh,
 		FGeometryScriptSimpleCollision* SimpleCollision,
 		FGeometryScriptPrimitiveOptions PrimitiveOptions,
-		FTransform Transform,
-		float DimensionX,
-		float DimensionY,
-		float DimensionZ,
+		const FTransform& Transform,
+		const UE::Geometry::FOrientedBox3d& OrientedBox,
 		int32 StepsX,
 		int32 StepsY,
 		int32 StepsZ,
-		EGeometryScriptPrimitiveOriginMode Origin,
 		UGeometryScriptDebug* Debug)
 	{
 		if (TargetMesh == nullptr)
@@ -232,33 +229,47 @@ namespace UE::Private::AppendMeshHelper
 			return TargetMesh;
 		}
 
-		UE::Geometry::FAxisAlignedBox3d ConvertBox(
-			FVector3d(-DimensionX / 2, -DimensionY / 2, 0),
-			FVector3d(DimensionX / 2, DimensionY / 2, DimensionZ));
-
 		// todo: if Steps X/Y/Z are zero, can use trivial box generator
 
 		FGridBoxMeshGenerator GridBoxGenerator;
-		GridBoxGenerator.Box = UE::Geometry::FOrientedBox3d(ConvertBox);
+		GridBoxGenerator.Box = OrientedBox;
 		GridBoxGenerator.EdgeVertices = FIndex3i(FMath::Max(0, StepsX), FMath::Max(0, StepsY), FMath::Max(0, StepsZ));
 		GridBoxGenerator.bPolygroupPerQuad = (PrimitiveOptions.PolygroupMode == EGeometryScriptPrimitivePolygroupMode::PerQuad);
 		GridBoxGenerator.Generate();
 
-		FVector3d OriginShift = (Origin == EGeometryScriptPrimitiveOriginMode::Center) ? FVector3d(0, 0, -DimensionZ / 2) : FVector3d::Zero();
-		AppendPrimitive(TargetMesh, &GridBoxGenerator, Transform, PrimitiveOptions, OriginShift);
+		AppendPrimitive(TargetMesh, &GridBoxGenerator, Transform, PrimitiveOptions, FVector3d::Zero());
 
 		if (SimpleCollision)
 		{
-			FVector3d Dims = Transform.GetScale3D() * FVector3d(DimensionX, DimensionY, DimensionZ);
+			FVector3d Dims = Transform.GetScale3D() * 2 * OrientedBox.Extents;
 			FKBoxElem BoxElem(Dims.X, Dims.Y, Dims.Z);
-			BoxElem.Center = Transform.TransformPosition(ConvertBox.Center() + OriginShift);
-			BoxElem.Rotation = Transform.Rotator();
+			BoxElem.Center = Transform.TransformPosition(OrientedBox.Center());
+			BoxElem.Rotation = (Transform.GetRotation() * FQuat(OrientedBox.Frame.Rotation)).Rotator();
 			SimpleCollision->AggGeom.BoxElems.Add(BoxElem);
 		}
 
 		return TargetMesh;
 	}
 
+	static UDynamicMesh* AppendBox(UDynamicMesh* TargetMesh,
+		FGeometryScriptSimpleCollision* SimpleCollision,
+		FGeometryScriptPrimitiveOptions PrimitiveOptions,
+		FTransform Transform,
+		float DimensionX,
+		float DimensionY,
+		float DimensionZ,
+		int32 StepsX,
+		int32 StepsY,
+		int32 StepsZ,
+		EGeometryScriptPrimitiveOriginMode Origin,
+		UGeometryScriptDebug* Debug)
+	{
+		UE::Geometry::FOrientedBox3d OrientedBox(
+			FVector3d(0, 0, Origin == EGeometryScriptPrimitiveOriginMode::Base ? DimensionZ * .5f : 0.f),
+			FVector3d(DimensionX * .5f, DimensionY * .5f, DimensionZ * .5f)); // Extents are half-dimensions
+		return UE::Private::AppendMeshHelper::AppendBox(TargetMesh, SimpleCollision, PrimitiveOptions, Transform, OrientedBox, StepsX, StepsY, StepsZ, Debug);
+	}
+
 	static UDynamicMesh* AppendBoundingBox(
 		UDynamicMesh* TargetMesh,
 		FGeometryScriptSimpleCollision* SimpleCollision,
@@ -462,6 +473,7 @@ UDynamicMesh* UGeometryScriptLibrary_MeshPrimitiveFunctions::AppendBoxWithCollis
 	EGeometryScriptPrimitiveOriginMode Origin,
 	UGeometryScriptDebug* Debug)
 {
+	SimpleCollision.AggGeom.EmptyElements(); // Clear the existing elements since BP can re-use these output structs
 	return UE::Private::AppendMeshHelper::AppendBox(TargetMesh, &SimpleCollision, PrimitiveOptions, Transform, DimensionX, DimensionY, DimensionZ, StepsX, StepsY, StepsZ, Origin, Debug);
 }
 
@@ -490,10 +502,41 @@ UDynamicMesh* UGeometryScriptLibrary_MeshPrimitiveFunctions::AppendBoundingBoxWi
 	int32 StepsZ,
 	UGeometryScriptDebug* Debug)
 {
+	SimpleCollision.AggGeom.EmptyElements(); // Clear the existing elements since BP can re-use these output structs
 	return UE::Private::AppendMeshHelper::AppendBoundingBox(TargetMesh, &SimpleCollision, PrimitiveOptions, Transform, Box, StepsX, StepsY, StepsZ, Debug);
 }
 
 
+UDynamicMesh* UGeometryScriptLibrary_MeshPrimitiveFunctions::AppendOrientedBox(
+	UDynamicMesh* TargetMesh,
+	FGeometryScriptPrimitiveOptions PrimitiveOptions,
+	FTransform Transform,
+	FOrientedBox OrientedBox,
+	int32 StepsX,
+	int32 StepsY,
+	int32 StepsZ,
+	UGeometryScriptDebug* Debug)
+{
+	return UE::Private::AppendMeshHelper::AppendBox(TargetMesh, nullptr, PrimitiveOptions, Transform, 
+		UE::Geometry::FOrientedBox3d(OrientedBox), StepsX, StepsY, StepsZ, Debug);
+}
+
+UDynamicMesh* UGeometryScriptLibrary_MeshPrimitiveFunctions::AppendOrientedBoxWithCollision(
+	UDynamicMesh* TargetMesh,
+	FGeometryScriptSimpleCollision& SimpleCollision,
+	FGeometryScriptPrimitiveOptions PrimitiveOptions,
+	FTransform Transform,
+	FOrientedBox OrientedBox,
+	int32 StepsX,
+	int32 StepsY,
+	int32 StepsZ,
+	UGeometryScriptDebug* Debug)
+{
+	SimpleCollision.AggGeom.EmptyElements(); // Clear the existing elements since BP can re-use these output structs
+	return UE::Private::AppendMeshHelper::AppendBox(TargetMesh, &SimpleCollision, PrimitiveOptions, Transform,
+		UE::Geometry::FOrientedBox3d(OrientedBox), StepsX, StepsY, StepsZ, Debug);
+}
+
 
 UDynamicMesh* UGeometryScriptLibrary_MeshPrimitiveFunctions::AppendSphereLatLong(
 	UDynamicMesh* TargetMesh,
@@ -519,6 +562,7 @@ UDynamicMesh* UGeometryScriptLibrary_MeshPrimitiveFunctions::AppendSphereLatLong
 	EGeometryScriptPrimitiveOriginMode Origin,
 	UGeometryScriptDebug* Debug)
 {
+	SimpleCollision.AggGeom.EmptyElements(); // Clear the existing elements since BP can re-use these output structs
 	return UE::Private::AppendMeshHelper::AppendSphereLatLong(TargetMesh, &SimpleCollision, PrimitiveOptions, Transform, Radius, StepsPhi, StepsTheta, Origin, Debug);
 }
 
@@ -548,6 +592,7 @@ UDynamicMesh* UGeometryScriptLibrary_MeshPrimitiveFunctions::AppendSphereBoxWith
 	EGeometryScriptPrimitiveOriginMode Origin,
 	UGeometryScriptDebug* Debug)
 {
+	SimpleCollision.AggGeom.EmptyElements(); // Clear the existing elements since BP can re-use these output structs
 	return UE::Private::AppendMeshHelper::AppendSphereBox(TargetMesh, &SimpleCollision, PrimitiveOptions, Transform, Radius, StepsX, StepsY, StepsZ, Origin, Debug);
 }
 
@@ -581,6 +626,7 @@ UDynamicMesh* UGeometryScriptLibrary_MeshPrimitiveFunctions::AppendCapsuleWithCo
 	EGeometryScriptPrimitiveOriginMode Origin,
 	UGeometryScriptDebug* Debug)
 {
+	SimpleCollision.AggGeom.EmptyElements(); // Clear the existing elements since BP can re-use these output structs
 	return UE::Private::AppendMeshHelper::AppendCapsule(TargetMesh, &SimpleCollision, PrimitiveOptions, Transform, Radius, LineLength, HemisphereSteps, CircleSteps, SegmentSteps, Origin, Debug);
 }
 
diff --git a/Engine/Plugins/Runtime/GeometryScripting/Source/GeometryScriptingCore/Public/GeometryScript/MeshPrimitiveFunctions.h b/Engine/Plugins/Runtime/GeometryScripting/Source/GeometryScriptingCore/Public/GeometryScript/MeshPrimitiveFunctions.h
--- a/Engine/Plugins/Runtime/GeometryScripting/Source/GeometryScriptingCore/Public/GeometryScript/MeshPrimitiveFunctions.h
+++ b/Engine/Plugins/Runtime/GeometryScripting/Source/GeometryScriptingCore/Public/GeometryScript/MeshPrimitiveFunctions.h
@@ -213,6 +213,38 @@ public:
 		int32 StepsZ = 0,
 		UGeometryScriptDebug* Debug = nullptr);
 
+	/**
+	 * Appends a 3D oriented box to the Target Mesh
+	 */
+	UFUNCTION(BlueprintCallable, Category = "GeometryScript|Primitives", meta = (ScriptMethod, HidePin = "Debug"))
+	static UE_API UPARAM(DisplayName = "Target Mesh") UDynamicMesh*
+	AppendOrientedBox(
+		UDynamicMesh* TargetMesh,
+		FGeometryScriptPrimitiveOptions PrimitiveOptions,
+		FTransform Transform,
+		FOrientedBox Box,
+		int32 StepsX = 0,
+		int32 StepsY = 0,
+		int32 StepsZ = 0,
+		UGeometryScriptDebug* Debug = nullptr);
+
+	/**
+	 * Appends a 3D oriented box to the Target Mesh
+	 * Also creates matching simple collision
+	 */
+	UFUNCTION(BlueprintCallable, Category = "GeometryScript|Primitives", meta = (ScriptMethod, HidePin = "Debug"))
+	static UE_API UPARAM(DisplayName = "Target Mesh") UDynamicMesh*
+	AppendOrientedBoxWithCollision(
+		UDynamicMesh* TargetMesh,
+		FGeometryScriptSimpleCollision& SimpleCollision,
+		FGeometryScriptPrimitiveOptions PrimitiveOptions,
+		FTransform Transform,
+		FOrientedBox Box,
+		int32 StepsX = 0,
+		int32 StepsY = 0,
+		int32 StepsZ = 0,
+		UGeometryScriptDebug* Debug = nullptr);
+
 	/**
 	 * Appends a 3D box to the Target Mesh with dimensions and origin taken from the input Box
 	 * Also creates matching simple collision
