Add Vertex and Triangle ID validation to geometry script Mesh Query methods so users do not crash on querying an invalid element ID

diff --git a/Engine/Plugins/Runtime/GeometryScripting/Source/GeometryScriptingCore/Private/MeshQueryFunctions.cpp b/Engine/Plugins/Runtime/GeometryScripting/Source/GeometryScriptingCore/Private/MeshQueryFunctions.cpp
--- a/Engine/Plugins/Runtime/GeometryScripting/Source/GeometryScriptingCore/Private/MeshQueryFunctions.cpp
+++ b/Engine/Plugins/Runtime/GeometryScripting/Source/GeometryScriptingCore/Private/MeshQueryFunctions.cpp
@@ -57,7 +57,38 @@ ReturnType SimpleMeshUVSetQuery(UDynamicMesh* Mesh, int32 UVSetIndex, bool& bIsV
 }
 
 
-
+namespace UE::MeshQueryFunctionLocals
+{
+	static bool ValidateMeshVertexID(const FDynamicMesh3& Mesh, int32 VID, bool bLogError = true)
+	{
+		if (!Mesh.IsVertex(VID))
+		{
+			if (bLogError)
+			{
+				UE::Geometry::MakeScriptError(EGeometryScriptErrorType::InvalidInputs, 
+					FText::Format(LOCTEXT("MeshQueryInvalidVertexID", "Mesh queried with invalid vertex ID: {0} (Max ID: {1})"),
+						VID, Mesh.MaxVertexID()));
+			}
+			return false;
+		}
+		return true;
+	}
+	
+	static bool ValidateMeshTriangleID(const FDynamicMesh3& Mesh, int32 TID, bool bLogError = true)
+	{
+		if (!Mesh.IsTriangle(TID))
+		{
+			if (bLogError)
+			{
+				UE::Geometry::MakeScriptError(EGeometryScriptErrorType::InvalidInputs,
+					FText::Format(LOCTEXT("MeshQueryInvalidTriangleID", "Mesh queried with invalid triangle ID: {0} (Max ID: {1})"),
+						TID, Mesh.MaxTriangleID()));
+			}
+			return false;
+		}
+		return true;
+	}
+}
 
 FString UGeometryScriptLibrary_MeshQueryFunctions::GetMeshInfoString(UDynamicMesh* TargetMesh)
 {
@@ -374,6 +405,10 @@ UDynamicMesh* UGeometryScriptLibrary_MeshQueryFunctions::GetAllVertexPositions(U
 UDynamicMesh* UGeometryScriptLibrary_MeshQueryFunctions::GetVertexConnectedTriangles(UDynamicMesh* TargetMesh, int32 VertexID, TArray<int32>& Triangles)
 {
 	SimpleMeshQuery<int32>(TargetMesh, 0, [&](const FDynamicMesh3& Mesh) {
+		if (!UE::MeshQueryFunctionLocals::ValidateMeshVertexID(Mesh, VertexID))
+		{
+			return 0;
+		}
 		Mesh.GetVtxTriangles(VertexID, Triangles);
 		return 0;
 	});
@@ -384,6 +419,10 @@ UDynamicMesh* UGeometryScriptLibrary_MeshQueryFunctions::GetVertexConnectedTrian
 UDynamicMesh* UGeometryScriptLibrary_MeshQueryFunctions::GetVertexConnectedVertices(UDynamicMesh* TargetMesh, int32 VertexID, TArray<int32>& Vertices)
 {
 	SimpleMeshQuery<int32>(TargetMesh, 0, [&](const FDynamicMesh3& Mesh) {
+		if (!UE::MeshQueryFunctionLocals::ValidateMeshVertexID(Mesh, VertexID))
+		{
+			return 0;
+		}
 		Mesh.EnumerateVertexVertices(VertexID, [&](int32 vid) { Vertices.Add(vid); });
 		return 0;
 	});
@@ -580,6 +619,10 @@ UDynamicMesh* UGeometryScriptLibrary_MeshQueryFunctions::GetTriangleNormals(UDyn
 {
 	bTriHasValidNormals = SimpleMeshQuery<bool>(TargetMesh, false, [&](const FDynamicMesh3& Mesh) 
 	{
+		if (!UE::MeshQueryFunctionLocals::ValidateMeshTriangleID(Mesh, TriangleID))
+		{
+			return false;
+		}
 		if (Mesh.HasAttributes())
 		{
 			if (const FDynamicMeshNormalOverlay* Normals = Mesh.Attributes()->PrimaryNormals())
@@ -606,6 +649,10 @@ UDynamicMesh* UGeometryScriptLibrary_MeshQueryFunctions::GetInterpolatedTriangle
 	bTriHasValidNormals = false;
 	InterpolatedNormal = SimpleMeshQuery<FVector>(TargetMesh, FVector::UnitZ(), [&](const FDynamicMesh3& Mesh)
 	{
+		if (!UE::MeshQueryFunctionLocals::ValidateMeshTriangleID(Mesh, TriangleID))
+		{
+			return FVector::UnitZ();
+		}
 		if (Mesh.HasAttributes())
 		{
 			if (const FDynamicMeshNormalOverlay* Normals = Mesh.Attributes()->PrimaryNormals())
@@ -630,6 +677,10 @@ UDynamicMesh* UGeometryScriptLibrary_MeshQueryFunctions::GetTriangleNormalTangen
 {
 	bTriHasValidElements = SimpleMeshQuery<bool>(TargetMesh, false, [&](const FDynamicMesh3& Mesh) 
 	{
+		if (!UE::MeshQueryFunctionLocals::ValidateMeshTriangleID(Mesh, TriangleID))
+		{
+			return false;
+		}
 		if (Mesh.HasAttributes() && Mesh.Attributes()->HasTangentSpace() )
 		{
 			const FDynamicMeshNormalOverlay* Normals = Mesh.Attributes()->PrimaryNormals();
@@ -658,6 +709,10 @@ UDynamicMesh* UGeometryScriptLibrary_MeshQueryFunctions::GetInterpolatedTriangle
 {
 	bTriHasValidElements = SimpleMeshQuery<bool>(TargetMesh, false, [&](const FDynamicMesh3& Mesh)
 	{
+		if (!UE::MeshQueryFunctionLocals::ValidateMeshTriangleID(Mesh, TriangleID))
+		{
+			return false;
+		}
 		if (Mesh.HasAttributes() && Mesh.Attributes()->HasTangentSpace() )
 		{
 			const FDynamicMeshNormalOverlay* Normals = Mesh.Attributes()->PrimaryNormals();
@@ -694,6 +749,10 @@ UDynamicMesh* UGeometryScriptLibrary_MeshQueryFunctions::GetTriangleVertexColors
 {
 	bHaveValidVertexColors = SimpleMeshQuery<bool>(TargetMesh, false, [&](const FDynamicMesh3& Mesh) 
 	{
+		if (!UE::MeshQueryFunctionLocals::ValidateMeshTriangleID(Mesh, TriangleID))
+		{
+			return false;
+		}
 		if (Mesh.HasAttributes())
 		{
 			if (const FDynamicMeshColorOverlay* Colors = Mesh.Attributes()->PrimaryColors())
@@ -720,6 +779,10 @@ UDynamicMesh* UGeometryScriptLibrary_MeshQueryFunctions::GetInterpolatedTriangle
 	bHaveValidVertexColors = false;
 	InterpolatedColor = SimpleMeshQuery<FLinearColor>(TargetMesh, DefaultColor, [&](const FDynamicMesh3& Mesh) 
 	{
+		if (!UE::MeshQueryFunctionLocals::ValidateMeshTriangleID(Mesh, TriangleID))
+		{
+			return DefaultColor;
+		}
 		if (Mesh.HasAttributes())
 		{
 			if (const FDynamicMeshColorOverlay* Colors = Mesh.Attributes()->PrimaryColors())
