Fill out geometry script list conversion extract/set methods for the various list types: Scalar, UV, Vector and Color

Use these to convert between vector types, and to e.g. extract components to run scalar math operations on, then write the results back.

diff --git a/Engine/Plugins/Runtime/GeometryScripting/Source/GeometryScriptingCore/Private/ListUtilityFunctions.cpp b/Engine/Plugins/Runtime/GeometryScripting/Source/GeometryScriptingCore/Private/ListUtilityFunctions.cpp
--- a/Engine/Plugins/Runtime/GeometryScripting/Source/GeometryScriptingCore/Private/ListUtilityFunctions.cpp
+++ b/Engine/Plugins/Runtime/GeometryScripting/Source/GeometryScriptingCore/Private/ListUtilityFunctions.cpp
@@ -361,27 +361,188 @@ void UGeometryScriptLibrary_ListUtilityFunctions::ClearColorList(FGeometryScript
 	ColorList.List->Init(ClearColor, Num);
 }
 
-
-void UGeometryScriptLibrary_ListUtilityFunctions::ExtractColorListChannel(FGeometryScriptColorList ColorList, FGeometryScriptScalarList& ScalarList, int32 ChannelIndex)
+namespace UE::ListUtilityFunctionsLocals
 {
-	if (ChannelIndex < 0 || ChannelIndex > 3)
+	//
+	// helpers to access different list-array types in a consistent way, for template use
+	//
+
+	static double AccessListHelper(const TArray<double>& List, int32 ElIdx, int32 CompIdx)
+	{
+		return List[ElIdx];
+	}
+	static double& AccessListHelper(TArray<double>& List, int32 ElIdx, int32 CompIdx)
 	{
-		UE_LOG(LogGeometry, Warning, TEXT("ExtractColorListChannel: ChannelIndex is not in valid range 0-3"));
+		return List[ElIdx];
 	}
-	ChannelIndex = FMath::Clamp(ChannelIndex, 0, 3);
 
-	ScalarList.Reset();
-	if (ColorList.List.IsValid())
+	static double AccessListHelper(const TArray<FVector2D>& List, int32 ElIdx, int32 CompIdx)
+	{
+		return List[ElIdx][CompIdx];
+	}
+	static double& AccessListHelper(TArray<FVector2D>& List, int32 ElIdx, int32 CompIdx)
+	{
+		return List[ElIdx][CompIdx];
+	}
+
+	static double AccessListHelper(const TArray<FVector>& List, int32 ElIdx, int32 CompIdx)
 	{
-		TArray<FLinearColor>& Colors = *ColorList.List;
-		TArray<double>& Scalars = *ScalarList.List;
-		int32 N = Colors.Num();
-		Scalars.SetNumUninitialized(N);
-		for (int32 k = 0; k < N; ++k)
+		return List[ElIdx][CompIdx];
+	}
+	static double& AccessListHelper(TArray<FVector>& List, int32 ElIdx, int32 CompIdx)
+	{
+		return List[ElIdx][CompIdx];
+	}
+
+	static float AccessListHelper(const TArray<FLinearColor>& List, int32 ElIdx, int32 CompIdx)
+	{
+		return List[ElIdx].Component(CompIdx);
+	}
+	static float& AccessListHelper(TArray<FLinearColor>& List, int32 ElIdx, int32 CompIdx)
+	{
+		return List[ElIdx].Component(CompIdx);
+	}
+
+	//
+	// helpers to get the number of components per element of different list-array types, for template use
+	//
+
+	template<typename ElType>
+	struct TNumComponents
+	{};
+	template<>
+	struct TNumComponents<double>
+	{
+		constexpr static int32 Value = 1;
+	};
+	template<>
+	struct TNumComponents<FVector2D>
+	{
+		constexpr static int32 Value = 2;
+	};
+	template<>
+	struct TNumComponents<FVector>
+	{
+		constexpr static int32 Value = 3;
+	};
+	template<>
+	struct TNumComponents<FLinearColor>
+	{
+		constexpr static int32 Value = 4;
+	};
+
+	// templated helper to copy a subset of element components values from one array to another (i.e. the Set Components methods)
+
+	template<typename ToRealType, typename ToArrayType, typename FromArrayType>
+	void SetArrayComponentsFromArrayHelper(
+		ToArrayType& ToArr, const FromArrayType& FromArr, typename ToArrayType::ElementType DefaultVal, 
+		int InToOutDimMapping[TNumComponents<typename ToArrayType::ElementType>::Value])
+	{
+		constexpr int32 ToDims = TNumComponents<typename ToArrayType::ElementType>::Value;
+		constexpr int32 FromDims = TNumComponents<typename FromArrayType::ElementType>::Value;
+		
+		// Validate / clamp component indices
+		for (int32 FromDim = 0; FromDim < FromDims; ++FromDim)
 		{
-			FVector4d Vector4 = (FVector4d)(Colors[k]);
-			Scalars[k] = Vector4[ChannelIndex];
+			if (InToOutDimMapping[FromDim] < 0 || InToOutDimMapping[FromDim] >= ToDims)
+			{
+				UE::Geometry::MakeScriptWarning(EGeometryScriptErrorType::InvalidInputs,
+					FText::Format(LOCTEXT("InvalidListElementComponentIndex", "Component Index {0} will be clamped to valid range (0-{1})"),
+						InToOutDimMapping[FromDim], ToDims - 1));
+			}
+			InToOutDimMapping[FromDim] = FMath::Clamp(InToOutDimMapping[FromDim], 0, ToDims - 1);
 		}
+		int32 InitialToArrNum = ToArr.Num();
+		int32 NoDefaultNeededNum = FMath::Min(ToArr.Num(), FromArr.Num());
+
+		// Helper to copy across the requested components for a single element
+		auto CopyElementComponents = [&InToOutDimMapping, &FromArr, &ToArr, FromDims](int32 Idx)
+		{
+			for (int32 FromDimIdx = 0; FromDimIdx < FromDims; ++FromDimIdx)
+			{
+				int32 ToDimIdx = InToOutDimMapping[FromDimIdx];
+				AccessListHelper(ToArr, Idx, ToDimIdx) = (ToRealType)AccessListHelper(FromArr, Idx, FromDimIdx);
+			}
+		};
+
+		// expand the 'to' array if needed
+		if (ToArr.Num() < FromArr.Num())
+		{
+			ToArr.SetNumUninitialized(FromArr.Num());
+		}
+		// for range we didn't just add new elements, no default value is needed -- just write the requested components 
+		for (int32 Idx = 0; Idx < NoDefaultNeededNum; ++Idx)
+		{
+			CopyElementComponents(Idx);
+		}
+		// for range we added, set the default value before writing requested components
+		for (int32 Idx = NoDefaultNeededNum; Idx < FromArr.Num(); ++Idx)
+		{
+			ToArr[Idx] = DefaultVal;
+			CopyElementComponents(Idx);
+		}
+	}
+
+	// standard validation for the 'Set List from List' methods
+
+	template<typename ListAType, typename ListBType>
+	bool ValidateListsToSetAFromB(ListAType& ListA, const ListBType& ListB)
+	{
+		if (!ListB.List.IsValid())
+		{
+			return false;
+		}
+		if (!ListA.List.IsValid())
+		{
+			ListA.Reset();
+		}
+		return true;
+	}
+
+	// templated helper to extract a subset element components into a new array (i.e. the Extract methods)
+
+	template<typename FromArrayType, typename ToArrayType>
+	void ExtractArrayComponentsFromArrayHelper(
+		
+		const FromArrayType& FromArr, ToArrayType& ToArr, int FromComponent[TNumComponents<typename ToArrayType::ElementType>::Value])
+	{
+		constexpr int32 ToDims = TNumComponents<typename ToArrayType::ElementType>::Value;
+		constexpr int32 FromDims = TNumComponents<typename FromArrayType::ElementType>::Value;
+		
+		// Validate / clamp component indices
+		for (int32 ToDim = 0; ToDim < ToDims; ++ToDim)
+		{
+			if (FromComponent[ToDim] < 0 || FromComponent[ToDim] >= FromDims)
+			{
+				UE::Geometry::MakeScriptWarning(EGeometryScriptErrorType::InvalidInputs,
+					FText::Format(LOCTEXT("InvalidListElementComponentIndex", "Component Index {0} will be clamped to valid range (0-{1})"), 
+						FromComponent[ToDim], FromDims - 1));
+			}
+			FromComponent[ToDim] = FMath::Clamp(FromComponent[ToDim], 0, FromDims - 1);
+		}
+
+		// Extract requested components to output array
+		ToArr.SetNumUninitialized(FromArr.Num());
+		for (int32 Idx = 0; Idx < FromArr.Num(); ++Idx)
+		{
+			for (int32 ToDimIdx = 0; ToDimIdx < ToDims; ++ToDimIdx)
+			{
+				int32 FromDimIdx = FromComponent[ToDimIdx];
+				AccessListHelper(ToArr, Idx, ToDimIdx) = (double)AccessListHelper(FromArr, Idx, FromDimIdx);
+			}
+		}
+	}
+}
+
+void UGeometryScriptLibrary_ListUtilityFunctions::ExtractColorListChannel(FGeometryScriptColorList ColorList, FGeometryScriptScalarList& ScalarList, int32 ChannelIndex)
+{
+	using namespace UE::ListUtilityFunctionsLocals;
+
+	ScalarList.Reset();
+	if (ColorList.List.IsValid())
+	{
+		int ChannelMap[1]{ ChannelIndex };
+		ExtractArrayComponentsFromArrayHelper(*ColorList.List, *ScalarList.List, ChannelMap);
 	}
 }
 
@@ -390,37 +551,134 @@ void UGeometryScriptLibrary_ListUtilityFunctions::ExtractColorListChannels(
 	FGeometryScriptColorList ColorList, FGeometryScriptVectorList& VectorList, 
 	int32 XChannelIndex, int32 YChannelIndex, int32 ZChannelIndex)
 {
-	if (XChannelIndex < 0 || XChannelIndex > 3)
+	using namespace UE::ListUtilityFunctionsLocals;
+
+	VectorList.Reset();
+	if (ColorList.List.IsValid())
 	{
-		UE_LOG(LogGeometry, Warning, TEXT("ExtractColorListChannels: XChannelIndex is not in valid range 0-3"));
+		int ChannelMap[3]{ XChannelIndex, YChannelIndex, ZChannelIndex };
+		ExtractArrayComponentsFromArrayHelper(*ColorList.List, *VectorList.List, ChannelMap);
 	}
-	if (YChannelIndex < 0 || YChannelIndex > 3)
+}
+
+void UGeometryScriptLibrary_ListUtilityFunctions::ExtractUVListComponent(
+	FGeometryScriptUVList UVList, FGeometryScriptScalarList& ScalarList, int32 ComponentIndex)
+{
+	using namespace UE::ListUtilityFunctionsLocals;
+
+	ScalarList.Reset();
+	if (UVList.List.IsValid())
 	{
-		UE_LOG(LogGeometry, Warning, TEXT("ExtractColorListChannels: YChannelIndex is not in valid range 0-3"));
+		int ComponentMap[1]{ ComponentIndex };
+		ExtractArrayComponentsFromArrayHelper(*UVList.List, *ScalarList.List, ComponentMap);
 	}
-	if (ZChannelIndex < 0 || ZChannelIndex > 3)
+}
+
+void UGeometryScriptLibrary_ListUtilityFunctions::ExtractVectorListComponent(
+	FGeometryScriptVectorList VectorList, FGeometryScriptScalarList& ScalarList, int32 ComponentIndex)
+{
+	using namespace UE::ListUtilityFunctionsLocals;
+
+	ScalarList.Reset();
+	if (VectorList.List.IsValid())
 	{
-		UE_LOG(LogGeometry, Warning, TEXT("ExtractColorListChannels: ZChannelIndex is not in valid range 0-3"));
+		int ComponentMap[1]{ ComponentIndex };
+		ExtractArrayComponentsFromArrayHelper(*VectorList.List, *ScalarList.List, ComponentMap);
 	}
-	XChannelIndex = FMath::Clamp(XChannelIndex, 0, 3);
-	YChannelIndex = FMath::Clamp(YChannelIndex, 0, 3);
-	ZChannelIndex = FMath::Clamp(ZChannelIndex, 0, 3);
+}
 
-	VectorList.Reset();
-	if (ColorList.List.IsValid())
+void UGeometryScriptLibrary_ListUtilityFunctions::ExtractVectorListComponentsAsUVs(
+	FGeometryScriptVectorList VectorList, FGeometryScriptUVList& UVList, int32 UComponentIndex, int32 VComponentIndex)
+{
+	using namespace UE::ListUtilityFunctionsLocals;
+
+	UVList.Reset();
+	if (VectorList.List.IsValid())
 	{
-		TArray<FLinearColor>& Colors = *ColorList.List;
-		TArray<FVector>& Vectors = *VectorList.List;
-		int32 N = Colors.Num();
-		Vectors.SetNumUninitialized(N);
-		for (int32 k = 0; k < N; ++k)
-		{
-			FVector4d Vector4 = (FVector4d)(Colors[k]);
-			Vectors[k] = FVector(Vector4[XChannelIndex], Vector4[YChannelIndex], Vector4[ZChannelIndex]);
-		}
+		int ComponentMap[2]{ UComponentIndex, VComponentIndex };
+		ExtractArrayComponentsFromArrayHelper(*VectorList.List, *UVList.List, ComponentMap);
 	}
 }
 
+void UGeometryScriptLibrary_ListUtilityFunctions::SetColorListChannelFromScalars(
+	FGeometryScriptColorList& ColorList, FGeometryScriptScalarList ScalarList, int32 ColorChannel, FLinearColor DefaultColor)
+{
+	using namespace UE::ListUtilityFunctionsLocals;
+
+	if (!ValidateListsToSetAFromB(ColorList, ScalarList))
+	{
+		return;
+	}
+	int InToOutDimMapping[1]{ ColorChannel };
+	SetArrayComponentsFromArrayHelper<float>(*ColorList.List, *ScalarList.List, DefaultColor, InToOutDimMapping);
+}
 
 
+void UGeometryScriptLibrary_ListUtilityFunctions::SetColorListChannelsFromUVs(
+	FGeometryScriptColorList& ColorList, FGeometryScriptUVList UVList,
+	int32 UToChannelIndex, int32 VToChannelIndex, FLinearColor DefaultColor)
+{
+	using namespace UE::ListUtilityFunctionsLocals;
+
+	if (!ValidateListsToSetAFromB(ColorList, UVList))
+	{
+		return;
+	}
+	int InToOutDimMapping[2]{ UToChannelIndex, VToChannelIndex };
+	SetArrayComponentsFromArrayHelper<float>(*ColorList.List, *UVList.List, DefaultColor, InToOutDimMapping);
+}
+
+void UGeometryScriptLibrary_ListUtilityFunctions::SetColorListChannelsFromVectors(
+	FGeometryScriptColorList& ColorList, FGeometryScriptVectorList VectorList,
+	int32 XToChannelIndex, int32 YToChannelIndex, int32 ZToChannelIndex, FLinearColor DefaultColor)
+{
+	using namespace UE::ListUtilityFunctionsLocals;
+
+	if (!ValidateListsToSetAFromB(ColorList, VectorList))
+	{
+		return;
+	}
+	int InToOutDimMapping[3]{ XToChannelIndex, YToChannelIndex, ZToChannelIndex };
+	SetArrayComponentsFromArrayHelper<float>(*ColorList.List, *VectorList.List, DefaultColor, InToOutDimMapping);
+}
+
+void UGeometryScriptLibrary_ListUtilityFunctions::SetUVListComponentFromScalars(
+	FGeometryScriptUVList& UVList, FGeometryScriptScalarList ScalarList, int32 UVComponent, FVector2D DefaultUV)
+{
+	using namespace UE::ListUtilityFunctionsLocals;
+
+	if (!ValidateListsToSetAFromB(UVList, ScalarList))
+	{
+		return;
+	}
+	int InToOutDimMapping[1]{ UVComponent };
+	SetArrayComponentsFromArrayHelper<double>(*UVList.List, *ScalarList.List, DefaultUV, InToOutDimMapping);
+}
+
+void UGeometryScriptLibrary_ListUtilityFunctions::SetVectorListComponentFromScalars(
+	FGeometryScriptVectorList& VectorList, FGeometryScriptScalarList ScalarList, int32 VectorComponent, FVector DefaultVector)
+{
+	using namespace UE::ListUtilityFunctionsLocals;
+
+	if (!ValidateListsToSetAFromB(VectorList, ScalarList))
+	{
+		return;
+	}
+	int InToOutDimMapping[1]{ VectorComponent };
+	SetArrayComponentsFromArrayHelper<double>(*VectorList.List, *ScalarList.List, DefaultVector, InToOutDimMapping);
+}
+
+void UGeometryScriptLibrary_ListUtilityFunctions::SetVectorListComponentsFromUVs(
+	FGeometryScriptVectorList& VectorList, FGeometryScriptUVList UVList, int32 UToVectorComponent, int32 VToVectorComponent, FVector DefaultVector)
+{
+	using namespace UE::ListUtilityFunctionsLocals;
+
+	if (!ValidateListsToSetAFromB(VectorList, UVList))
+	{
+		return;
+	}
+	int InToOutDimMapping[2]{ UToVectorComponent, VToVectorComponent };
+	SetArrayComponentsFromArrayHelper<double>(*VectorList.List, *UVList.List, DefaultVector, InToOutDimMapping);
+}
+
 #undef LOCTEXT_NAMESPACE
diff --git a/Engine/Plugins/Runtime/GeometryScripting/Source/GeometryScriptingCore/Public/GeometryScript/ListUtilityFunctions.h b/Engine/Plugins/Runtime/GeometryScripting/Source/GeometryScriptingCore/Public/GeometryScript/ListUtilityFunctions.h
--- a/Engine/Plugins/Runtime/GeometryScripting/Source/GeometryScriptingCore/Public/GeometryScript/ListUtilityFunctions.h
+++ b/Engine/Plugins/Runtime/GeometryScripting/Source/GeometryScriptingCore/Public/GeometryScript/ListUtilityFunctions.h
@@ -311,7 +311,7 @@ public:
 	static UE_API void ClearColorList(UPARAM(ref) FGeometryScriptColorList& ColorList, FLinearColor ClearColor);
 
 	/**
-	* Populates a Scalar List with values that corresponds to the 0, 1, 2, or 3 channel of a Color List.
+	* Populates a Scalar List with values that correspond to the 0, 1, 2, or 3 channel of a Color List.
 	*/
 	UFUNCTION(BlueprintCallable, Category = "GeometryScript|ListUtils", meta=(ScriptMethod))
 	static UE_API void ExtractColorListChannel(FGeometryScriptColorList ColorList, FGeometryScriptScalarList& ScalarList, int32 ChannelIndex = 0);
@@ -322,6 +322,56 @@ public:
 	UFUNCTION(BlueprintCallable, Category = "GeometryScript|ListUtils", meta=(ScriptMethod))
 	static UE_API void ExtractColorListChannels(FGeometryScriptColorList ColorList, FGeometryScriptVectorList& VectorList, int32 XChannelIndex = 0, int32 YChannelIndex = 1, int32 ZChannelIndex = 2);
 
+	/**
+	 * Populates a Scalar List with values that correspond to the U or V component of a UV List.
+	 */
+	UFUNCTION(BlueprintCallable, Category = "GeometryScript|ListUtils", meta = (ScriptMethod))
+	static UE_API void ExtractUVListComponent(FGeometryScriptUVList UVList, FGeometryScriptScalarList& ScalarList, int32 ComponentIndex = 0);
+
+	/**
+	 * Populates a Scalar List with values that correspond to the X Y or Z component of a Vector List.
+	 */
+	UFUNCTION(BlueprintCallable, Category = "GeometryScript|ListUtils", meta = (ScriptMethod))
+	static UE_API void ExtractVectorListComponent(FGeometryScriptVectorList VectorList, FGeometryScriptScalarList& ScalarList, int32 ComponentIndex = 0);
+
+	/**
+	 * Populates a UV List with values that correspond to the X Y or Z component of a Vector List.
+	 */
+	UFUNCTION(BlueprintCallable, Category = "GeometryScript|ListUtils", meta = (ScriptMethod, DisplayName = "Extract Vector List Components As UVs"))
+	static UE_API void ExtractVectorListComponentsAsUVs(FGeometryScriptVectorList VectorList, FGeometryScriptUVList& UVList, 
+		int32 UComponentIndex = 0, int32 VComponentIndex = 1);
+
+	/** Sets the specified Color Channel of the Color List from the corresponding values in the Scalar List, resizing Color List as needed */
+	UFUNCTION(BlueprintCallable, Category = "GeometryScript|ListUtils", meta = (ScriptMethod))
+	static UE_API void SetColorListChannelFromScalars(UPARAM(ref) FGeometryScriptColorList& ColorList, FGeometryScriptScalarList ScalarList,
+		int32 ColorChannel = 0, FLinearColor DefaultColor = FLinearColor::Black);
+
+	/** Sets the specified Color Channels of the Color List from the corresponding values in the UV List, resizing Color List as needed */
+	UFUNCTION(BlueprintCallable, Category = "GeometryScript|ListUtils", meta = (ScriptMethod, DisplayName = "Set Color List Channels From UVs"))
+	static UE_API void SetColorListChannelsFromUVs(UPARAM(ref) FGeometryScriptColorList& ColorList, FGeometryScriptUVList UVList,
+		int32 UToChannelIndex = 0, int32 VToChannelIndex = 1, FLinearColor DefaultColor = FLinearColor::Black);
+
+	/** Sets the specified Color Channels of the Color List from the corresponding values in the Vector List, resizing Color List as needed */
+	UFUNCTION(BlueprintCallable, Category = "GeometryScript|ListUtils", meta = (ScriptMethod))
+	static UE_API void SetColorListChannelsFromVectors(UPARAM(ref) FGeometryScriptColorList& ColorList, FGeometryScriptVectorList VectorList,
+		int32 XToChannelIndex = 0, int32 YToChannelIndex = 1, int32 ZToChannelIndex = 2, FLinearColor DefaultColor = FLinearColor::Black);
+
+	/** Sets the specified Component of the UV List from the corresponding values in the Scalar List, resizing UV List as needed */
+	UFUNCTION(BlueprintCallable, Category = "GeometryScript|ListUtils", meta = (ScriptMethod))
+	static UE_API void SetUVListComponentFromScalars(UPARAM(ref) FGeometryScriptUVList& UVList, FGeometryScriptScalarList ScalarList, 
+		int32 UVComponent = 0, FVector2D DefaultUV = FVector2D::ZeroVector);
+
+	/** Sets the specified Component of the Vector List from the corresponding values in the Scalar List, resizing Vector List as needed */
+	UFUNCTION(BlueprintCallable, Category = "GeometryScript|ListUtils", meta = (ScriptMethod))
+	static UE_API void SetVectorListComponentFromScalars(UPARAM(ref) FGeometryScriptVectorList& VectorList, FGeometryScriptScalarList ScalarList, 
+		int32 VectorComponent = 0, FVector DefaultVector = FVector::ZeroVector);
+	
+	/** Sets the specified Components of the Vector List from the corresponding values in the UV List, resizing Vector List as needed */
+	UFUNCTION(BlueprintCallable, Category = "GeometryScript|ListUtils", meta = (ScriptMethod, DisplayName = "Set Vector List Components From UVs"))
+	static UE_API void SetVectorListComponentsFromUVs(UPARAM(ref) FGeometryScriptVectorList& VectorList, FGeometryScriptUVList UVList,
+		int32 UToVectorComponent = 0, int32 VToVectorComponent = 1, FVector DefaultVector = FVector::ZeroVector);
+
+
 };
 
 #undef UE_API
