minor comment fix for geometry script list functions (component -> components)

diff --git a/Engine/Plugins/Runtime/GeometryScripting/Source/GeometryScriptingCore/Public/GeometryScript/ListUtilityFunctions.h b/Engine/Plugins/Runtime/GeometryScripting/Source/GeometryScriptingCore/Public/GeometryScript/ListUtilityFunctions.h
--- a/Engine/Plugins/Runtime/GeometryScripting/Source/GeometryScriptingCore/Public/GeometryScript/ListUtilityFunctions.h
+++ b/Engine/Plugins/Runtime/GeometryScripting/Source/GeometryScriptingCore/Public/GeometryScript/ListUtilityFunctions.h
@@ -335,7 +335,7 @@ public:
 	static UE_API void ExtractVectorListComponent(FGeometryScriptVectorList VectorList, FGeometryScriptScalarList& ScalarList, int32 ComponentIndex = 0);
 
 	/**
-	 * Populates a UV List with values that correspond to the X Y or Z component of a Vector List.
+	 * Populates a UV List with values that correspond to the X Y or Z components of a Vector List.
 	 */
 	UFUNCTION(BlueprintCallable, Category = "GeometryScript|ListUtils", meta = (ScriptMethod, DisplayName = "Extract Vector List Components As UVs"))
 	static UE_API void ExtractVectorListComponentsAsUVs(FGeometryScriptVectorList VectorList, FGeometryScriptUVList& UVList, 
