Fix incorrect function name in geometry script VectorToScalar error message

diff --git a/Engine/Plugins/Runtime/GeometryScripting/Source/GeometryScriptingCore/Private/VectorMathFunctions.cpp b/Engine/Plugins/Runtime/GeometryScripting/Source/GeometryScriptingCore/Private/VectorMathFunctions.cpp
--- a/Engine/Plugins/Runtime/GeometryScripting/Source/GeometryScriptingCore/Private/VectorMathFunctions.cpp
+++ b/Engine/Plugins/Runtime/GeometryScripting/Source/GeometryScriptingCore/Private/VectorMathFunctions.cpp
@@ -291,7 +291,7 @@ FGeometryScriptScalarList UGeometryScriptLibrary_VectorMathFunctions::VectorToSc
 		[&](const FVector& Vector) { return ConstantX*Vector.X + ConstantY*Vector.Y + ConstantZ*Vector.Z; } );
 	if ( Result == VectorMathFunctionsLocal::EListOpResultCode::DifferentSizes )
 	{
-		UE_LOG(LogGeometry, Warning, TEXT("VectorLength: VectorList is empty"));
+		UE_LOG(LogGeometry, Warning, TEXT("VectorToScalar: VectorList is empty"));
 	}
 	return ResultList;
 }
