In geometry script asset creation methods, test that the paths for new assets are valid before attempting to use them

diff --git a/Engine/Plugins/Runtime/GeometryScripting/Source/GeometryScriptingEditor/Private/CreateNewAssetUtilityFunctions.cpp b/Engine/Plugins/Runtime/GeometryScripting/Source/GeometryScriptingEditor/Private/CreateNewAssetUtilityFunctions.cpp
--- a/Engine/Plugins/Runtime/GeometryScripting/Source/GeometryScriptingEditor/Private/CreateNewAssetUtilityFunctions.cpp
+++ b/Engine/Plugins/Runtime/GeometryScripting/Source/GeometryScriptingEditor/Private/CreateNewAssetUtilityFunctions.cpp
@@ -66,6 +66,26 @@ static bool CreateReferenceSkeletonFromMeshLods(const TArray<FDynamicMesh3>& Mes
 	return true;
 }
 
+namespace UE::CreateNewAssetUtilityFunctionsLocals
+{
+	static bool IsAssetPathValid(const FString& Path, bool bLogFailureReason = true)
+	{
+		FText BadPackageNameReason;
+		if (!FPackageName::IsValidTextForLongPackageName(Path, &BadPackageNameReason))
+		{
+			if (bLogFailureReason)
+			{
+				MakeScriptError(EGeometryScriptErrorType::InvalidInputs,
+					FText::Format(LOCTEXT("GeometryScript_BadPackageName", "Invalid Package Name '{0}' (Reason: {1})"),
+						FText::FromString(Path), BadPackageNameReason)
+				);
+			}
+			return false;
+		}
+		return true;
+	}
+}
+
 
 void UGeometryScriptLibrary_CreateNewAssetFunctions::CreateUniqueNewAssetPathName(
 	FString AssetFolderPath,
@@ -85,6 +105,11 @@ void UGeometryScriptLibrary_CreateNewAssetFunctions::CreateUniqueNewAssetPathNam
 		FString UUID = UE::Modeling::GenerateRandomShortHexString(Options.UniqueIDDigits);
 		UniqueAssetName = FString::Printf(TEXT("%s_%s"), *BaseAssetName, *UUID);
 		UniqueAssetPathAndName = UPackageTools::SanitizePackageName(FPaths::Combine(AssetFolderPath, UniqueAssetName));
+		if (!UE::CreateNewAssetUtilityFunctionsLocals::IsAssetPathValid(UniqueAssetPathAndName))
+		{
+			// path has invalid format; early out here b/c changing the random ID won't fix this
+			return;
+		}
 
 		// if asset does not exist at this path, we can use it
 		FAssetData AssetData = AssetRegistryModule.Get().GetAssetByObjectPath(FSoftObjectPath(UniqueAssetPathAndName));
@@ -231,12 +256,16 @@ UStaticMesh* UGeometryScriptLibrary_CreateNewAssetFunctions::CreateNewStaticMesh
 			return nullptr;
 		}
 	}
+	AssetPathAndName = UPackageTools::SanitizePackageName(AssetPathAndName);
+	if (!UE::CreateNewAssetUtilityFunctionsLocals::IsAssetPathValid(AssetPathAndName))
+	{
+		return nullptr;
+	}
 	// todo: other safety checks
 
 	GEditor->BeginTransaction(LOCTEXT("CreateNewStaticMeshAssetFromMeshLODs_Transaction", "Create StaticMesh"));
 
 	UE::AssetUtils::FStaticMeshAssetOptions AssetOptions;
-	AssetPathAndName = UPackageTools::SanitizePackageName(AssetPathAndName);
 	AssetOptions.NewAssetPath = AssetPathAndName;
 
 	AssetOptions.NumSourceModels = FromDynamicMeshLODs.Num();
@@ -349,13 +378,16 @@ USkeletalMesh* UGeometryScriptLibrary_CreateNewAssetFunctions::CreateNewSkeletal
 			return nullptr;
 		}
 	}
-	
+	AssetPathAndName = UPackageTools::SanitizePackageName(AssetPathAndName);
+	if (!UE::CreateNewAssetUtilityFunctionsLocals::IsAssetPathValid(AssetPathAndName))
+	{
+		return nullptr;
+	}
 	// todo: other safety checks
 
 	GEditor->BeginTransaction(LOCTEXT("CreateNewSkeletalMeshAssetFromMeshLODs_Transaction", "Create SkeletalMesh"));
 
 	FSkeletalMeshAssetOptions AssetOptions;
-	AssetPathAndName = UPackageTools::SanitizePackageName(AssetPathAndName);
 	AssetOptions.NewAssetPath = AssetPathAndName;
 	AssetOptions.Skeleton = InSkeleton;
 
@@ -556,9 +588,13 @@ UTexture2D* UGeometryScriptLibrary_CreateNewAssetFunctions::CreateNewTexture2DAs
 		UE::Geometry::AppendError(Debug, EGeometryScriptErrorType::InvalidInputs, LOCTEXT("CreateNewTexture2DAsset_InvalidInput1", "CreateNewTexture2DAsset: FromTexture is Null"));
 		return nullptr;
 	}
+	AssetPathAndName = UPackageTools::SanitizePackageName(AssetPathAndName);
+	if (!UE::CreateNewAssetUtilityFunctionsLocals::IsAssetPathValid(AssetPathAndName))
+	{
+		return nullptr;
+	}
 
 	UE::AssetUtils::FTexture2DAssetOptions AssetOptions;
-	AssetPathAndName = UPackageTools::SanitizePackageName(AssetPathAndName);
 	AssetOptions.NewAssetPath = AssetPathAndName;
 	AssetOptions.bOverwriteIfExists = Options.bOverwriteIfExists;
 
