Change geometry script user-facing term 'components' to instead to 'connected islands' since components is too overloaded in UE
+ add some keywords to help the method find in search if people still look for components in search

diff --git a/Engine/Plugins/Runtime/GeometryScripting/Source/GeometryScriptingCore/Public/GeometryScript/MeshPolygroupFunctions.h b/Engine/Plugins/Runtime/GeometryScripting/Source/GeometryScriptingCore/Public/GeometryScript/MeshPolygroupFunctions.h
--- a/Engine/Plugins/Runtime/GeometryScripting/Source/GeometryScriptingCore/Public/GeometryScript/MeshPolygroupFunctions.h
+++ b/Engine/Plugins/Runtime/GeometryScripting/Source/GeometryScriptingCore/Public/GeometryScript/MeshPolygroupFunctions.h
@@ -110,10 +110,11 @@ public:
 		UGeometryScriptDebug* Debug = nullptr );
 
 	/**
-	* Creates and assigns a new PolyGroup for each disconnected component of a Mesh. Regions of a mesh are disconnected they do not have a triangle in common.
+	* Creates and assigns a new PolyGroup for each disconnected island (component) of a Mesh. Regions of a mesh are disconnected they do not have a triangle in common.
 	* Note, this will have no effect if the Group Layer does not exist.
 	*/
-	UFUNCTION(BlueprintCallable, Category = "GeometryScript|PolyGroups", meta = (ScriptMethod, HidePin = "Debug", DisplayName = "Convert Components To PolyGroups"))
+	UFUNCTION(BlueprintCallable, Category = "GeometryScript|PolyGroups", meta = (ScriptMethod, HidePin = "Debug", KeyWords = "Components",
+		DisplayName = "Convert Connected Islands To PolyGroups"))
 	static UE_API UPARAM(DisplayName = "Target Mesh") UDynamicMesh* 
 	ConvertComponentsToPolygroups( 
 		UDynamicMesh* TargetMesh, 
@@ -252,7 +253,7 @@ public:
 		bool& bIsValid);
 
 	/**
-	 * Create list of all triangles with the given PolyGroup ID in the given GroupLayer (not necessarily a single connected-component)
+	 * Create list of all triangles with the given PolyGroup ID in the given GroupLayer (not necessarily a single connected-component / island)
 	 */
 	UFUNCTION(BlueprintCallable, Category = "GeometryScript|PolyGroups", meta = (ScriptMethod, DisplayName = "Get Triangles In PolyGroup"))
 	static UE_API UPARAM(DisplayName = "Target Mesh") UDynamicMesh* 
