Expose target edge length QEM and cluster-based simplification methods in geometry script

diff --git a/Engine/Plugins/Runtime/GeometryScripting/Source/GeometryScriptingCore/Private/MeshSimplifyFunctions.cpp b/Engine/Plugins/Runtime/GeometryScripting/Source/GeometryScriptingCore/Private/MeshSimplifyFunctions.cpp
--- a/Engine/Plugins/Runtime/GeometryScripting/Source/GeometryScriptingCore/Private/MeshSimplifyFunctions.cpp
+++ b/Engine/Plugins/Runtime/GeometryScripting/Source/GeometryScriptingCore/Private/MeshSimplifyFunctions.cpp
@@ -12,6 +12,7 @@
 #include "Polygroups/PolygroupSet.h"
 #include "GroupTopology.h"
 #include "Operations/PolygroupRemesh.h"
+#include "Operations/MeshClusterSimplifier.h"
 #include "ConstrainedDelaunay2.h"
 
 #if WITH_EDITOR
@@ -141,73 +142,110 @@ namespace UE::PrivateSimplifyHelper
 		return false;
 #endif
 	}
-}
-
 
-template<typename SimplificationType>
-void DoSimplifyMesh(
-	FDynamicMesh3& EditMesh, 
-	FGeometryScriptSimplifyMeshOptions Options,
-	bool bTargetIsTriCount,
-	int32 TargetCount,
-	FMeshProjectionTarget* ProjectionTarget = nullptr,
-	double GeometricTolerance = -1 // Note: Values < 0 will skip using GeometricTolerance
-)
-{
-	SimplificationType Simplifier(&EditMesh);
 
-	Simplifier.ProjectionMode = SimplificationType::ETargetProjectionMode::NoProjection;
-	if (ProjectionTarget != nullptr)
+	enum class ESimplifyTargetMode
 	{
-		Simplifier.SetProjectionTarget(ProjectionTarget);
-	}
-
-	Simplifier.DEBUG_CHECK_LEVEL = 0;
-	Simplifier.bRetainQuadricMemory = Options.bRetainQuadricMemory; 
-	Simplifier.bAllowSeamCollapse = Options.bAllowSeamCollapse;
-	if (Options.bAllowSeamCollapse)
+		TriangleCount,
+		VertexCount,
+		EdgeLength
+	};
+
+	template<typename SimplificationType>
+	static void DoSimplifyMesh(
+		FDynamicMesh3& EditMesh,
+		FGeometryScriptSimplifyMeshOptions Options,
+		ESimplifyTargetMode SimplifyTargetMode,
+		int32 TargetCount,
+		double TargetMinEdgeLength = 0.,
+		FMeshProjectionTarget* ProjectionTarget = nullptr,
+		double GeometricTolerance = -1 // Note: Values < 0 will skip using GeometricTolerance
+	)
 	{
-		Simplifier.SetEdgeFlipTolerance(1.e-5);
-		if (EditMesh.HasAttributes())
+		SimplificationType Simplifier(&EditMesh);
+
+		Simplifier.ProjectionMode = SimplificationType::ETargetProjectionMode::NoProjection;
+		if (ProjectionTarget != nullptr)
 		{
-			EditMesh.Attributes()->SplitAllBowties();	// eliminate any bowties that might have formed on attribute seams.
+			Simplifier.SetProjectionTarget(ProjectionTarget);
 		}
-	}
 
-	// do these flags matter here since we are not flipping??
-	EEdgeRefineFlags MeshBoundaryConstraints = EEdgeRefineFlags::NoFlip;
-	EEdgeRefineFlags GroupBorderConstraints = EEdgeRefineFlags::NoConstraint;
-	EEdgeRefineFlags MaterialBorderConstraints = EEdgeRefineFlags::NoConstraint;
+		Simplifier.DEBUG_CHECK_LEVEL = 0;
+		Simplifier.bRetainQuadricMemory = Options.bRetainQuadricMemory;
+		Simplifier.bAllowSeamCollapse = Options.bAllowSeamCollapse;
+		if (Options.bAllowSeamCollapse)
+		{
+			Simplifier.SetEdgeFlipTolerance(1.e-5);
+			if (EditMesh.HasAttributes())
+			{
+				EditMesh.Attributes()->SplitAllBowties();	// eliminate any bowties that might have formed on attribute seams.
+			}
+		}
 
-	FMeshConstraints Constraints;
-	FMeshConstraintsUtil::ConstrainAllBoundariesAndSeams(Constraints, EditMesh,
-		MeshBoundaryConstraints, GroupBorderConstraints, MaterialBorderConstraints,
-		Options.bAllowSeamSplits, Options.bAllowSeamSmoothing, Options.bAllowSeamCollapse);
-	Simplifier.SetExternalConstraints(MoveTemp(Constraints));
+		// do these flags matter here since we are not flipping??
+		EEdgeRefineFlags MeshBoundaryConstraints = EEdgeRefineFlags::NoFlip;
+		EEdgeRefineFlags GroupBorderConstraints = EEdgeRefineFlags::NoConstraint;
+		EEdgeRefineFlags MaterialBorderConstraints = EEdgeRefineFlags::NoConstraint;
 
-	if (Options.bPreserveVertexPositions)
-	{
-		Simplifier.CollapseMode = SimplificationType::ESimplificationCollapseModes::MinimalExistingVertexError;
-	}
+		FMeshConstraints Constraints;
+		FMeshConstraintsUtil::ConstrainAllBoundariesAndSeams(Constraints, EditMesh,
+			MeshBoundaryConstraints, GroupBorderConstraints, MaterialBorderConstraints,
+			Options.bAllowSeamSplits, Options.bAllowSeamSmoothing, Options.bAllowSeamCollapse);
+		Simplifier.SetExternalConstraints(MoveTemp(Constraints));
 
-	if ( ProjectionTarget != nullptr && GeometricTolerance >= 0)
-	{
-		Simplifier.GeometricErrorConstraint = SimplificationType::EGeometricErrorCriteria::PredictedPointToProjectionTarget;
-		Simplifier.GeometricErrorTolerance = GeometricTolerance;
-	}
+		if (Options.bPreserveVertexPositions)
+		{
+			Simplifier.CollapseMode = SimplificationType::ESimplificationCollapseModes::MinimalExistingVertexError;
+		}
 
-	if (bTargetIsTriCount)
-	{
-		Simplifier.SimplifyToTriangleCount( FMath::Max(1,TargetCount) );
-	}
-	else
-	{
-		Simplifier.SimplifyToVertexCount( FMath::Max(1,TargetCount) );
+		if (ProjectionTarget != nullptr && GeometricTolerance >= 0)
+		{
+			Simplifier.GeometricErrorConstraint = SimplificationType::EGeometricErrorCriteria::PredictedPointToProjectionTarget;
+			Simplifier.GeometricErrorTolerance = GeometricTolerance;
+		}
+
+		if (SimplifyTargetMode == ESimplifyTargetMode::TriangleCount)
+		{
+			Simplifier.SimplifyToTriangleCount(FMath::Max(1, TargetCount));
+		}
+		else if (SimplifyTargetMode == ESimplifyTargetMode::VertexCount)
+		{
+			Simplifier.SimplifyToVertexCount(FMath::Max(1, TargetCount));
+		}
+		else // SimplifyTargetMode == ESimplifyTargetMode::EdgeLength
+		{
+			Simplifier.SimplifyToEdgeLength(TargetMinEdgeLength);
+		}
+
+		if (Options.bAutoCompact)
+		{
+			EditMesh.CompactInPlace();
+		}
 	}
 
-	if (Options.bAutoCompact)
+	static void DoSimplifyMesh(
+		EGeometryScriptRemoveMeshSimplificationType SimplificationType,
+		FDynamicMesh3& EditMesh,
+		FGeometryScriptSimplifyMeshOptions Options,
+		ESimplifyTargetMode SimplifyTargetMode,
+		int32 TargetCount,
+		double TargetMinEdgeLength = 0.,
+		FMeshProjectionTarget* ProjectionTarget = nullptr,
+		double GeometricTolerance = -1 // Note: Values < 0 will skip using GeometricTolerance
+	)
 	{
-		EditMesh.CompactInPlace();
+		if (Options.Method == EGeometryScriptRemoveMeshSimplificationType::AttributeAware)
+		{
+			DoSimplifyMesh<FAttrMeshSimplification>(EditMesh, Options, SimplifyTargetMode, TargetCount, TargetMinEdgeLength, ProjectionTarget, GeometricTolerance);
+		}
+		else if (Options.Method == EGeometryScriptRemoveMeshSimplificationType::VolumePreserving)
+		{
+			DoSimplifyMesh<FVolPresMeshSimplification>(EditMesh, Options, SimplifyTargetMode, TargetCount, TargetMinEdgeLength, ProjectionTarget, GeometricTolerance);
+		}
+		else // EGeometryScriptRemoveMeshSimplificationType::StandardQEM
+		{
+			DoSimplifyMesh<FQEMSimplification>(EditMesh, Options, SimplifyTargetMode, TargetCount, TargetMinEdgeLength, ProjectionTarget, GeometricTolerance);
+		}
 	}
 }
 
@@ -269,19 +307,8 @@ UDynamicMesh* UGeometryScriptLibrary_MeshSimplifyFunctions::ApplySimplifyToTrian
 
 	TargetMesh->EditMesh([&](FDynamicMesh3& EditMesh) 
 	{
-		if (Options.Method == EGeometryScriptRemoveMeshSimplificationType::AttributeAware)
-		{
-			DoSimplifyMesh<FAttrMeshSimplification>(EditMesh, Options, true, TriangleCount);
-		}
-		else if (Options.Method == EGeometryScriptRemoveMeshSimplificationType::VolumePreserving)
-		{
-			DoSimplifyMesh<FVolPresMeshSimplification>(EditMesh, Options, true, TriangleCount);
-		}
-		else
-		{
-			DoSimplifyMesh<FQEMSimplification>(EditMesh, Options, true, TriangleCount);
-		}
-
+		using namespace UE::PrivateSimplifyHelper;
+		DoSimplifyMesh(Options.Method, EditMesh, Options, ESimplifyTargetMode::TriangleCount, TriangleCount);
 	}, EDynamicMeshChangeType::GeneralEdit, EDynamicMeshAttributeChangeFlags::Unknown, false);
 
 	return TargetMesh;
@@ -304,19 +331,8 @@ UDynamicMesh* UGeometryScriptLibrary_MeshSimplifyFunctions::ApplySimplifyToVerte
 
 	TargetMesh->EditMesh([&](FDynamicMesh3& EditMesh) 
 	{
-		if (Options.Method == EGeometryScriptRemoveMeshSimplificationType::AttributeAware)
-		{
-			DoSimplifyMesh<FAttrMeshSimplification>(EditMesh, Options, false, VertexCount);
-		}
-		else if (Options.Method == EGeometryScriptRemoveMeshSimplificationType::VolumePreserving)
-		{
-			DoSimplifyMesh<FVolPresMeshSimplification>(EditMesh, Options, false, VertexCount);
-		}
-		else
-		{
-			DoSimplifyMesh<FQEMSimplification>(EditMesh, Options, false, VertexCount);
-		}
-
+		using namespace UE::PrivateSimplifyHelper;
+		DoSimplifyMesh(Options.Method, EditMesh, Options, ESimplifyTargetMode::VertexCount, VertexCount);
 	}, EDynamicMeshChangeType::GeneralEdit, EDynamicMeshAttributeChangeFlags::Unknown, false);
 
 	return TargetMesh;
@@ -344,19 +360,95 @@ UDynamicMesh* UGeometryScriptLibrary_MeshSimplifyFunctions::ApplySimplifyToToler
 		FMeshProjectionTarget ProjTarget(&TempCopy, &Spatial);
 		float UseTolerance = FMath::Max(0.0, Tolerance);
 
-		if (Options.Method == EGeometryScriptRemoveMeshSimplificationType::AttributeAware)
-		{
-			DoSimplifyMesh<FAttrMeshSimplification>(EditMesh, Options, true, 1, &ProjTarget, UseTolerance);
-		}
-		else if (Options.Method == EGeometryScriptRemoveMeshSimplificationType::VolumePreserving)
-		{
-			DoSimplifyMesh<FVolPresMeshSimplification>(EditMesh, Options, true, 1, &ProjTarget, UseTolerance);
-		}
-		else
+		using namespace UE::PrivateSimplifyHelper;
+		DoSimplifyMesh(Options.Method, EditMesh, Options, ESimplifyTargetMode::TriangleCount, 1, -1., &ProjTarget, UseTolerance);
+	}, EDynamicMeshChangeType::GeneralEdit, EDynamicMeshAttributeChangeFlags::Unknown, false);
+
+	return TargetMesh;
+}
+
+UDynamicMesh* UGeometryScriptLibrary_MeshSimplifyFunctions::ApplySimplifyToEdgeLength(
+	UDynamicMesh* TargetMesh,
+	double EdgeLength,
+	FGeometryScriptSimplifyMeshOptions Options,
+	UGeometryScriptDebug* Debug)
+{
+	if (TargetMesh == nullptr)
+	{
+		UE::Geometry::AppendError(Debug, EGeometryScriptErrorType::InvalidInputs, LOCTEXT("ApplySimplifyToVertexCount_InvalidInput", "ApplySimplifyToVertexCount: TargetMesh is Null"));
+		return TargetMesh;
+	}
+
+	TargetMesh->EditMesh([&EdgeLength, &Options](FDynamicMesh3& EditMesh) 
+	{
+		// Clamp target length to non-negative
+		double UseEdgeLength = FMath::Max(0.f, EdgeLength);
+
+		using namespace UE::PrivateSimplifyHelper;
+		DoSimplifyMesh(Options.Method, EditMesh, Options, ESimplifyTargetMode::EdgeLength, 1, UseEdgeLength);
+	}, EDynamicMeshChangeType::GeneralEdit, EDynamicMeshAttributeChangeFlags::Unknown, false);
+
+	return TargetMesh;
+}
+
+UDynamicMesh* UGeometryScriptLibrary_MeshSimplifyFunctions::ApplyClusterSimplifyToEdgeLength(
+	UDynamicMesh* TargetMesh,
+	double EdgeLength,
+	FGeometryScriptClusterSimplifyMeshOptions Options,
+	UGeometryScriptDebug* Debug)
+{
+	if (TargetMesh == nullptr)
+	{
+		UE::Geometry::AppendError(Debug, EGeometryScriptErrorType::InvalidInputs, LOCTEXT("ApplySimplifyToVertexCount_InvalidInput", "ApplySimplifyToVertexCount: TargetMesh is Null"));
+		return TargetMesh;
+	}
+
+	TargetMesh->EditMesh([&EdgeLength, &Options](FDynamicMesh3& EditMesh) 
+	{
+		// Clamp target length to non-negative
+		double UseEdgeLength = FMath::Max(0.f, EdgeLength);
+
+		MeshClusterSimplify::FSimplifyOptions SimplifyOptions;
+		SimplifyOptions.TargetEdgeLength = UseEdgeLength;
+		SimplifyOptions.FixBoundaryAngleTolerance = Options.FixBoundaryAngleTolerance;
+
+		auto ConvertConstraintLevel = [](EGeometryScriptClusterSimplifyConstraintLevel ConstraintLevel) -> MeshClusterSimplify::FSimplifyOptions::EConstraintLevel
+			{
+				switch (ConstraintLevel)
+				{
+				default:
+					ensure(false);
+				case EGeometryScriptClusterSimplifyConstraintLevel::Free:
+					return MeshClusterSimplify::FSimplifyOptions::EConstraintLevel::Free;
+				case EGeometryScriptClusterSimplifyConstraintLevel::Constrained:
+					return MeshClusterSimplify::FSimplifyOptions::EConstraintLevel::Constrained;
+				case EGeometryScriptClusterSimplifyConstraintLevel::Fixed:
+					return MeshClusterSimplify::FSimplifyOptions::EConstraintLevel::Fixed;
+				}
+			};
+		SimplifyOptions.bTransferGroups = true;
+
+		SimplifyOptions.PreserveEdges.Boundary = ConvertConstraintLevel(Options.EdgeConstraints.MeshBoundary);
+		SimplifyOptions.PreserveEdges.PolyGroup = ConvertConstraintLevel(Options.EdgeConstraints.GroupBoundary);
+		SimplifyOptions.PreserveEdges.Material = ConvertConstraintLevel(Options.EdgeConstraints.MaterialBoundary);
+		SimplifyOptions.PreserveEdges.UVSeam = ConvertConstraintLevel(Options.EdgeConstraints.UVSeam);
+		SimplifyOptions.PreserveEdges.ColorSeam = ConvertConstraintLevel(Options.EdgeConstraints.ColorSeam);
+		SimplifyOptions.PreserveEdges.NormalSeam = ConvertConstraintLevel(Options.EdgeConstraints.NormalSeam);
+
+		// Note: Currently not exposing tangent seam constraints; they usually should be left free
+		SimplifyOptions.PreserveEdges.TangentSeam = MeshClusterSimplify::FSimplifyOptions::EConstraintLevel::Free;
+
+		if (Options.bDiscardAttributes)
 		{
-			DoSimplifyMesh<FQEMSimplification>(EditMesh, Options, true, 1, &ProjTarget, UseTolerance);
+			// if discarding attributes, also discard constraints from the attribute layer
+			SimplifyOptions.PreserveEdges.SetSeamConstraints(MeshClusterSimplify::FSimplifyOptions::EConstraintLevel::Free);
+			SimplifyOptions.PreserveEdges.Material = MeshClusterSimplify::FSimplifyOptions::EConstraintLevel::Free;
+			SimplifyOptions.bTransferAttributes = false;
 		}
-
+		
+		FDynamicMesh3 ResultMesh;
+		MeshClusterSimplify::Simplify(EditMesh, ResultMesh, SimplifyOptions);
+		EditMesh = ResultMesh;
 	}, EDynamicMeshChangeType::GeneralEdit, EDynamicMeshAttributeChangeFlags::Unknown, false);
 
 	return TargetMesh;
diff --git a/Engine/Plugins/Runtime/GeometryScripting/Source/GeometryScriptingCore/Public/GeometryScript/MeshSimplifyFunctions.h b/Engine/Plugins/Runtime/GeometryScripting/Source/GeometryScriptingCore/Public/GeometryScript/MeshSimplifyFunctions.h
--- a/Engine/Plugins/Runtime/GeometryScripting/Source/GeometryScriptingCore/Public/GeometryScript/MeshSimplifyFunctions.h
+++ b/Engine/Plugins/Runtime/GeometryScripting/Source/GeometryScriptingCore/Public/GeometryScript/MeshSimplifyFunctions.h
@@ -78,6 +78,62 @@ public:
 	bool bAutoCompact = true;
 };
 
+UENUM(BlueprintType)
+enum class EGeometryScriptClusterSimplifyConstraintLevel : uint8
+{
+	// Edge will be kept as-is in final output
+	Fixed,
+	// Attempt to preserve 'paths' of constrained edges, but may simplify along the path
+	// Vertices where paths converge will be kept in final output
+	Constrained,
+	// Edge can be simplified without constraint
+	Free
+};
+
+USTRUCT(BlueprintType)
+struct FGeometryScriptClusterSimplifyEdgeConstraintOptions
+{
+	GENERATED_BODY()
+public:
+
+	UPROPERTY(EditAnywhere, BlueprintReadWrite, Category = ConstraintOptions)
+	EGeometryScriptClusterSimplifyConstraintLevel MeshBoundary = EGeometryScriptClusterSimplifyConstraintLevel::Constrained;
+
+	UPROPERTY(EditAnywhere, BlueprintReadWrite, Category = ConstraintOptions)
+	EGeometryScriptClusterSimplifyConstraintLevel GroupBoundary = EGeometryScriptClusterSimplifyConstraintLevel::Constrained;
+
+	UPROPERTY(EditAnywhere, BlueprintReadWrite, Category = ConstraintOptions)
+	EGeometryScriptClusterSimplifyConstraintLevel MaterialBoundary = EGeometryScriptClusterSimplifyConstraintLevel::Constrained;
+
+	UPROPERTY(EditAnywhere, BlueprintReadWrite, Category = ConstraintOptions)
+	EGeometryScriptClusterSimplifyConstraintLevel UVSeam = EGeometryScriptClusterSimplifyConstraintLevel::Constrained;
+
+	UPROPERTY(EditAnywhere, BlueprintReadWrite, Category = ConstraintOptions)
+	EGeometryScriptClusterSimplifyConstraintLevel ColorSeam = EGeometryScriptClusterSimplifyConstraintLevel::Constrained;
+
+	UPROPERTY(EditAnywhere, BlueprintReadWrite, Category = ConstraintOptions)
+	EGeometryScriptClusterSimplifyConstraintLevel NormalSeam = EGeometryScriptClusterSimplifyConstraintLevel::Constrained;
+};
+
+USTRUCT(BlueprintType)
+struct FGeometryScriptClusterSimplifyMeshOptions
+{
+	GENERATED_BODY()
+public:
+
+	// Whether to discard attributes (materials, UV/normal/tangent/color attributes)
+	// Note that attribute seam edges will not be constrained if attributes are discarded
+	UPROPERTY(EditAnywhere, BlueprintReadWrite, Category = Options)
+	bool bDiscardAttributes = false;
+
+	// If > 0, boundary vertices w/ incident boundary edge angle greater than this (in degrees) will be kept in the output
+	UPROPERTY(EditAnywhere, BlueprintReadWrite, Category = Options)
+	double FixBoundaryAngleTolerance = 45;
+
+	// Options to constrain simplification of different mesh edge types, to optionally help preserve mesh boundaries and seams
+	UPROPERTY(EditAnywhere, BlueprintReadWrite, Category = Options)
+	FGeometryScriptClusterSimplifyEdgeConstraintOptions EdgeConstraints;
+};
 
 
 
@@ -164,6 +220,30 @@ public:
 		FGeometryScriptSimplifyMeshOptions Options,
 		UGeometryScriptDebug* Debug = nullptr);
 
+	/**
+	 * Simplifies the mesh to a target edge length, using error-based edge collapse.
+	 * Note: Not intended to create uniform edge lengths. Result may have edge lengths much larger than the target value.
+	 */
+	UFUNCTION(BlueprintCallable, Category = "GeometryScript|Simplification", meta=(ScriptMethod, HidePin = "Debug"))
+	static UE_API UPARAM(DisplayName = "Target Mesh") UDynamicMesh* 
+	ApplySimplifyToEdgeLength(  
+		UDynamicMesh* TargetMesh, 
+		double EdgeLength,
+		FGeometryScriptSimplifyMeshOptions Options,
+		UGeometryScriptDebug* Debug = nullptr);
+
+	/**
+	 * Simplifies the mesh to a target edge length, using a fast cluster-based method (not error-based).
+	 * Note: Lengths are approximated via 'graph distance' on the original mesh; will not exactly match the target.
+	 */
+	UFUNCTION(BlueprintCallable, Category = "GeometryScript|Simplification", meta = (ScriptMethod, HidePin = "Debug"))
+	static UE_API UPARAM(DisplayName = "Target Mesh") UDynamicMesh*
+	ApplyClusterSimplifyToEdgeLength(
+		UDynamicMesh* TargetMesh,
+		double EdgeLength,
+		FGeometryScriptClusterSimplifyMeshOptions Options,
+		UGeometryScriptDebug* Debug = nullptr);
+
 };
 
 #undef UE_API
