Change GeometryScript displaynames to use the term 'Islands' instead of 'Components' more consistently + add more keywords to help people find these nodes

diff --git a/Engine/Plugins/Runtime/GeometryScripting/Source/GeometryScriptingCore/Public/GeometryScript/MeshDecompositionFunctions.h b/Engine/Plugins/Runtime/GeometryScripting/Source/GeometryScriptingCore/Public/GeometryScript/MeshDecompositionFunctions.h
--- a/Engine/Plugins/Runtime/GeometryScripting/Source/GeometryScriptingCore/Public/GeometryScript/MeshDecompositionFunctions.h
+++ b/Engine/Plugins/Runtime/GeometryScripting/Source/GeometryScriptingCore/Public/GeometryScript/MeshDecompositionFunctions.h
@@ -21,12 +21,13 @@ class UGeometryScriptLibrary_MeshDecompositionFunctions : public UBlueprintFunct
 public:
 
 	/**
-	 * Create a new Mesh for each Connected Component of TargetMesh.
+	 * Create a new Mesh for each Connected Island (Component) of TargetMesh.
 	 * New meshes are drawn from MeshPool if it is provided, otherwise new UDynamicMesh instances are allocated
 	 * @param ComponentMeshes New List of meshes is returned here
 	 * @param MeshPool New meshes in ComponentMeshes output list are allocated from this pool if it is provided (highly recommended!!)
 	 */
-	UFUNCTION(BlueprintCallable, Category = "GeometryScript|Decomposition", meta = (ScriptMethod, HidePin = "Debug", Keywords = "Connected Islands"))
+	UFUNCTION(BlueprintCallable, Category = "GeometryScript|Decomposition", meta = (ScriptMethod, HidePin = "Debug", 
+		DisplayName = "Split Mesh by Connected Islands", Keywords = "Connected Components Parts"))
 	static UE_API UPARAM(DisplayName = "Target Mesh") UDynamicMesh* 
 	SplitMeshByComponents(  
 		UDynamicMesh* TargetMesh, 
@@ -41,7 +42,7 @@ public:
 	 * @param MeshPool New meshes in ComponentMeshes output list are allocated from this pool if it is provided (highly recommended!!)
 	 * @param ConnectVerticesThreshold Vertices closer than this distance will be classified as part of the same component, even if they aren't connected by the mesh triangulation
 	 */
-	UFUNCTION(BlueprintCallable, Category = "GeometryScript|Decomposition", meta = (ScriptMethod, HidePin = "Debug", Keywords = "Islands Connected Components"))
+	UFUNCTION(BlueprintCallable, Category = "GeometryScript|Decomposition", meta = (ScriptMethod, HidePin = "Debug", Keywords = "Islands Connected Components Parts"))
 	static UE_API UPARAM(DisplayName = "Target Mesh") UDynamicMesh*
 	SplitMeshByVertexOverlap(
 		UDynamicMesh* TargetMesh,
diff --git a/Engine/Plugins/Runtime/GeometryScripting/Source/GeometryScriptingCore/Public/GeometryScript/MeshRepairFunctions.h b/Engine/Plugins/Runtime/GeometryScripting/Source/GeometryScriptingCore/Public/GeometryScript/MeshRepairFunctions.h
--- a/Engine/Plugins/Runtime/GeometryScripting/Source/GeometryScriptingCore/Public/GeometryScript/MeshRepairFunctions.h
+++ b/Engine/Plugins/Runtime/GeometryScripting/Source/GeometryScriptingCore/Public/GeometryScript/MeshRepairFunctions.h
@@ -283,9 +283,10 @@ public:
 		UGeometryScriptDebug* Debug = nullptr);
 
 	/*
-	* Removes connected components of the mesh that have a volume, area, or triangle count below a threshold as specified by the Options.
+	* Removes connected islands (components) of the mesh that have a volume, area, or triangle count below a threshold as specified by the Options.
 	*/
-	UFUNCTION(BlueprintCallable, Category = "GeometryScript|Repair", meta=(ScriptMethod, HidePin = "Debug"))
+	UFUNCTION(BlueprintCallable, Category = "GeometryScript|Repair", meta=(ScriptMethod, HidePin = "Debug",
+		DisplayName = "Remove Small Connected Islands", KeyWords = "Delete Components Parts"))
 	static UE_API UPARAM(DisplayName = "Target Mesh") UDynamicMesh* 
 	RemoveSmallComponents(  
 		UDynamicMesh* TargetMesh, 
