Changing Components -> Islands in geometry script

diff --git a/Engine/Plugins/Runtime/GeometryScripting/Source/GeometryScriptingCore/Public/GeometryScript/MeshQueryFunctions.h b/Engine/Plugins/Runtime/GeometryScripting/Source/GeometryScriptingCore/Public/GeometryScript/MeshQueryFunctions.h
--- a/Engine/Plugins/Runtime/GeometryScripting/Source/GeometryScriptingCore/Public/GeometryScript/MeshQueryFunctions.h
+++ b/Engine/Plugins/Runtime/GeometryScripting/Source/GeometryScriptingCore/Public/GeometryScript/MeshQueryFunctions.h
@@ -74,9 +74,9 @@ public:
 	static UE_API UPARAM(DisplayName = "Num Edges") int32 GetNumOpenBorderEdges( UDynamicMesh* TargetMesh );
 
 	/**
-	* Returns the number of separate connected components in the mesh, such as "triangle patches" internally connected by shared edges.
+	* Returns the number of separate connected islands (components) in the mesh, such as "triangle patches" internally connected by shared edges.
 	*/
-	UFUNCTION(BlueprintPure, Category = "GeometryScript|MeshQueries", meta=(ScriptMethod))
+	UFUNCTION(BlueprintPure, Category = "GeometryScript|MeshQueries", meta = (ScriptMethod, DisplayName = "Get Num Connected Islands", KeyWords = "Components Parts"))
 	static UE_API UPARAM(DisplayName = "Num Edges") int32 GetNumConnectedComponents( UDynamicMesh* TargetMesh );
 
 	// UDynamicMesh already has this function
