fix incorrect display name for 'get num connected islands' geometry script method

diff --git a/Engine/Plugins/Runtime/GeometryScripting/Source/GeometryScriptingCore/Public/GeometryScript/MeshQueryFunctions.h b/Engine/Plugins/Runtime/GeometryScripting/Source/GeometryScriptingCore/Public/GeometryScript/MeshQueryFunctions.h
--- a/Engine/Plugins/Runtime/GeometryScripting/Source/GeometryScriptingCore/Public/GeometryScript/MeshQueryFunctions.h
+++ b/Engine/Plugins/Runtime/GeometryScripting/Source/GeometryScriptingCore/Public/GeometryScript/MeshQueryFunctions.h
@@ -77,7 +77,7 @@ public:
 	* Returns the number of separate connected islands (components) in the mesh, such as "triangle patches" internally connected by shared edges.
 	*/
 	UFUNCTION(BlueprintPure, Category = "GeometryScript|MeshQueries", meta = (ScriptMethod, DisplayName = "Get Num Connected Islands", KeyWords = "Components Parts"))
-	static UE_API UPARAM(DisplayName = "Num Edges") int32 GetNumConnectedComponents( UDynamicMesh* TargetMesh );
+	static UE_API UPARAM(DisplayName = "Num Islands") int32 GetNumConnectedComponents( UDynamicMesh* TargetMesh );
 
 	// UDynamicMesh already has this function
 	//UFUNCTION(BlueprintPure, Category = "GeometryScript|MeshQueries", meta=(ScriptMethod))
