Add oriented box shape functions to geometry script

diff --git a/Engine/Plugins/Runtime/GeometryScripting/Source/GeometryScriptingCore/Private/ShapeFunctions.cpp b/Engine/Plugins/Runtime/GeometryScripting/Source/GeometryScriptingCore/Private/ShapeFunctions.cpp
--- a/Engine/Plugins/Runtime/GeometryScripting/Source/GeometryScriptingCore/Private/ShapeFunctions.cpp
+++ b/Engine/Plugins/Runtime/GeometryScripting/Source/GeometryScriptingCore/Private/ShapeFunctions.cpp
@@ -3,6 +3,7 @@
 #include "GeometryScript/ShapeFunctions.h"
 #include "VectorTypes.h"
 #include "FrameTypes.h"
+#include "OrientedBoxTypes.h"
 #include "Intersection/IntersectionUtil.h"
 #include "Intersection/IntrRay3AxisAlignedBox3.h"
 #include "Distance/DistLine3Ray3.h"
@@ -298,6 +299,168 @@ bool UGeometryScriptLibrary_BoxFunctions::TestBoxSphereIntersection(FBox Box, FV
 	return FMath::SphereAABBIntersection(SphereCenter, SphereRadius*SphereRadius, Box);
 }
 
+///
+/// Oriented Box functions
+///
+
+FOrientedBox UGeometryScriptLibrary_OrientedBoxFunctions::MakeOrientedBoxFromBox(FBox Box, FTransform Transform)
+{
+	FOrientedBox OrientedBox;
+	FVector BoxCenter, BoxExtents;
+	Box.GetCenterAndExtents(BoxCenter, BoxExtents);
+	OrientedBox.Center = Transform.TransformPosition(BoxCenter);
+	FVector OrientedExtents = (Transform.GetScale3D() * BoxExtents).GetAbs();
+	OrientedBox.ExtentX = OrientedExtents.X;
+	OrientedBox.ExtentY = OrientedExtents.Y;
+	OrientedBox.ExtentZ = OrientedExtents.Z;
+	OrientedBox.AxisX = Transform.GetUnitAxis(EAxis::X);
+	OrientedBox.AxisY = Transform.GetUnitAxis(EAxis::Y);
+	OrientedBox.AxisZ = Transform.GetUnitAxis(EAxis::Z);
+	return OrientedBox;
+}
+
+FRotator UGeometryScriptLibrary_OrientedBoxFunctions::GetOrientedBoxOrientation(FOrientedBox Box)
+{
+	return FRotator(FQuat(FQuaterniond(FMatrix3d(Box.AxisX, Box.AxisY, Box.AxisZ, false))));
+}
+
+
+FVector UGeometryScriptLibrary_OrientedBoxFunctions::GetOrientedBoxCorner(FOrientedBox Box, int CornerIndex)
+{
+	CornerIndex = FMath::Clamp(CornerIndex, 0, 7);
+	
+	// matches Box.GetVertices() / GetBoxCorner() (and not the FOrientedBox3d GetCorner ordering)
+	auto GetCornerDirection = [](int32 CornerIndex)
+		{
+			switch (CornerIndex)		
+			{
+			default:
+			case 0: return FVector(-1.);
+			case 1: return FVector(-1., -1., 1.);
+			case 2: return FVector(-1., 1., -1.);
+			case 3: return FVector(1., -1., -1.);
+			case 4: return FVector(1., 1., -1.);
+			case 5: return FVector(1., -1., 1.);
+			case 6: return FVector(-1., 1., 1.);
+			case 7: return FVector(1.);
+			}
+		};
+	FVector Directions = GetCornerDirection(CornerIndex);
+	return Box.Center
+		+ Box.AxisX * Box.ExtentX * Directions.X
+		+ Box.AxisY * Box.ExtentY * Directions.Y
+		+ Box.AxisZ * Box.ExtentZ * Directions.Z;
+}
+
+FVector UGeometryScriptLibrary_OrientedBoxFunctions::GetOrientedBoxFaceCenter(FOrientedBox Box, int FaceIndex, FVector& FaceNormal)
+{
+	FaceIndex = FMath::Clamp(FaceIndex, 0, 5);
+	switch (FaceIndex)
+	{
+	default:
+	case 0:
+		FaceNormal = -Box.AxisZ;
+		return Box.Center + FaceNormal * Box.ExtentZ;
+	case 1:
+		FaceNormal = Box.AxisZ;
+		return Box.Center + FaceNormal * Box.ExtentZ;
+	case 2:
+		FaceNormal = -Box.AxisY;
+		return Box.Center + FaceNormal * Box.ExtentY;
+	case 3:
+		FaceNormal = Box.AxisY;
+		return Box.Center + FaceNormal * Box.ExtentY;
+	case 4:
+		FaceNormal = -Box.AxisX;
+		return Box.Center + FaceNormal * Box.ExtentX;
+	case 5:
+		FaceNormal = Box.AxisX;
+		return Box.Center + FaceNormal * Box.ExtentX;
+	}
+}
+
+void UGeometryScriptLibrary_OrientedBoxFunctions::GetOrientedBoxVolumeArea(FOrientedBox Box, double& Volume, double& SurfaceArea)
+{
+	FVector Dimensions = FVector(Box.ExtentX, Box.ExtentY, Box.ExtentZ) * 2.;
+	double AreaXY = Dimensions.X * Dimensions.Y;
+	double AreaXZ = Dimensions.X * Dimensions.Z;
+	double AreaYZ = Dimensions.Y * Dimensions.Z;
+	SurfaceArea = 2.0 * AreaXY + 2.0 * AreaXZ + 2.0 * AreaYZ;
+	Volume = Dimensions.X * Dimensions.Y * Dimensions.Z;
+}
+
+FOrientedBox UGeometryScriptLibrary_OrientedBoxFunctions::GetExpandedOrientedBox(FOrientedBox Box, FVector ExpandBy, double MinExtent)
+{
+	// Min Extents cannot be negative
+	MinExtent = FMath::Max(0, MinExtent);
+
+	// Expand extents by the provided vector, clamping at MinExtent in case of a negative ExpandBy amount
+	Box.ExtentX = FMath::Max(MinExtent, Box.ExtentX + ExpandBy.X);
+	Box.ExtentY = FMath::Max(MinExtent, Box.ExtentY + ExpandBy.Y);
+	Box.ExtentZ = FMath::Max(MinExtent, Box.ExtentZ + ExpandBy.Z);
+	return Box;
+}
+
+FOrientedBox UGeometryScriptLibrary_OrientedBoxFunctions::GetTransformedOrientedBox(FOrientedBox OrientedBox, FTransform Transform)
+{
+	FOrientedBox Result;
+	Result.Center = Transform.TransformPosition(OrientedBox.Center);
+	FVector ScaledExtents = (Transform.GetScale3D() * FVector(OrientedBox.ExtentX, OrientedBox.ExtentY, OrientedBox.ExtentZ)).GetAbs();
+	Result.ExtentX = ScaledExtents.X;
+	Result.ExtentY = ScaledExtents.Y;
+	Result.ExtentZ = ScaledExtents.Z;
+	Result.AxisX = Transform.TransformVector(OrientedBox.AxisX);
+	Result.AxisY = Transform.TransformVector(OrientedBox.AxisY);
+	Result.AxisZ = Transform.TransformVector(OrientedBox.AxisZ);
+	return Result;
+}
+
+FBox UGeometryScriptLibrary_OrientedBoxFunctions::ConvertOrientedBoxToAxisAlignedBox(FOrientedBox Box)
+{
+	FBox Result;
+
+	FVector Size = FVector(Box.ExtentX, Box.ExtentY, Box.ExtentZ) * 2.;
+	FVector Corners[8];
+	Corners[0] = Box.Center - Box.AxisX * Box.ExtentX - Box.AxisY * Box.ExtentY - Box.ExtentZ * Box.AxisZ;
+	Corners[7] = Box.Center + Box.AxisX * Box.ExtentX + Box.AxisY * Box.ExtentY + Box.ExtentZ * Box.AxisZ;
+	Corners[1] = Corners[0] + Size.Z * Box.AxisZ;
+	Corners[2] = Corners[0] + Size.Y * Box.AxisY;
+	Corners[3] = Corners[0] + Size.X * Box.AxisX;
+	Corners[4] = Corners[7] - Size.X * Box.AxisX;
+	Corners[5] = Corners[7] - Size.Y * Box.AxisY;
+	Corners[6] = Corners[7] - Size.Z * Box.AxisZ;
+	for (int32 Idx = 0; Idx < 8; ++Idx)
+	{
+		Result += Corners[Idx];
+	}
+	return Result;
+}
+
+bool UGeometryScriptLibrary_OrientedBoxFunctions::TestPointInsideOrOnOrientedBox(FOrientedBox Box, FVector Point)
+{
+	UE::Geometry::FOrientedBox3d OBox(Box);
+	return OBox.Contains(Point);
+}
+
+FVector UGeometryScriptLibrary_OrientedBoxFunctions::FindClosestPointOnOrientedBox(FOrientedBox Box, FVector Point, bool& bIsInside)
+{
+	UE::Geometry::FOrientedBox3d OBox(Box);
+	bIsInside = OBox.Contains(Point);
+	return OBox.ClosestPoint(Point);
+}
+
+bool UGeometryScriptLibrary_OrientedBoxFunctions::TestOrientedBoxOrientedBoxIntersection(FOrientedBox Box1, FOrientedBox Box2)
+{
+	UE::Geometry::FOrientedBox3d OBox1(Box1), OBox2(Box2);
+	return OBox1.Intersects(OBox2);
+}
+
+bool UGeometryScriptLibrary_OrientedBoxFunctions::TestBoxOrientedBoxIntersection(FBox Box1, FOrientedBox OrientedBox2)
+{
+	UE::Geometry::FOrientedBox3d OBox2(OrientedBox2);
+	UE::Geometry::FAxisAlignedBox3d ABox1(Box1);
+	return OBox2.Intersects(ABox1);
+}
 
 
 #undef LOCTEXT_NAMESPACE
diff --git a/Engine/Plugins/Runtime/GeometryScripting/Source/GeometryScriptingCore/Public/GeometryScript/ShapeFunctions.h b/Engine/Plugins/Runtime/GeometryScripting/Source/GeometryScriptingCore/Public/GeometryScript/ShapeFunctions.h
--- a/Engine/Plugins/Runtime/GeometryScripting/Source/GeometryScriptingCore/Public/GeometryScript/ShapeFunctions.h
+++ b/Engine/Plugins/Runtime/GeometryScripting/Source/GeometryScriptingCore/Public/GeometryScript/ShapeFunctions.h
@@ -306,4 +306,101 @@ public:
 
 };
 
+// Oriented Box methods
+UCLASS(MinimalAPI, meta = (ScriptName = "GeometryScript_OrientedBox"))
+class UGeometryScriptLibrary_OrientedBoxFunctions : public UBlueprintFunctionLibrary
+{
+	GENERATED_BODY()
+public:
+
+	/**
+	 * Create an Oriented Box from a Box and a Transform
+	 */
+	UFUNCTION(BlueprintPure, Category = "GeometryScript|Shapes|OrientedBox", meta = (KeyWords = "Transform"))
+	static UE_API UPARAM(DisplayName="OrientedBox") FOrientedBox 
+	MakeOrientedBoxFromBox(FBox Box, FTransform Transform);
+
+	/**
+	 * Get the orientation of an Oriented Box
+	 */
+	UFUNCTION(BlueprintPure, Category = "GeometryScript|Shapes|OrientedBox", meta=(ScriptMethod, KeyWords = "Rotation"))
+	static UE_API UPARAM(DisplayName = "Orientation") FRotator GetOrientedBoxOrientation(FOrientedBox Box);
+
+	/**
+	 * Get the position of a corner of the Oriented Box. Corners are indexed from 0 to 7:
+	 * Corner 0 is in the negative direction along each axis, 1/2/3 are along +Z/+Y/+X from the corner 0,
+	 * 7 is the all-positive direction corner, and 4/5/6 are along -Z/-Y/-X from the corner 7.
+	 */
+	UFUNCTION(BlueprintPure, Category = "GeometryScript|Shapes|OrientedBox", meta=(ScriptMethod))
+	static UE_API UPARAM(DisplayName="Corner Point") FVector 
+	GetOrientedBoxCorner(FOrientedBox Box, int CornerIndex);
+
+	/**
+	 * Get the position of the center of a face of the Oriented Box. Faces are indexed from 0 to 5,
+	 * using an ordering where 0/1 are the negative then positive Z-axis faces, 2/3 are Y-axis faces, and 4/5 are X-axis faces
+	 * @param FaceNormal returned Normal vector of the identified face
+	 */
+	UFUNCTION(BlueprintPure, Category = "GeometryScript|Shapes|OrientedBox", meta=(ScriptMethod))
+	static UE_API UPARAM(DisplayName="Center Point") FVector 
+	GetOrientedBoxFaceCenter(FOrientedBox Box, int FaceIndex, FVector& FaceNormal);
+
+	/**
+	 * Get the Volume and Surface Area of an Oriented Box
+	 */
+	UFUNCTION(BlueprintPure, Category = "GeometryScript|Shapes|OrientedBox", meta=(ScriptMethod))
+	static UE_API void GetOrientedBoxVolumeArea(FOrientedBox Box, double& Volume, double& SurfaceArea);
+
+	/**
+	 * Get the Oriented Box with extents expanded by the ExpandBy vector
+	 * Can also shrink the box if negative ExpandBy values are provided, but will clamp to Max(0, MinExtent)
+	 */
+	UFUNCTION(BlueprintPure, Category = "GeometryScript|Shapes|OrientedBox", meta=(ScriptMethod, KeyWords = "Shrink"))
+	static UE_API UPARAM(DisplayName="Expanded Oriented Box") FOrientedBox 
+	GetExpandedOrientedBox(FOrientedBox OrientedBox, FVector ExpandBy, double MinExtent = 0);
+
+	/**
+	 * Apply the Transform to the Oriented Box
+	 */
+	UFUNCTION(BlueprintPure, Category = "GeometryScript|Shapes|OrientedBox", meta=(ScriptMethod))
+	static UE_API UPARAM(DisplayName="Transformed Box") FOrientedBox 
+	GetTransformedOrientedBox(FOrientedBox OrientedBox, FTransform Transform);
+
+	/**
+	 * Find the bounding Box enclosing the input Oriented Box.
+	 */
+	UFUNCTION(BlueprintPure, Category = "GeometryScript|Shapes|OrientedBox", meta=(ScriptMethod, KeyWords = "Bounding"))
+	static UE_API UPARAM(DisplayName="Transformed Box") FBox 
+	ConvertOrientedBoxToAxisAlignedBox(FOrientedBox Box);
+
+	/**
+	 * Test if a Point is inside or on the Box, returning true if so, otherwise false
+	 */
+	UFUNCTION(BlueprintPure, Category = "GeometryScript|Shapes|OrientedBox", meta = (ScriptMethod))
+	static UE_API UPARAM(DisplayName = "Is Inside") bool
+	TestPointInsideOrOnOrientedBox(FOrientedBox Box, FVector Point);
+
+	/**
+	 * Find the point on the faces of the Box that is closest to the input Point.
+	 * If the Point is inside the Box, it is returned, ie points Inside do not project to the Box Faces
+	 * @param bIsInsideOrOn if the Point is inside or on the Box, this will return as true, otherwise false
+	 */
+	UFUNCTION(BlueprintPure, Category = "GeometryScript|Shapes|OrientedBox", meta = (ScriptMethod, KeyWords = "Project"))
+	static UE_API UPARAM(DisplayName = "Closest Point") FVector
+	FindClosestPointOnOrientedBox(FOrientedBox Box, FVector Point, bool& bIsInsideOrOn);
+
+	/**
+	 * Test if two oriented boxes intersect
+	 */
+	UFUNCTION(BlueprintPure, Category = "GeometryScript|Shapes|OrientedBox", meta = (ScriptMethod, KeyWords = "Overlaps"))
+	static UE_API UPARAM(DisplayName = "Intersects") bool
+	TestOrientedBoxOrientedBoxIntersection(FOrientedBox Box1, FOrientedBox Box2);
+
+	/**
+	 * Test if an axis-aligned box intersects an oriented box
+	 */
+	UFUNCTION(BlueprintPure, Category = "GeometryScript|Shapes|OrientedBox", meta = (ScriptMethod, KeyWords = "Overlaps"))
+	static UE_API UPARAM(DisplayName = "Intersects") bool
+	TestBoxOrientedBoxIntersection(FBox Box, FOrientedBox OrientedBox);
+};
+
 #undef UE_API
