fixes for new geometry script oriented box methods -- apply unscaled transform for box axes + fix displayname on OBB->ABB convert

diff --git a/Engine/Plugins/Runtime/GeometryScripting/Source/GeometryScriptingCore/Private/ShapeFunctions.cpp b/Engine/Plugins/Runtime/GeometryScripting/Source/GeometryScriptingCore/Private/ShapeFunctions.cpp
--- a/Engine/Plugins/Runtime/GeometryScripting/Source/GeometryScriptingCore/Private/ShapeFunctions.cpp
+++ b/Engine/Plugins/Runtime/GeometryScripting/Source/GeometryScriptingCore/Private/ShapeFunctions.cpp
@@ -409,9 +409,9 @@ FOrientedBox UGeometryScriptLibrary_OrientedBoxFunctions::GetTransformedOriented
 	Result.ExtentX = ScaledExtents.X;
 	Result.ExtentY = ScaledExtents.Y;
 	Result.ExtentZ = ScaledExtents.Z;
-	Result.AxisX = Transform.TransformVector(OrientedBox.AxisX);
-	Result.AxisY = Transform.TransformVector(OrientedBox.AxisY);
-	Result.AxisZ = Transform.TransformVector(OrientedBox.AxisZ);
+	Result.AxisX = Transform.TransformVectorNoScale(OrientedBox.AxisX);
+	Result.AxisY = Transform.TransformVectorNoScale(OrientedBox.AxisY);
+	Result.AxisZ = Transform.TransformVectorNoScale(OrientedBox.AxisZ);
 	return Result;
 }
 
diff --git a/Engine/Plugins/Runtime/GeometryScripting/Source/GeometryScriptingCore/Public/GeometryScript/ShapeFunctions.h b/Engine/Plugins/Runtime/GeometryScripting/Source/GeometryScriptingCore/Public/GeometryScript/ShapeFunctions.h
--- a/Engine/Plugins/Runtime/GeometryScripting/Source/GeometryScriptingCore/Public/GeometryScript/ShapeFunctions.h
+++ b/Engine/Plugins/Runtime/GeometryScripting/Source/GeometryScriptingCore/Public/GeometryScript/ShapeFunctions.h
@@ -369,7 +369,7 @@ public:
 	 * Find the bounding Box enclosing the input Oriented Box.
 	 */
 	UFUNCTION(BlueprintPure, Category = "GeometryScript|Shapes|OrientedBox", meta=(ScriptMethod, KeyWords = "Bounding"))
-	static UE_API UPARAM(DisplayName="Transformed Box") FBox 
+	static UE_API UPARAM(DisplayName="Bounding Box") FBox 
 	ConvertOrientedBoxToAxisAlignedBox(FOrientedBox Box);
 
 	/**
