pull sort-by-separate-values helper method from geometry script cpp file over to bufferutil so it can be used in more places

diff --git a/Engine/Plugins/Runtime/GeometryScripting/Source/GeometryScriptingCore/Private/MeshDecompositionFunctions.cpp b/Engine/Plugins/Runtime/GeometryScripting/Source/GeometryScriptingCore/Private/MeshDecompositionFunctions.cpp
--- a/Engine/Plugins/Runtime/GeometryScripting/Source/GeometryScriptingCore/Private/MeshDecompositionFunctions.cpp
+++ b/Engine/Plugins/Runtime/GeometryScripting/Source/GeometryScriptingCore/Private/MeshDecompositionFunctions.cpp
@@ -10,6 +10,7 @@
 #include "Polygroups/PolygroupSet.h"
 #include "DynamicMeshEditor.h"
 #include "UDynamicMesh.h"
+#include "Util/BufferUtil.h"
 
 #include UE_INLINE_GENERATED_CPP_BY_NAME(MeshDecompositionFunctions)
 
@@ -227,60 +228,6 @@ namespace UE::MeshDecompositionFunctionsLocals
 		}
 		return Values;
 	}
-
-	template<typename SortType, typename ValueType>
-	void SortByValuesArrayHelper(TArray<SortType>& ToSort, const TArray<ValueType>& Values, bool bStableSort, EArraySortOrder SortOrder)
-	{
-		const int32 NumValues = Values.Num();
-		check(ToSort.Num() == NumValues);
-
-		// make a reference array of indices, that we will sort based on values
-		TArray<int32> Indices;
-		Indices.SetNumUninitialized(NumValues);
-		for (int32 Idx = 0; Idx < NumValues; ++Idx)
-		{
-			Indices[Idx] = Idx;
-		}
-
-		// sort it
-		if (!bStableSort)
-		{
-			if (SortOrder == EArraySortOrder::Ascending)
-			{
-				Indices.Sort([&Values](int32 A, int32 B) { return Values[A] < Values[B]; });
-			}
-			else
-			{
-				Indices.Sort([&Values](int32 A, int32 B) { return Values[A] > Values[B]; });
-			}
-		}
-		else
-		{
-			if (SortOrder == EArraySortOrder::Ascending)
-			{
-				Indices.StableSort([&Values](int32 A, int32 B) { return Values[A] < Values[B]; });
-			}
-			else
-			{
-				Indices.StableSort([&Values](int32 A, int32 B) { return Values[A] > Values[B]; });
-			}
-		}
-
-		// swap the ToSort array to follow the ordering of the reference array
-		for (int32 Idx = 0; Idx < NumValues; ++Idx)
-		{
-			int32 SwapFromIndex = Indices[Idx];
-			while (SwapFromIndex < Idx)
-			{
-				SwapFromIndex = Indices[SwapFromIndex];
-			}
-
-			if (SwapFromIndex != Idx)
-			{
-				Swap(ToSort[Idx], ToSort[SwapFromIndex]);
-			}
-		}
-	}
 }
 
 
@@ -299,7 +246,7 @@ SortMeshesByVolume(
 		}, 0.0f);
 	
 	// sort by volumes
-	SortByValuesArrayHelper<UDynamicMesh*, float>(Meshes, MeshVolumes, bStableSort, SortOrder);
+	BufferUtil::SortByValuesArray<UDynamicMesh*, float>(Meshes, MeshVolumes, bStableSort, SortOrder == EArraySortOrder::Ascending);
 }
 
 void UGeometryScriptLibrary_MeshDecompositionFunctions::
@@ -317,7 +264,7 @@ SortMeshesByArea(
 		}, 0.0f);
 
 	// sort by areas
-	SortByValuesArrayHelper<UDynamicMesh*, float>(Meshes, MeshAreas, bStableSort, SortOrder);
+	BufferUtil::SortByValuesArray<UDynamicMesh*, float>(Meshes, MeshAreas, bStableSort, SortOrder == EArraySortOrder::Ascending);
 }
 
 void UGeometryScriptLibrary_MeshDecompositionFunctions::
@@ -335,7 +282,7 @@ SortMeshesByBoundsVolume(
 		}, 0.0f);
 
 	// sort by bounds volumes
-	SortByValuesArrayHelper<UDynamicMesh*, float>(Meshes, MeshBoundsVolumes, bStableSort, SortOrder);
+	BufferUtil::SortByValuesArray<UDynamicMesh*, float>(Meshes, MeshBoundsVolumes, bStableSort, SortOrder == EArraySortOrder::Ascending);
 }
 
 void UGeometryScriptLibrary_MeshDecompositionFunctions::
@@ -350,7 +297,7 @@ SortMeshesByCustomValues(
 		UE::Geometry::MakeScriptError(EGeometryScriptErrorType::InvalidInputs, LOCTEXT("SortMeshesByCustomValues_InvalidInput", "SortMeshesByCustomValues: Meshes and Values array must have same number of elements"));
 		return;
 	}
-	UE::MeshDecompositionFunctionsLocals::SortByValuesArrayHelper<UDynamicMesh*, float>(Meshes, ValuesToSortBy, bStableSort, SortOrder);
+	BufferUtil::SortByValuesArray<UDynamicMesh*, float>(Meshes, ValuesToSortBy, bStableSort, SortOrder == EArraySortOrder::Ascending);
 }
 
 
diff --git a/Engine/Source/Runtime/GeometryCore/Public/Util/BufferUtil.h b/Engine/Source/Runtime/GeometryCore/Public/Util/BufferUtil.h
--- a/Engine/Source/Runtime/GeometryCore/Public/Util/BufferUtil.h
+++ b/Engine/Source/Runtime/GeometryCore/Public/Util/BufferUtil.h
@@ -76,4 +76,65 @@ namespace BufferUtil
 		}
 	}
 
+
+	/**
+	 * Sort an array view's values by the values in a second array view
+	 */
+	template<typename SortType, typename ValueType>
+	void SortByValuesArray(TArrayView<SortType> ToSort, TConstArrayView<ValueType> Values, bool bStableSort, bool bAscendingOrder)
+	{
+		const int32 NumValues = Values.Num();
+		if (!ensureMsgf(ToSort.Num() == NumValues, TEXT("Values array must have same length as ToSort array, but did not (%d vs %d)"), NumValues, ToSort.Num()))
+		{
+			return;
+		}
+
+		// make a reference array of indices, that we will sort based on values
+		TArray<int32> Indices;
+		Indices.SetNumUninitialized(NumValues);
+		for (int32 Idx = 0; Idx < NumValues; ++Idx)
+		{
+			Indices[Idx] = Idx;
+		}
+
+		// sort it
+		if (!bStableSort)
+		{
+			if (bAscendingOrder)
+			{
+				Indices.Sort([&Values](int32 A, int32 B) { return Values[A] < Values[B]; });
+			}
+			else
+			{
+				Indices.Sort([&Values](int32 A, int32 B) { return Values[A] > Values[B]; });
+			}
+		}
+		else
+		{
+			if (bAscendingOrder)
+			{
+				Indices.StableSort([&Values](int32 A, int32 B) { return Values[A] < Values[B]; });
+			}
+			else
+			{
+				Indices.StableSort([&Values](int32 A, int32 B) { return Values[A] > Values[B]; });
+			}
+		}
+
+		// swap the ToSort array to follow the ordering of the reference array
+		for (int32 Idx = 0; Idx < NumValues; ++Idx)
+		{
+			int32 SwapFromIndex = Indices[Idx];
+			while (SwapFromIndex < Idx)
+			{
+				SwapFromIndex = Indices[SwapFromIndex];
+			}
+
+			if (SwapFromIndex != Idx)
+			{
+				Swap(ToSort[Idx], ToSort[SwapFromIndex]);
+			}
+		}
+	}
+
 }
\ No newline at end of file
