more consistently use templated valuetype in geometry script mesh sorting helpers, rather than sometimes explicitly using float

diff --git a/Engine/Plugins/Runtime/GeometryScripting/Source/GeometryScriptingCore/Private/MeshDecompositionFunctions.cpp b/Engine/Plugins/Runtime/GeometryScripting/Source/GeometryScriptingCore/Private/MeshDecompositionFunctions.cpp
--- a/Engine/Plugins/Runtime/GeometryScripting/Source/GeometryScriptingCore/Private/MeshDecompositionFunctions.cpp
+++ b/Engine/Plugins/Runtime/GeometryScripting/Source/GeometryScriptingCore/Private/MeshDecompositionFunctions.cpp
@@ -210,18 +210,18 @@ UDynamicMesh* UGeometryScriptLibrary_MeshDecompositionFunctions::SplitMeshByPoly
 namespace UE::MeshDecompositionFunctionsLocals
 {
 	template<typename ValueType>
-	TArray<ValueType> ValuePerMeshArrayHelper(TArray<UDynamicMesh*>& Meshes, TFunctionRef<float(const FDynamicMesh3&)> MeshValueFn, ValueType DefaultValue = 0)
+	TArray<ValueType> ValuePerMeshArrayHelper(TArray<UDynamicMesh*>& Meshes, TFunctionRef<ValueType(const FDynamicMesh3&)> MeshValueFn, ValueType DefaultValue = 0)
 	{
 		TArray<ValueType> Values;
 		Values.SetNumUninitialized(Meshes.Num());
 		for (int32 Idx = 0; Idx < Meshes.Num(); ++Idx)
 		{
-			float MeshValue = DefaultValue;
+			ValueType MeshValue = DefaultValue;
 			if (Meshes[Idx])
 			{
 				Meshes[Idx]->ProcessMesh([&MeshValue, &MeshValueFn](const FDynamicMesh3& Mesh) -> void
 				{
-					MeshValue = (float)MeshValueFn(Mesh);
+					MeshValue = (ValueType)MeshValueFn(Mesh);
 				});
 			}
 			Values[Idx] = MeshValue;
