add a geometry script perlin noise method that computes noise scale from positions
+ add option for mesh perlin noise node to use overlay normals instead of computing new vertex normals
+ refactor overlay normals -> vertex normals logic from geometry script node into FMeshNormals class, and add a MoveNormals to extract the normals array from the class when the class is no longer needed

diff --git a/Engine/Plugins/Runtime/GeometryScripting/Source/GeometryScriptingCore/Private/MeshDeformFunctions.cpp b/Engine/Plugins/Runtime/GeometryScripting/Source/GeometryScriptingCore/Private/MeshDeformFunctions.cpp
--- a/Engine/Plugins/Runtime/GeometryScripting/Source/GeometryScriptingCore/Private/MeshDeformFunctions.cpp
+++ b/Engine/Plugins/Runtime/GeometryScripting/Source/GeometryScriptingCore/Private/MeshDeformFunctions.cpp
@@ -23,6 +23,24 @@ using namespace UE::Geometry;
 #define LOCTEXT_NAMESPACE "UGeometryScriptLibrary_MeshDeformFunctions"
 
 
+namespace UE::MeshDeformFunctionsLocals
+{
+	void ComputePerlinNoiseRandomOffsets(int RandomSeed, FVector FrequencyShift, FVector3d Offset[3])
+	{
+		FRandomStream Random(RandomSeed);
+		FVector3d Offsets[3];
+		for (int k = 0; k < 3; ++k)
+		{
+			const float RandomOffset = 10000.0f * Random.GetFraction();
+			Offsets[k] = FVector3d(RandomOffset, RandomOffset, RandomOffset);
+			Offsets[k] += FrequencyShift;
+		}
+
+	}
+}
+
+
+
 UDynamicMesh* UGeometryScriptLibrary_MeshDeformFunctions::ApplyBendWarpToMesh(
 	UDynamicMesh* TargetMesh,
 	FGeometryScriptBendWarpOptions Options,
@@ -258,20 +276,26 @@ UDynamicMesh* UGeometryScriptLibrary_MeshDeformFunctions::ApplyPerlinNoiseToMesh
 
 	TargetMesh->EditMesh([&](FDynamicMesh3& EditMesh) 
 	{
-		FRandomStream Random(Options.BaseLayer.RandomSeed);
 		FVector3d Offsets[3];
-		for (int k = 0; k < 3; ++k)
-		{
-			const float RandomOffset = 10000.0f * Random.GetFraction();
-			Offsets[k] = FVector3d(RandomOffset, RandomOffset, RandomOffset);
-			Offsets[k] += (FVector3d)Options.BaseLayer.FrequencyShift;
-		}
+		UE::MeshDeformFunctionsLocals::ComputePerlinNoiseRandomOffsets(Options.BaseLayer.RandomSeed, Options.BaseLayer.FrequencyShift, Offsets);
 
 		FMeshNormals Normals(&EditMesh);
 		bool bAlignWithNormal = Options.bApplyAlongNormal;
 		if (bAlignWithNormal)
 		{
-			Normals.ComputeVertexNormals();
+			if (Options.NormalSource == EGeometryScriptPerVertexNormalSource::Computed)
+			{
+				Normals.ComputeVertexNormals();
+			}
+			else // EGeometryScriptPerVertexNormalSource::Averaged
+			{
+				if (!EditMesh.HasAttributes() || !EditMesh.Attributes()->PrimaryNormals())
+				{
+					UE::Geometry::AppendError(Debug, EGeometryScriptErrorType::InvalidInputs, LOCTEXT("ApplyPerlinNoiseToMesh_MissingOverlayNormals", "ApplyPerlinNoiseToMesh: TargetMesh did not have overlay normals computed"));
+					return;
+				}
+				Normals.GetVertexNormalsFromOverlayNormals();
+			}
 		}
 
 		auto GetDisplacedPosition = [&EditMesh, &Options, &Offsets, &Normals, bAlignWithNormal](int32 VertexID)
@@ -320,7 +344,32 @@ UDynamicMesh* UGeometryScriptLibrary_MeshDeformFunctions::ApplyPerlinNoiseToMesh
 }
 
 
+FGeometryScriptScalarList
+UGeometryScriptLibrary_MeshDeformFunctions::ComputePerlinNoise(FGeometryScriptVectorList Positions, FGeometryScriptPerlinNoiseOptions Options)
+{
+	FGeometryScriptScalarList NoiseValues;
+	NoiseValues.Reset();
 
+	if (!Positions.List)
+	{
+		return NoiseValues;
+	}
+
+	FVector3d Offsets[3];
+	UE::MeshDeformFunctionsLocals::ComputePerlinNoiseRandomOffsets(Options.BaseLayer.RandomSeed, Options.BaseLayer.FrequencyShift, Offsets);
+
+	TArray<double>& NoiseArray = *NoiseValues.List;
+	TArray<FVector>& PosArray = *Positions.List;
+	NoiseArray.SetNumUninitialized(PosArray.Num());
+	ParallelFor(PosArray.Num(), [&PosArray, &NoiseArray, &Options, &Offsets](int32 Idx)
+		{
+			FVector NoisePos = (FVector)((double)Options.BaseLayer.Frequency * (PosArray[Idx] + Offsets[0]));
+			NoiseArray[Idx] = Options.BaseLayer.Magnitude * FMath::PerlinNoise3D(NoisePos);
+		}
+	);
+
+	return NoiseValues;
+}
 
 
 
diff --git a/Engine/Plugins/Runtime/GeometryScripting/Source/GeometryScriptingCore/Private/MeshNormalsFunctions.cpp b/Engine/Plugins/Runtime/GeometryScripting/Source/GeometryScriptingCore/Private/MeshNormalsFunctions.cpp
--- a/Engine/Plugins/Runtime/GeometryScripting/Source/GeometryScriptingCore/Private/MeshNormalsFunctions.cpp
+++ b/Engine/Plugins/Runtime/GeometryScripting/Source/GeometryScriptingCore/Private/MeshNormalsFunctions.cpp
@@ -621,48 +621,9 @@ UDynamicMesh* UGeometryScriptLibrary_MeshNormalsFunctions::GetMeshPerVertexNorma
 
 		if (ReadMesh.HasAttributes() && ReadMesh.Attributes()->NumNormalLayers() > 0 )
 		{
-			const FDynamicMeshNormalOverlay* NormalOverlay = ReadMesh.Attributes()->PrimaryNormals();
-
-			if (bAverageSplitVertexValues)
-			{
-				Normals.Init(FVector::Zero(), ReadMesh.MaxVertexID());
-				for (int32 tid : ReadMesh.TriangleIndicesItr())
-				{
-					if (NormalOverlay->IsSetTriangle(tid))
-					{
-						FIndex3i TriV = ReadMesh.GetTriangle(tid);
-						FVector3f A, B, C;
-						NormalOverlay->GetTriElements(tid, A, B, C);
-						Normals[TriV.A] += (FVector)A;
-						Normals[TriV.B] += (FVector)B;
-						Normals[TriV.C] += (FVector)C;
-					}
-				}
-
-				for (int32 k = 0; k < Normals.Num(); ++k)
-				{
-					if (Normals[k].SquaredLength() > 0)
-					{
-						Normalize(Normals[k]);
-					}
-				}
-			}
-			else
-			{
-				Normals.Init(FVector::UnitZ(), ReadMesh.MaxVertexID());
-				for (int32 tid : ReadMesh.TriangleIndicesItr())
-				{
-					if (NormalOverlay->IsSetTriangle(tid))
-					{
-						FIndex3i TriV = ReadMesh.GetTriangle(tid);
-						FVector3f A, B, C;
-						NormalOverlay->GetTriElements(tid, A, B, C);
-						Normals[TriV.A] = (FVector)A;
-						Normals[TriV.B] = (FVector)B;
-						Normals[TriV.C] = (FVector)C;
-					}
-				}
-			}
+			FMeshNormals MeshNormals(&ReadMesh);
+			MeshNormals.GetVertexNormalsFromOverlayNormals(bAverageSplitVertexValues ? FMeshNormals::ECombineSplitNormalsMethod::Average : FMeshNormals::ECombineSplitNormalsMethod::CopyAny);
+			Normals = MeshNormals.MoveNormals();
 
 			bIsValidNormalSet = true;
 		}
diff --git a/Engine/Plugins/Runtime/GeometryScripting/Source/GeometryScriptingCore/Public/GeometryScript/MeshDeformFunctions.h b/Engine/Plugins/Runtime/GeometryScripting/Source/GeometryScriptingCore/Public/GeometryScript/MeshDeformFunctions.h
--- a/Engine/Plugins/Runtime/GeometryScripting/Source/GeometryScriptingCore/Public/GeometryScript/MeshDeformFunctions.h
+++ b/Engine/Plugins/Runtime/GeometryScripting/Source/GeometryScriptingCore/Public/GeometryScript/MeshDeformFunctions.h
@@ -129,6 +129,14 @@ public:
 };
 
 
+UENUM(BlueprintType)
+enum class EGeometryScriptPerVertexNormalSource : uint8
+{
+	// Compute new per-vertex normals from the mesh geometry
+	Computed,
+	// Average existing overlay normals to compute per-vertex normals
+	AverageFromOverlay
+};
 
 
 USTRUCT(BlueprintType)
@@ -142,6 +150,9 @@ public:
 	UPROPERTY(EditAnywhere, BlueprintReadWrite, Category = Options)
 	bool bApplyAlongNormal = true;
 
+	UPROPERTY(EditAnywhere, BlueprintReadWrite, Category = Options, meta=(EditCondition="bApplyAlongNormal"))
+	EGeometryScriptPerVertexNormalSource NormalSource = EGeometryScriptPerVertexNormalSource::Computed;
+
 	/** EmptyBehavior Defines how an empty input selection should be interpreted */
 	UPROPERTY(EditAnywhere, BlueprintReadWrite, Category = Options)
 	EGeometryScriptEmptySelectionBehavior EmptyBehavior = EGeometryScriptEmptySelectionBehavior::FullMeshSelection;
@@ -273,6 +284,13 @@ public:
 		FGeometryScriptPerlinNoiseOptions Options,
 		UGeometryScriptDebug* Debug = nullptr);
 
+	/**
+	 * Compute Perlin Noise at the requested positions
+	 */
+	UFUNCTION(BlueprintCallable, Category = "GeometryScript|Deformations", meta = (ScriptMethod))
+	static UE_API UPARAM(DisplayName = "Noise Values") FGeometryScriptScalarList
+	ComputePerlinNoise(FGeometryScriptVectorList Positions, FGeometryScriptPerlinNoiseOptions Options);
+
 	/**
 	* Applies a number of iterations of mesh smoothing to a Dynamic Mesh.
 	* @param Selection if non-empty, only the vertices identified by the selection will be subject to smoothing,  otherwise
diff --git a/Engine/Source/Runtime/GeometryCore/Private/DynamicMesh/MeshNormals.cpp b/Engine/Source/Runtime/GeometryCore/Private/DynamicMesh/MeshNormals.cpp
--- a/Engine/Source/Runtime/GeometryCore/Private/DynamicMesh/MeshNormals.cpp
+++ b/Engine/Source/Runtime/GeometryCore/Private/DynamicMesh/MeshNormals.cpp
@@ -45,6 +45,57 @@ void FMeshNormals::CopyToVertexNormals(FDynamicMesh3* SetMesh, bool bInvert) con
 }
 
 
+void FMeshNormals::GetVertexNormalsFromOverlayNormals(ECombineSplitNormalsMethod CombineSplitNormals)
+{
+	check(Mesh);
+
+	// no overlay to copy from
+	if (!Mesh->HasAttributes() || !Mesh->Attributes()->PrimaryNormals())
+	{
+		Normals.Init(FVector::UnitZ(), Mesh->MaxVertexID());
+		return;
+	}
+
+	const FDynamicMeshNormalOverlay* NormalOverlay = Mesh->Attributes()->PrimaryNormals();
+	if (CombineSplitNormals == ECombineSplitNormalsMethod::Average)
+	{
+		Normals.Init(FVector::Zero(), Mesh->MaxVertexID());
+		for (int32 TID : Mesh->TriangleIndicesItr())
+		{
+			if (NormalOverlay->IsSetTriangle(TID))
+			{
+				FIndex3i TriV = Mesh->GetTriangle(TID);
+				FVector3f A, B, C;
+				NormalOverlay->GetTriElements(TID, A, B, C);
+				Normals[TriV.A] += (FVector)A;
+				Normals[TriV.B] += (FVector)B;
+				Normals[TriV.C] += (FVector)C;
+			}
+		}
+
+		for (int32 k = 0; k < Normals.Num(); ++k)
+		{
+			Normalize(Normals[k]);
+		}
+	}
+	else // ECombineSplitNormalsMethod::CopyAny
+	{
+		Normals.Init(FVector::UnitZ(), Mesh->MaxVertexID());
+		for (int32 TID : Mesh->TriangleIndicesItr())
+		{
+			if (NormalOverlay->IsSetTriangle(TID))
+			{
+				FIndex3i TriV = Mesh->GetTriangle(TID);
+				FVector3f A, B, C;
+				NormalOverlay->GetTriElements(TID, A, B, C);
+				Normals[TriV.A] = (FVector)A;
+				Normals[TriV.B] = (FVector)B;
+				Normals[TriV.C] = (FVector)C;
+			}
+		}
+	}
+}
+
 
 void FMeshNormals::CopyToOverlay(FDynamicMeshNormalOverlay* NormalOverlay, bool bInvert) const
 {
diff --git a/Engine/Source/Runtime/GeometryCore/Public/DynamicMesh/MeshNormals.h b/Engine/Source/Runtime/GeometryCore/Public/DynamicMesh/MeshNormals.h
--- a/Engine/Source/Runtime/GeometryCore/Public/DynamicMesh/MeshNormals.h
+++ b/Engine/Source/Runtime/GeometryCore/Public/DynamicMesh/MeshNormals.h
@@ -49,6 +49,14 @@ public:
 
 	const TArray<FVector3d>& GetNormals() const { return Normals; }
 
+	/*
+	 * @return the normals array by a move
+	 */
+	TArray<FVector3d>&& MoveNormals()
+	{ 
+		return MoveTemp(Normals); 
+	}
+
 	FVector3d& operator[](int i) { return Normals[i]; }
 	const FVector3d& operator[](int i) const { return Normals[i]; }
 
@@ -65,6 +73,21 @@ public:
 	{
 		Compute_FaceAvg(bWeightByArea, bWeightByAngle);
 	}
+	
+	// Method to use when copying split normals to per-vertex normals
+	enum class ECombineSplitNormalsMethod
+	{
+		// Average the split normals
+		Average,
+		// Arbitrarily choose one of the normals to exactly copy
+		CopyAny
+	};
+
+	/**
+	 * Fill in per-vertex normals in the FMeshNormals::Normals array by copying normals from the overlay
+	 * @param CombineSplitNormals Method to use when copying split normals to per-vertex normals.
+	 */
+	GEOMETRYCORE_API void GetVertexNormalsFromOverlayNormals(ECombineSplitNormalsMethod CombineSplitNormals = ECombineSplitNormalsMethod::Average);
 
 	/**
 	 * Compute per-triangle normals
