geometry script perlin noise fixes: fix incorrect offsets in perlin noise method + incorrect comment

#FailedValidations Preflight

diff --git a/Engine/Plugins/Runtime/GeometryScripting/Source/GeometryScriptingCore/Private/MeshDeformFunctions.cpp b/Engine/Plugins/Runtime/GeometryScripting/Source/GeometryScriptingCore/Private/MeshDeformFunctions.cpp
--- a/Engine/Plugins/Runtime/GeometryScripting/Source/GeometryScriptingCore/Private/MeshDeformFunctions.cpp
+++ b/Engine/Plugins/Runtime/GeometryScripting/Source/GeometryScriptingCore/Private/MeshDeformFunctions.cpp
@@ -25,15 +25,14 @@ using namespace UE::Geometry;
 
 namespace UE::MeshDeformFunctionsLocals
 {
-	void ComputePerlinNoiseRandomOffsets(int RandomSeed, FVector FrequencyShift, FVector3d Offset[3])
+	void ComputePerlinNoiseRandomOffsets(int RandomSeed, FVector FrequencyShift, FVector3d (&OutOffsets)[3])
 	{
 		FRandomStream Random(RandomSeed);
-		FVector3d Offsets[3];
 		for (int k = 0; k < 3; ++k)
 		{
 			const float RandomOffset = 10000.0f * Random.GetFraction();
-			Offsets[k] = FVector3d(RandomOffset, RandomOffset, RandomOffset);
-			Offsets[k] += FrequencyShift;
+			OutOffsets[k] = FVector3d(RandomOffset, RandomOffset, RandomOffset);
+			OutOffsets[k] += FrequencyShift;
 		}
 
 	}
@@ -287,7 +286,7 @@ UDynamicMesh* UGeometryScriptLibrary_MeshDeformFunctions::ApplyPerlinNoiseToMesh
 			{
 				Normals.ComputeVertexNormals();
 			}
-			else // EGeometryScriptPerVertexNormalSource::Averaged
+			else // EGeometryScriptPerVertexNormalSource::AverageFromOverlay
 			{
 				if (!EditMesh.HasAttributes() || !EditMesh.Attributes()->PrimaryNormals())
 				{
