GeometryScript: add weight map support to Remesh node

diff --git a/Engine/Plugins/Runtime/GeometryScripting/Source/GeometryScriptingCore/Private/MeshQueryFunctions.cpp b/Engine/Plugins/Runtime/GeometryScripting/Source/GeometryScriptingCore/Private/MeshQueryFunctions.cpp
--- a/Engine/Plugins/Runtime/GeometryScripting/Source/GeometryScriptingCore/Private/MeshQueryFunctions.cpp
+++ b/Engine/Plugins/Runtime/GeometryScripting/Source/GeometryScriptingCore/Private/MeshQueryFunctions.cpp
@@ -824,6 +824,32 @@ int UGeometryScriptLibrary_MeshQueryFunctions::GetNumExtendedPolygroupLayers( UD
 	});
 }
 
+FGeometryScriptWeightMapHandle UGeometryScriptLibrary_MeshQueryFunctions::GetWeightMapHandle(UDynamicMesh* TargetMesh, const FName& Name, EGeometryScriptSearchOutcomePins& FoundWeights)
+{
+	FGeometryScriptWeightMapHandle Handle;
+	Handle.WeightMapAttributeLayerIndex = SimpleMeshQuery<int32>(TargetMesh, (int32)-1, [&Name](const FDynamicMesh3& Mesh) -> int32
+	{
+		if (Mesh.HasAttributes())
+		{
+			for (int32 LayerIndex = 0; LayerIndex < Mesh.Attributes()->NumWeightLayers(); ++LayerIndex)
+			{
+				if (Mesh.Attributes()->GetWeightLayer(LayerIndex)->GetName() == Name)
+				{
+					return LayerIndex;
+				}
+			}
+		}
+
+		UE::Geometry::MakeScriptWarning(EGeometryScriptErrorType::InvalidInputs,
+			FText::Format(LOCTEXT("InvalidWeightMapName", "Weight Map Layer with name \"{0}\" not found on TargetMesh"), FText::FromName(Name)));
+
+		return -1;
+	});
+
+	FoundWeights = Handle.IsValid() ? EGeometryScriptSearchOutcomePins::Found : EGeometryScriptSearchOutcomePins::NotFound;
+
+	return Handle;
+}
 
 
 #undef LOCTEXT_NAMESPACE
diff --git a/Engine/Plugins/Runtime/GeometryScripting/Source/GeometryScriptingCore/Private/MeshRemeshFunctions.cpp b/Engine/Plugins/Runtime/GeometryScripting/Source/GeometryScriptingCore/Private/MeshRemeshFunctions.cpp
--- a/Engine/Plugins/Runtime/GeometryScripting/Source/GeometryScriptingCore/Private/MeshRemeshFunctions.cpp
+++ b/Engine/Plugins/Runtime/GeometryScripting/Source/GeometryScriptingCore/Private/MeshRemeshFunctions.cpp
@@ -19,44 +19,32 @@ using namespace UE::Geometry;
 
 #define LOCTEXT_NAMESPACE "UGeometryScriptLibrary_MeshRemeshFunctions"
 
-
-
-static EEdgeRefineFlags MakeEdgeRefineFlagsFromConstraintType(EGeometryScriptRemeshEdgeConstraintType ConstraintType)
+namespace GeometryScriptRemeshLocals
 {
-	switch (ConstraintType)
+	static EEdgeRefineFlags MakeEdgeRefineFlagsFromConstraintType(EGeometryScriptRemeshEdgeConstraintType ConstraintType)
 	{
+		switch (ConstraintType)
+		{
 		case EGeometryScriptRemeshEdgeConstraintType::Fixed: return EEdgeRefineFlags::FullyConstrained;
 		case EGeometryScriptRemeshEdgeConstraintType::Refine: return EEdgeRefineFlags::SplitsOnly;
 		case EGeometryScriptRemeshEdgeConstraintType::Free: return EEdgeRefineFlags::NoFlip;
 		case EGeometryScriptRemeshEdgeConstraintType::Ignore: return EEdgeRefineFlags::NoConstraint;
-	}
-	return EEdgeRefineFlags::NoFlip;
-}
-
-
-UDynamicMesh* UGeometryScriptLibrary_RemeshingFunctions::ApplyUniformRemesh(
-	UDynamicMesh* TargetMesh, 
-	FGeometryScriptRemeshOptions RemeshOptions,
-	FGeometryScriptUniformRemeshOptions UniformOptions,
-	UGeometryScriptDebug* Debug)
-{
-	if (TargetMesh == nullptr)
-	{
-		UE::Geometry::AppendError(Debug, EGeometryScriptErrorType::InvalidInputs, LOCTEXT("ApplyUniformRemesh_InvalidInput", "ApplyUniformRemesh: TargetMesh is Null"));
-		return TargetMesh;
+		}
+		return EEdgeRefineFlags::NoFlip;
 	}
 
-	TargetMesh->EditMesh([&](FDynamicMesh3& EditMesh) 
+	template<typename OptionsT>
+	void CreateRemeshOp(TSharedPtr<FDynamicMesh3> SourceMesh,
+		const FGeometryScriptRemeshOptions& RemeshOptions,
+		const OptionsT& Options,
+		FRemeshMeshOp& RemeshOp)
 	{
-		TSharedPtr<FDynamicMesh3> SourceMesh = MakeShared<FDynamicMesh3>(MoveTemp(EditMesh));
-
 		TSharedPtr<FDynamicMeshAABBTree3> SourceSpatial;
 		if (RemeshOptions.bReprojectToInputMesh)
 		{
 			SourceSpatial = MakeShared<FDynamicMeshAABBTree3>(SourceMesh.Get(), true);
 		}
 
-		FRemeshMeshOp RemeshOp;
 		RemeshOp.OriginalMesh = SourceMesh;
 		RemeshOp.OriginalMeshSpatial = SourceSpatial;
 
@@ -85,20 +73,20 @@ UDynamicMesh* UGeometryScriptLibrary_RemeshingFunctions::ApplyUniformRemesh(
 			break;
 		}
 
-		if (UniformOptions.TargetType == EGeometryScriptUniformRemeshTargetType::TriangleCount)
+		if (Options.TargetType == EGeometryScriptUniformRemeshTargetType::TriangleCount)
 		{
-			RemeshOp.TargetEdgeLength = FRemeshMeshOp::CalculateTargetEdgeLength(SourceMesh.Get(), UniformOptions.TargetTriangleCount);
+			RemeshOp.TargetEdgeLength = FRemeshMeshOp::CalculateTargetEdgeLength(SourceMesh.Get(), Options.TargetTriangleCount);
 		}
 		else
 		{
-			RemeshOp.TargetEdgeLength = UniformOptions.TargetEdgeLength;
+			RemeshOp.TargetEdgeLength = Options.TargetEdgeLength;
 		}
 
 		// currently not exposing this option. It seems to control multiple things that should be independent...
 		RemeshOp.bPreserveSharpEdges = (RemeshOp.bDiscardAttributes == false);
-		RemeshOp.bFlips = RemeshOptions.bAllowFlips; 
-		RemeshOp.bSplits = RemeshOptions.bAllowSplits; 
-		RemeshOp.bCollapses = RemeshOptions.bAllowCollapses; 
+		RemeshOp.bFlips = RemeshOptions.bAllowFlips;
+		RemeshOp.bSplits = RemeshOptions.bAllowSplits;
+		RemeshOp.bCollapses = RemeshOptions.bAllowCollapses;
 		RemeshOp.bPreventNormalFlips = RemeshOptions.bPreventNormalFlips;
 		RemeshOp.bPreventTinyTriangles = RemeshOptions.bPreventTinyTriangles;
 
@@ -106,9 +94,9 @@ UDynamicMesh* UGeometryScriptLibrary_RemeshingFunctions::ApplyUniformRemesh(
 		RemeshOp.GroupBoundaryConstraint = MakeEdgeRefineFlagsFromConstraintType(RemeshOptions.GroupBoundaryConstraint);
 		RemeshOp.MaterialBoundaryConstraint = MakeEdgeRefineFlagsFromConstraintType(RemeshOptions.MaterialBoundaryConstraint);
 
-		RemeshOp.bReproject = RemeshOptions.bReprojectToInputMesh; 
+		RemeshOp.bReproject = RemeshOptions.bReprojectToInputMesh;
 		RemeshOp.ProjectionTarget = nullptr;
-		RemeshOp.ProjectionTargetSpatial = nullptr;		
+		RemeshOp.ProjectionTargetSpatial = nullptr;
 
 
 		RemeshOp.bReprojectConstraints = false;
@@ -119,10 +107,32 @@ UDynamicMesh* UGeometryScriptLibrary_RemeshingFunctions::ApplyUniformRemesh(
 		RemeshOp.ToolMeshLocalToWorld = FTransformSRT3d::Identity();
 		RemeshOp.bUseWorldSpace = false;
 		RemeshOp.bParallel = true;
+	}
+
+}
+
+UDynamicMesh* UGeometryScriptLibrary_RemeshingFunctions::ApplyUniformRemesh(
+	UDynamicMesh* TargetMesh, 
+	FGeometryScriptRemeshOptions RemeshOptions,
+	FGeometryScriptUniformRemeshOptions UniformOptions,
+	UGeometryScriptDebug* Debug)
+{
+	if (TargetMesh == nullptr)
+	{
+		UE::Geometry::AppendError(Debug, EGeometryScriptErrorType::InvalidInputs, LOCTEXT("ApplyUniformRemesh_InvalidInput", "ApplyUniformRemesh: TargetMesh is Null"));
+		return TargetMesh;
+	}
+
+	TargetMesh->EditMesh([&](FDynamicMesh3& EditMesh) 
+	{
+		TSharedPtr<FDynamicMesh3> SourceMesh = MakeShared<FDynamicMesh3>(MoveTemp(EditMesh));
+
+		FRemeshMeshOp RemeshOp;
+		GeometryScriptRemeshLocals::CreateRemeshOp(SourceMesh, RemeshOptions, UniformOptions, RemeshOp);
 
 		RemeshOp.CalculateResult(nullptr);
 		if (RemeshOp.GetResultInfo().Result == EGeometryResultType::Success)
-		{
+		{ 
 			TUniquePtr<FDynamicMesh3> ResultMesh = RemeshOp.ExtractResult();
 			EditMesh = MoveTemp(*ResultMesh);
 		}
@@ -147,11 +157,92 @@ UDynamicMesh* UGeometryScriptLibrary_RemeshingFunctions::ApplyUniformRemesh(
 			EditMesh.Attributes()->EnableMaterialID();
 		}
 
+
 	}, EDynamicMeshChangeType::GeneralEdit, EDynamicMeshAttributeChangeFlags::Unknown, false);
 
 	return TargetMesh;
 }
 
 
+UDynamicMesh* UGeometryScriptLibrary_RemeshingFunctions::ApplyAdaptiveRemesh(
+	UDynamicMesh* TargetMesh,
+	FGeometryScriptRemeshOptions RemeshOptions,
+	FGeometryScriptAdaptiveRemeshOptions AdaptiveOptions,
+	FGeometryScriptWeightMapHandle WeightMapHandle,
+	UGeometryScriptDebug* Debug)
+{
+	if (TargetMesh == nullptr)
+	{
+		UE::Geometry::AppendError(Debug, EGeometryScriptErrorType::InvalidInputs, LOCTEXT("ApplyAdaptiveRemesh_InvalidInput", "ApplyAdaptiveRemesh: TargetMesh is Null"));
+		return TargetMesh;
+	}
+
+	if (!WeightMapHandle.IsValid() || WeightMapHandle.WeightMapAttributeLayerIndex >= TargetMesh->GetMeshRef().Attributes()->NumWeightLayers())
+	{
+		UE::Geometry::AppendError(Debug, EGeometryScriptErrorType::InvalidInputs, LOCTEXT("ApplyAdaptiveRemesh_InvalidWeightMap", "ApplyAdaptiveRemesh: WeightMapHandle is not valid"));
+		return TargetMesh;
+	}
+
+	TargetMesh->EditMesh([&](FDynamicMesh3& EditMesh)
+	{
+		TSharedPtr<FDynamicMesh3> SourceMesh = MakeShared<FDynamicMesh3>(MoveTemp(EditMesh));
+
+		FRemeshMeshOp RemeshOp;
+		GeometryScriptRemeshLocals::CreateRemeshOp(SourceMesh, RemeshOptions, AdaptiveOptions, RemeshOp);
+
+		if (SourceMesh->HasAttributes())
+		{
+			const int32 WeightMapAttributeLayerIndex = WeightMapHandle.WeightMapAttributeLayerIndex;
+
+			RemeshOp.CustomEdgeLengthScaleF = [WeightMapAttributeLayerIndex, RelativeDensity = AdaptiveOptions.RelativeDensity](const FDynamicMesh3& Mesh, int VertexA, int VertexB) -> double
+			{
+				const FDynamicMeshWeightAttribute* const AttributeLayer = Mesh.Attributes()->GetWeightLayer(WeightMapAttributeLayerIndex);
+
+				float WeightValueA;
+				AttributeLayer->GetValue(VertexA, &WeightValueA);
+				WeightValueA = FMath::Clamp(WeightValueA, 0.0f, 1.0f);
+
+				float WeightValueB;
+				AttributeLayer->GetValue(VertexB, &WeightValueB);
+				WeightValueB = FMath::Clamp(WeightValueB, 0.0f, 1.0f);
+
+				const double AvgWeight = 0.5 * (WeightValueA + WeightValueB);
+
+				return FMath::Pow(0.5, -RelativeDensity * AvgWeight);
+			};
+		}
+
+		RemeshOp.CalculateResult(nullptr);
+		if (RemeshOp.GetResultInfo().Result == EGeometryResultType::Success)
+		{
+			TUniquePtr<FDynamicMesh3> ResultMesh = RemeshOp.ExtractResult();
+			EditMesh = MoveTemp(*ResultMesh);
+		}
+		else
+		{
+			EditMesh = MoveTemp(*SourceMesh);
+			UE::Geometry::AppendError(Debug, EGeometryScriptErrorType::InvalidInputs, LOCTEXT("ApplyAdaptiveRemesh_ComputeError", "ApplyAdaptiveRemesh: Error computing result, returning input mesh"));
+		}
+
+		// compact the input mesh if enabled
+		if (RemeshOptions.bAutoCompact)
+		{
+			EditMesh.CompactInPlace();
+		}
+
+		// if we discarded attributes, re-enable the standard attributes, and initialize with per-vertex normals
+		if (RemeshOptions.bDiscardAttributes)
+		{
+			EditMesh.EnableTriangleGroups();
+			EditMesh.EnableAttributes();
+			FMeshNormals::InitializeOverlayToPerVertexNormals(EditMesh.Attributes()->PrimaryNormals(), false);
+			EditMesh.Attributes()->EnableMaterialID();
+		}
+
+	}, EDynamicMeshChangeType::GeneralEdit, EDynamicMeshAttributeChangeFlags::Unknown, false);
+
+	return TargetMesh;
+}
+
 
 #undef LOCTEXT_NAMESPACE
diff --git a/Engine/Plugins/Runtime/GeometryScripting/Source/GeometryScriptingCore/Public/GeometryScript/GeometryScriptTypes.h b/Engine/Plugins/Runtime/GeometryScripting/Source/GeometryScriptingCore/Public/GeometryScript/GeometryScriptTypes.h
--- a/Engine/Plugins/Runtime/GeometryScripting/Source/GeometryScriptingCore/Public/GeometryScript/GeometryScriptTypes.h
+++ b/Engine/Plugins/Runtime/GeometryScripting/Source/GeometryScriptingCore/Public/GeometryScript/GeometryScriptTypes.h
@@ -289,6 +289,27 @@ public:
 };
 
 
+//
+// Attributes
+//
+
+USTRUCT(BlueprintType, meta = (DisplayName = "Weight Map Handle"))
+struct FGeometryScriptWeightMapHandle
+{
+	GENERATED_BODY()
+
+public:
+
+	bool IsValid() const
+	{
+		return WeightMapAttributeLayerIndex >= 0;
+	}
+
+	UPROPERTY(EditAnywhere, BlueprintReadWrite, Category = "Handle")
+	int WeightMapAttributeLayerIndex = -1;
+};
+
+
 //
 // List Types
 //
diff --git a/Engine/Plugins/Runtime/GeometryScripting/Source/GeometryScriptingCore/Public/GeometryScript/MeshQueryFunctions.h b/Engine/Plugins/Runtime/GeometryScripting/Source/GeometryScriptingCore/Public/GeometryScript/MeshQueryFunctions.h
--- a/Engine/Plugins/Runtime/GeometryScripting/Source/GeometryScriptingCore/Public/GeometryScript/MeshQueryFunctions.h
+++ b/Engine/Plugins/Runtime/GeometryScripting/Source/GeometryScriptingCore/Public/GeometryScript/MeshQueryFunctions.h
@@ -478,6 +478,13 @@ public:
 	static UE_API UPARAM(DisplayName = "Num Layers") int GetNumExtendedPolygroupLayers( UDynamicMesh* TargetMesh );
 
 
+	//
+	// Weight Map Attribute Layer Queries
+	//
+
+	UFUNCTION(BlueprintCallable, Category = "GeometryScript|MeshQueries", meta = (ScriptMethod, DisplayName = "Get Weight Map Handle By Name", ExpandEnumAsExecs = "FoundWeights"))
+	static UE_API UPARAM(DisplayName = "Weight Map Handle") FGeometryScriptWeightMapHandle GetWeightMapHandle(UDynamicMesh* TargetMesh, const FName& Name, EGeometryScriptSearchOutcomePins& FoundWeights);
+
 };
 
 #undef UE_API
diff --git a/Engine/Plugins/Runtime/GeometryScripting/Source/GeometryScriptingCore/Public/GeometryScript/MeshRemeshFunctions.h b/Engine/Plugins/Runtime/GeometryScripting/Source/GeometryScriptingCore/Public/GeometryScript/MeshRemeshFunctions.h
--- a/Engine/Plugins/Runtime/GeometryScripting/Source/GeometryScriptingCore/Public/GeometryScript/MeshRemeshFunctions.h
+++ b/Engine/Plugins/Runtime/GeometryScripting/Source/GeometryScriptingCore/Public/GeometryScript/MeshRemeshFunctions.h
@@ -141,6 +141,30 @@ public:
 	float TargetEdgeLength = 1.0f;
 };
 
+/**
+ * Adaptive Remeshing Options
+ */
+USTRUCT(Blueprintable)
+struct FGeometryScriptAdaptiveRemeshOptions
+{
+	GENERATED_BODY()
+public:
+	/** Method used to define target/goal of Adaptive Remeshing */
+	UPROPERTY(EditAnywhere, BlueprintReadWrite, Category = Options)
+	EGeometryScriptUniformRemeshTargetType TargetType = EGeometryScriptUniformRemeshTargetType::TriangleCount;
+
+	/** Approximate Target Triangle Count, combined with mesh surface area to derive a TargetEdgeLength */
+	UPROPERTY(EditAnywhere, BlueprintReadWrite, Category = Options, meta = (UIMin = 0, ClampMin = 0, EditCondition = "TargetType == EGeometryScriptUniformRemeshTargetType::TriangleCount"))
+	int32 TargetTriangleCount = 5000;
+
+	/** Explicit Target Edge Length that is desired in the output uniform mesh */
+	UPROPERTY(EditAnywhere, BlueprintReadWrite, Category = Options, meta = (UIMin = 0, ClampMin = 0, EditCondition = "TargetType == EGeometryScriptUniformRemeshTargetType::TargetEdgeLength"))
+	float TargetEdgeLength = 1.0f;
+
+	/** Amount to adjust density in regions where the vertex weight map is greater than zero. A positive value here results in a denser mesh, a negative value results in a coarser one */
+	UPROPERTY(EditAnywhere, BlueprintReadWrite, Category = Options, meta = (UIMin = "-2.0", UIMax = "2.0"))
+	float RelativeDensity = 0.f;
+};
 
 
 UCLASS(MinimalAPI, meta = (ScriptName = "GeometryScript_Remeshing"))
@@ -161,6 +185,15 @@ public:
 		FGeometryScriptUniformRemeshOptions UniformOptions,
 		UGeometryScriptDebug* Debug = nullptr);
 
+	UFUNCTION(BlueprintCallable, Category = "GeometryScript|Simplification", meta = (ScriptMethod, HidePin = "Debug"))
+	static UE_API UPARAM(DisplayName = "Target Mesh") UDynamicMesh*
+	ApplyAdaptiveRemesh(
+		UDynamicMesh* TargetMesh,
+		FGeometryScriptRemeshOptions RemeshOptions,
+		FGeometryScriptAdaptiveRemeshOptions AdaptiveOptions,
+		FGeometryScriptWeightMapHandle WeightMapHandle,
+		UGeometryScriptDebug* Debug = nullptr);
+
 };
 
 #undef UE_API
