Added geometry script support for driving simplification error tolerance or edge lengths via weight maps

diff --git a/Engine/Plugins/Runtime/GeometryScripting/Source/GeometryScriptingCore/Private/MeshSimplifyFunctions.cpp b/Engine/Plugins/Runtime/GeometryScripting/Source/GeometryScriptingCore/Private/MeshSimplifyFunctions.cpp
--- a/Engine/Plugins/Runtime/GeometryScripting/Source/GeometryScriptingCore/Private/MeshSimplifyFunctions.cpp
+++ b/Engine/Plugins/Runtime/GeometryScripting/Source/GeometryScriptingCore/Private/MeshSimplifyFunctions.cpp
@@ -198,8 +198,31 @@ namespace UE::PrivateSimplifyHelper
 			Simplifier.CollapseMode = SimplificationType::ESimplificationCollapseModes::MinimalExistingVertexError;
 		}
 
+		TFunction<double(const FDynamicMesh3& Mesh, int VertexA, int VertexB)> EdgeMetricScale = nullptr;
+		if (EditMesh.HasAttributes() && Options.WeightMapHandle.IsValid() && Options.WeightMapHandle.WeightMapAttributeLayerIndex < EditMesh.Attributes()->NumWeightLayers())
+		{
+			const FDynamicMeshWeightAttribute* const AttributeLayer = EditMesh.Attributes()->GetWeightLayer(Options.WeightMapHandle.WeightMapAttributeLayerIndex);
+			EdgeMetricScale =
+				[AttributeLayer, RelativeDensity = Options.RelativeDensity](const FDynamicMesh3& Mesh, int VertexA, int VertexB) -> double
+				{
+					float WeightValueA;
+					AttributeLayer->GetValue(VertexA, &WeightValueA);
+					WeightValueA = FMath::Clamp(WeightValueA, 0.0f, 1.0f);
+
+					float WeightValueB;
+					AttributeLayer->GetValue(VertexB, &WeightValueB);
+					WeightValueB = FMath::Clamp(WeightValueB, 0.0f, 1.0f);
+
+					// choose the weight that will lead to a larger measure, so we choose the more conservative of the two weights
+					const float UseWeight = RelativeDensity < 0 ? FMath::Max(WeightValueA, WeightValueB) : FMath::Min(WeightValueA, WeightValueB);
+
+					return FMath::Pow(0.5, -RelativeDensity * (double)UseWeight);
+				};
+		}
+
 		if (ProjectionTarget != nullptr && GeometricTolerance >= 0)
 		{
+			Simplifier.CustomGeometricErrorScaleF = EdgeMetricScale;
 			Simplifier.GeometricErrorConstraint = SimplificationType::EGeometricErrorCriteria::PredictedPointToProjectionTarget;
 			Simplifier.GeometricErrorTolerance = GeometricTolerance;
 		}
@@ -214,6 +237,7 @@ namespace UE::PrivateSimplifyHelper
 		}
 		else // SimplifyTargetMode == ESimplifyTargetMode::EdgeLength
 		{
+			Simplifier.CustomEdgeLengthScaleF = EdgeMetricScale;
 			Simplifier.SimplifyToEdgeLength(TargetMinEdgeLength);
 		}
 
@@ -445,6 +469,23 @@ UDynamicMesh* UGeometryScriptLibrary_MeshSimplifyFunctions::ApplyClusterSimplify
 			SimplifyOptions.PreserveEdges.Material = MeshClusterSimplify::FSimplifyOptions::EConstraintLevel::Free;
 			SimplifyOptions.bTransferAttributes = false;
 		}
+
+		if (EditMesh.HasAttributes() && Options.WeightMapHandle.IsValid() && Options.WeightMapHandle.WeightMapAttributeLayerIndex < EditMesh.Attributes()->NumWeightLayers())
+		{
+			const FDynamicMeshWeightAttribute* const AttributeLayer = EditMesh.Attributes()->GetWeightLayer(Options.WeightMapHandle.WeightMapAttributeLayerIndex);
+			TFunction<double(int Vertex)> ClusterEdgeTargetLengthScale =
+				[AttributeLayer,
+				RelativeDensity = Options.RelativeDensity](int Vertex) -> double
+				{
+					float WeightValue;
+					AttributeLayer->GetValue(Vertex, &WeightValue);
+					WeightValue = FMath::Clamp(WeightValue, 0.0f, 1.0f);
+
+					// Note: We use positive RelativeDensity here instead of negative, b/c cluster simplifier scales target length rather measured length
+					return FMath::Pow(0.5, RelativeDensity * WeightValue);
+				};
+			SimplifyOptions.OptionalTargetEdgeLengthScale = ClusterEdgeTargetLengthScale;
+		}
 		
 		FDynamicMesh3 ResultMesh;
 		MeshClusterSimplify::Simplify(EditMesh, ResultMesh, SimplifyOptions);
diff --git a/Engine/Plugins/Runtime/GeometryScripting/Source/GeometryScriptingCore/Public/GeometryScript/MeshSimplifyFunctions.h b/Engine/Plugins/Runtime/GeometryScripting/Source/GeometryScriptingCore/Public/GeometryScript/MeshSimplifyFunctions.h
--- a/Engine/Plugins/Runtime/GeometryScripting/Source/GeometryScriptingCore/Public/GeometryScript/MeshSimplifyFunctions.h
+++ b/Engine/Plugins/Runtime/GeometryScripting/Source/GeometryScriptingCore/Public/GeometryScript/MeshSimplifyFunctions.h
@@ -73,6 +73,14 @@ public:
 	UPROPERTY(EditAnywhere, BlueprintReadWrite, Category = Options)
 	bool bRetainQuadricMemory = false;
 
+	// Optional weight map to control simplification error tolerances and/or edge lengths
+	UPROPERTY(EditAnywhere, BlueprintReadWrite, Category = Options)
+	FGeometryScriptWeightMapHandle WeightMapHandle;
+
+	/** Amount to adjust density in regions where the vertex weight map is greater than zero. A positive value results in more precision / less simplification in weighted regions, while a negative value results in more simplification */
+	UPROPERTY(EditAnywhere, BlueprintReadWrite, Category = Options, meta = (UIMin = "-2.0", UIMax = "2.0"))
+	float RelativeDensity = 0.f;
+
 	/** If enabled, the simplified mesh is automatically compacted to remove gaps in the index space. This is expensive and can be disabled by advanced users. */
 	UPROPERTY(EditAnywhere, BlueprintReadWrite, Category = Options)
 	bool bAutoCompact = true;
@@ -133,6 +141,15 @@ public:
 	// Options to constrain simplification of different mesh edge types, to optionally help preserve mesh boundaries and seams
 	UPROPERTY(EditAnywhere, BlueprintReadWrite, Category = Options)
 	FGeometryScriptClusterSimplifyEdgeConstraintOptions EdgeConstraints;
+
+	// Optional weight map to control simplification error tolerances and/or edge lengths
+	UPROPERTY(EditAnywhere, BlueprintReadWrite, Category = Options)
+	FGeometryScriptWeightMapHandle WeightMapHandle;
+
+	/** Amount to adjust density in regions where the vertex weight map is greater than zero. A positive value results in less simplification in weighted regions, while a negative value results in more */
+	UPROPERTY(EditAnywhere, BlueprintReadWrite, Category = Options, meta = (UIMin = "-2.0", UIMax = "2.0"))
+	float RelativeDensity = 0.f;
+
 };
 
 
