change weight channel find method to pass name by value, so that it is directly settable on the node in BP

diff --git a/Engine/Plugins/Runtime/GeometryScripting/Source/GeometryScriptingCore/Private/MeshQueryFunctions.cpp b/Engine/Plugins/Runtime/GeometryScripting/Source/GeometryScriptingCore/Private/MeshQueryFunctions.cpp
--- a/Engine/Plugins/Runtime/GeometryScripting/Source/GeometryScriptingCore/Private/MeshQueryFunctions.cpp
+++ b/Engine/Plugins/Runtime/GeometryScripting/Source/GeometryScriptingCore/Private/MeshQueryFunctions.cpp
@@ -824,7 +824,7 @@ int UGeometryScriptLibrary_MeshQueryFunctions::GetNumExtendedPolygroupLayers( UD
 	});
 }
 
-FGeometryScriptWeightMapHandle UGeometryScriptLibrary_MeshQueryFunctions::GetWeightMapHandle(UDynamicMesh* TargetMesh, const FName& Name, EGeometryScriptSearchOutcomePins& FoundWeights)
+FGeometryScriptWeightMapHandle UGeometryScriptLibrary_MeshQueryFunctions::GetWeightMapHandle(UDynamicMesh* TargetMesh, FName Name, EGeometryScriptSearchOutcomePins& FoundWeights)
 {
 	FGeometryScriptWeightMapHandle Handle;
 	Handle.WeightMapAttributeLayerIndex = SimpleMeshQuery<int32>(TargetMesh, (int32)-1, [&Name](const FDynamicMesh3& Mesh) -> int32
diff --git a/Engine/Plugins/Runtime/GeometryScripting/Source/GeometryScriptingCore/Public/GeometryScript/MeshQueryFunctions.h b/Engine/Plugins/Runtime/GeometryScripting/Source/GeometryScriptingCore/Public/GeometryScript/MeshQueryFunctions.h
--- a/Engine/Plugins/Runtime/GeometryScripting/Source/GeometryScriptingCore/Public/GeometryScript/MeshQueryFunctions.h
+++ b/Engine/Plugins/Runtime/GeometryScripting/Source/GeometryScriptingCore/Public/GeometryScript/MeshQueryFunctions.h
@@ -483,7 +483,7 @@ public:
 	//
 
 	UFUNCTION(BlueprintCallable, Category = "GeometryScript|MeshQueries", meta = (ScriptMethod, DisplayName = "Get Weight Map Handle By Name", ExpandEnumAsExecs = "FoundWeights"))
-	static UE_API UPARAM(DisplayName = "Weight Map Handle") FGeometryScriptWeightMapHandle GetWeightMapHandle(UDynamicMesh* TargetMesh, const FName& Name, EGeometryScriptSearchOutcomePins& FoundWeights);
+	static UE_API UPARAM(DisplayName = "Weight Map Handle") FGeometryScriptWeightMapHandle GetWeightMapHandle(UDynamicMesh* TargetMesh, FName Name, EGeometryScriptSearchOutcomePins& FoundWeights);
 
 };
 
