add regularization controls to the geometry script simplification options

diff --git a/Engine/Plugins/Runtime/GeometryScripting/Source/GeometryScriptingCore/Private/MeshSimplifyFunctions.cpp b/Engine/Plugins/Runtime/GeometryScripting/Source/GeometryScriptingCore/Private/MeshSimplifyFunctions.cpp
--- a/Engine/Plugins/Runtime/GeometryScripting/Source/GeometryScriptingCore/Private/MeshSimplifyFunctions.cpp
+++ b/Engine/Plugins/Runtime/GeometryScripting/Source/GeometryScriptingCore/Private/MeshSimplifyFunctions.cpp
@@ -199,6 +199,7 @@ namespace UE::PrivateSimplifyHelper
 		}
 
 		TFunction<double(const FDynamicMesh3& Mesh, int VertexA, int VertexB)> EdgeMetricScale = nullptr;
+		TFunction<double(const FDynamicMesh3& Mesh, int Vertex)> VertexMetricScale = nullptr;
 		if (EditMesh.HasAttributes() && Options.WeightMapHandle.IsValid() && Options.WeightMapHandle.WeightMapAttributeLayerIndex < EditMesh.Attributes()->NumWeightLayers())
 		{
 			const FDynamicMeshWeightAttribute* const AttributeLayer = EditMesh.Attributes()->GetWeightLayer(Options.WeightMapHandle.WeightMapAttributeLayerIndex);
@@ -219,6 +220,19 @@ namespace UE::PrivateSimplifyHelper
 					return FMath::Pow(0.5, -RelativeDensity * (double)UseWeight);
 				};
 		}
+		if (EditMesh.HasAttributes() && Options.RegularizeWeightMapHandle.IsValid() && Options.RegularizeWeightMapHandle.WeightMapAttributeLayerIndex < EditMesh.Attributes()->NumWeightLayers())
+		{
+			const FDynamicMeshWeightAttribute* const AttributeLayer = EditMesh.Attributes()->GetWeightLayer(Options.RegularizeWeightMapHandle.WeightMapAttributeLayerIndex);
+			VertexMetricScale =
+				[AttributeLayer, RelativeDensity = Options.RegularizeRelativeDensity](const FDynamicMesh3& Mesh, int Vertex) -> double
+				{
+					float WeightValue;
+					AttributeLayer->GetValue(Vertex, &WeightValue);
+					WeightValue = FMath::Clamp(WeightValue, 0.0f, 1.0f);
+
+					return FMath::Pow(0.5, -RelativeDensity * (double)WeightValue);
+				};
+		}
 
 		if (ProjectionTarget != nullptr && GeometricTolerance >= 0)
 		{
@@ -227,6 +241,12 @@ namespace UE::PrivateSimplifyHelper
 			Simplifier.GeometricErrorTolerance = GeometricTolerance;
 		}
 
+		Simplifier.RegularizeWeight = Options.RegularizeWeight;
+		if (Options.RegularizeWeight > 0)
+		{
+			Simplifier.CustomRegularizationScaleF = VertexMetricScale;
+		}
+
 		if (SimplifyTargetMode == ESimplifyTargetMode::TriangleCount)
 		{
 			Simplifier.SimplifyToTriangleCount(FMath::Max(1, TargetCount));
diff --git a/Engine/Plugins/Runtime/GeometryScripting/Source/GeometryScriptingCore/Public/GeometryScript/MeshSimplifyFunctions.h b/Engine/Plugins/Runtime/GeometryScripting/Source/GeometryScriptingCore/Public/GeometryScript/MeshSimplifyFunctions.h
--- a/Engine/Plugins/Runtime/GeometryScripting/Source/GeometryScriptingCore/Public/GeometryScript/MeshSimplifyFunctions.h
+++ b/Engine/Plugins/Runtime/GeometryScripting/Source/GeometryScriptingCore/Public/GeometryScript/MeshSimplifyFunctions.h
@@ -81,6 +81,18 @@ public:
 	UPROPERTY(EditAnywhere, BlueprintReadWrite, Category = Options, meta = (UIMin = "-2.0", UIMax = "2.0"))
 	float RelativeDensity = 0.f;
 
+	/** A small regularize weight can improve triangle quality in flat regions, and locally control the amount of simplification if Regularize Weight Map is provided */
+	UPROPERTY(EditAnywhere, BlueprintReadWrite, Category = Options, meta = (ClampMin = "0.0", ClampMax = "100.0", UIMax = ".01"))
+	float RegularizeWeight = 0.f;
+
+	/** Optional weight map to locally control regularization weight */
+	UPROPERTY(EditAnywhere, BlueprintReadWrite, Category = Options)
+	FGeometryScriptWeightMapHandle RegularizeWeightMapHandle;
+
+	/** Amount to adjust regularization weights in regions where the vertex weight map is greater than zero. A positive value results in more weight / less simplification in weighted regions, while a negative value results in more simplification */
+	UPROPERTY(EditAnywhere, BlueprintReadWrite, Category = Options, meta = (UIMin = "-2.0", UIMax = "2.0"))
+	float RegularizeRelativeDensity = 0.f;
+
 	/** If enabled, the simplified mesh is automatically compacted to remove gaps in the index space. This is expensive and can be disabled by advanced users. */
 	UPROPERTY(EditAnywhere, BlueprintReadWrite, Category = Options)
 	bool bAutoCompact = true;
