remove concept of scaling the regularization term from mesh simplification; use the quadric error scaling instead
+ make quadric error scaling directly scale the quadric, instead of just scaling the computed error
+ compensate for scaling when applying "memoryless" updates

change the interface so that geometry script can specify separate maps for each term, and the tool can toggle which terms are scaled

diff --git a/Engine/Plugins/Runtime/GeometryProcessing/Source/DynamicMesh/Private/MeshSimplification.cpp b/Engine/Plugins/Runtime/GeometryProcessing/Source/DynamicMesh/Private/MeshSimplification.cpp
--- a/Engine/Plugins/Runtime/GeometryProcessing/Source/DynamicMesh/Private/MeshSimplification.cpp
+++ b/Engine/Plugins/Runtime/GeometryProcessing/Source/DynamicMesh/Private/MeshSimplification.cpp
@@ -157,6 +157,16 @@ void TMeshSimplification<QuadricErrorType>::InitializeVertexQuadrics()
 
 	int NV = Mesh->MaxVertexID();
 	vertQuadrics.SetNum(NV);
+
+	if (!bRetainQuadricMemory && CustomQuadricErrorScaleF)
+	{
+		AppliedVertexQuadricScales.Init(1., NV);
+	}
+	else
+	{
+		AppliedVertexQuadricScales.Empty();
+	}
+
 	// tested with ParallelFor - no measurable benefit 
 	//gParallel.BlockStartEnd(0, Mesh->MaxVertexID - 1, (start_vid, end_vid) = > {
 	for (int vid : Mesh->VertexIndicesItr())
@@ -177,14 +187,14 @@ void TMeshSimplification<QuadricErrorType>::InitializeVertexQuadrics()
 			Q.Add(triAreas[tid], triQuadrics[tid]);
 		}
 
-		ApplyRegularizationToVertexQuadric(Q, vid, AreaSum, Normal);
+		ApplyScalingAndRegularizationToVertexQuadric(Q, vid, AreaSum, Normal);
 
 		vertQuadrics[vid] = Q;
 	}
 }
 
 template <typename QuadricErrorType>
-void TMeshSimplification<QuadricErrorType>::ApplyRegularizationToVertexQuadric(FQuadricErrorType& Q, int32 VID, double AreaSum, FVector3d Normal)
+void TMeshSimplification<QuadricErrorType>::ApplyScalingAndRegularizationToVertexQuadric(FQuadricErrorType& Q, int32 VID, double AreaSum, FVector3d Normal)
 {
 	if (RegularizeWeight > 0.)
 	{
@@ -203,13 +213,17 @@ void TMeshSimplification<QuadricErrorType>::ApplyRegularizationToVertexQuadric(F
 		{
 			Q.AddPositionRegularizer(VertPos, UseRegularizeWeight * AreaSum);
 		}
-
-		// Note we scale the *entire* quadric after adding the regularization term
-		// to ensure that the 'actual' QEM error remains the dominant term even at
-		// large scales
-		if (CustomRegularizationScaleF)
+	}
+	if (CustomQuadricErrorScaleF)
+	{
+		double QEMScale = FMath::Max(UE_DOUBLE_SMALL_NUMBER, CustomQuadricErrorScaleF(*Mesh, VID));
+		if (FMath::IsFinite(QEMScale))
 		{
-			Q.Scale(CustomRegularizationScaleF(*Mesh, VID));
+			if (!AppliedVertexQuadricScales.IsEmpty())
+			{
+				AppliedVertexQuadricScales[VID] = QEMScale;
+			}
+			Q.Scale(QEMScale);
 		}
 	}
 }
@@ -223,15 +237,21 @@ QuadricErrorType TMeshSimplification<QuadricErrorType>::AssembleEdgeQuadric(cons
 	if (!bRetainQuadricMemory)
 	{ 
 		// the edge.Tri faces are double counted. Remove one.
+		// if using scaled vertex quadrics, removed the smaller of the two
+		double RemoveScale = 1.;
+		if (!AppliedVertexQuadricScales.IsEmpty())
+		{
+			RemoveScale = FMath::Min(AppliedVertexQuadricScales[edge.Vert.A], AppliedVertexQuadricScales[edge.Vert.B]);
+		}
 		const FIndex2i& Tris = edge.Tri;
 		if (Tris.A != FDynamicMesh3::InvalidID)
 		{
-			EdgeQuadric.Add(-triAreas[Tris.A], triQuadrics[Tris.A]);
+			EdgeQuadric.Add(-triAreas[Tris.A] * RemoveScale, triQuadrics[Tris.A]);
 		}
 
 		if (Tris.B != FDynamicMesh3::InvalidID)
 		{
-			EdgeQuadric.Add(-triAreas[Tris.B], triQuadrics[Tris.B]);
+			EdgeQuadric.Add(-triAreas[Tris.B] * RemoveScale, triQuadrics[Tris.B]);
 		}
 	}
 	if (bAllowSeamCollapse)
@@ -277,11 +297,6 @@ void TMeshSimplification<QuadricErrorType>::InitializeQueue()
 		FQuadricErrorType Q = AssembleEdgeQuadric(edge);
 		FVector3d opt = OptimalPoint(eid, Q, edge.Vert.A, edge.Vert.B);
 		EdgeErrors[eid] = { (float)Q.Evaluate(opt), eid };
-		if (CustomQuadricErrorScaleF)
-		{
-			double ErrorScale = CustomQuadricErrorScaleF(*Mesh, edge.Vert.A, edge.Vert.B);
-			EdgeErrors[eid].error *= (float)(ErrorScale * ErrorScale);
-		}
 		EdgeQuadrics[eid] = QEdge(eid, Q, opt);
 	}
 
@@ -554,8 +569,9 @@ void TMeshSimplification<QuadricErrorType>::UpdateNeighborhood(const FDynamicMes
 				if (tri_vids[i] == kvid) continue;
 
 				// correct the adjacent vertQuadrics
-				vertQuadrics[tri_vids[i]].Add(-OldtriArea, OldtriQuadric); // subtract old quadric
-				vertQuadrics[tri_vids[i]].Add(NewtriArea, NewtriQuadric); // add new quadric
+				double VertexScale = AppliedVertexQuadricScales.IsEmpty() ? 1. : AppliedVertexQuadricScales[tri_vids[i]];
+				vertQuadrics[tri_vids[i]].Add(-OldtriArea * VertexScale, OldtriQuadric); // subtract old quadric
+				vertQuadrics[tri_vids[i]].Add(NewtriArea * VertexScale, NewtriQuadric); // add new quadric
 			}
 		}
 
@@ -569,7 +585,8 @@ void TMeshSimplification<QuadricErrorType>::UpdateNeighborhood(const FDynamicMes
 					FQuadricErrorType oldQuadric = triQuadrics[removedTris[i]];
 
 					// subtract the quadric from the opposing vert
-					vertQuadrics[opposingVerts[i]].Add(-oldArea, oldQuadric);
+					double VertexScale = AppliedVertexQuadricScales.IsEmpty() ? 1. : AppliedVertexQuadricScales[opposingVerts[i]];
+					vertQuadrics[opposingVerts[i]].Add(-oldArea * VertexScale, oldQuadric);
 
 					// zero out the quadric & area for the removed tris.
 					triQuadrics[removedTris[i]] = FQuadricErrorType::Zero();
@@ -593,7 +610,7 @@ void TMeshSimplification<QuadricErrorType>::UpdateNeighborhood(const FDynamicMes
 				}
 				vertQuadric.Add(triAreas[tid], triQuadrics[tid]);
 			}
-			ApplyRegularizationToVertexQuadric(vertQuadric, kvid, AreaSum, Normal);
+			ApplyScalingAndRegularizationToVertexQuadric(vertQuadric, kvid, AreaSum, Normal);
 			vertQuadrics[kvid] = vertQuadric;
 		}
 	}
@@ -611,11 +628,6 @@ void TMeshSimplification<QuadricErrorType>::UpdateNeighborhood(const FDynamicMes
 			QuadricErrorType Q = AssembleEdgeQuadric(ne);
 			FVector3d opt = OptimalPoint(eid, Q, ne.Vert.A, ne.Vert.B);
 			float err = (float)Q.Evaluate(opt);
-			if (CustomQuadricErrorScaleF)
-			{
-				double ErrorScale = CustomQuadricErrorScaleF(*Mesh, ne.Vert.A, ne.Vert.B);
-				err *= (float)(ErrorScale * ErrorScale);
-			}
 			EdgeQuadrics[eid] = QEdge(eid, Q, opt);
 			if (EdgeQueue.Contains(eid))
 			{
@@ -656,11 +668,6 @@ void TMeshSimplification<QuadricErrorType>::UpdateNeighborhood(const FDynamicMes
 
 			FVector3d opt = OptimalPoint(eid, Q, edgeData.Vert[0], edgeData.Vert[1]);
 			float err = (float)Q.Evaluate(opt);
-			if (CustomQuadricErrorScaleF)
-			{
-				double ErrorScale = CustomQuadricErrorScaleF(*Mesh, edgeData.Vert.A, edgeData.Vert.B);
-				err *= (float)(ErrorScale * ErrorScale);
-			}
 			EdgeQuadrics[eid] = QEdge(eid, Q, opt);
 			if (EdgeQueue.Contains(eid))
 			{
diff --git a/Engine/Plugins/Runtime/GeometryProcessing/Source/DynamicMesh/Public/MeshSimplification.h b/Engine/Plugins/Runtime/GeometryProcessing/Source/DynamicMesh/Public/MeshSimplification.h
--- a/Engine/Plugins/Runtime/GeometryProcessing/Source/DynamicMesh/Public/MeshSimplification.h
+++ b/Engine/Plugins/Runtime/GeometryProcessing/Source/DynamicMesh/Public/MeshSimplification.h
@@ -122,10 +122,9 @@ public:
 
 	// Custom scales to apply to *measured* errors and edge lengths. Note this is in effect the inverse of scaling the tolerances / target lengths.
 	TFunction<double(const FDynamicMesh3&, int VertexA, int VertexB)> CustomGeometricErrorScaleF;
-	TFunction<double(const FDynamicMesh3&, int VertexA, int VertexB)> CustomQuadricErrorScaleF;
 	TFunction<double(const FDynamicMesh3&, int VertexA, int VertexB)> CustomEdgeLengthScaleF;
-	TFunction<double(const FDynamicMesh3&, int Vertex)> CustomRegularizationScaleF;
-
+	// Custom per-vertex quadric error scaling. Values should be greater than zero.
+	TFunction<double(const FDynamicMesh3&, int Vertex)> CustomQuadricErrorScaleF;
 
 	TMeshSimplification(FDynamicMesh3* m) : FMeshRefinerBase(m)
 	{
@@ -247,6 +246,7 @@ protected:
 	TArray<double> triAreas;
 	TArray<FVector3d> triNormals; // this array is only populate if needed; i.e. if RegularizeUsesNormals() returns true
 	TArray<FQuadricErrorType> triQuadrics;
+	TArray<double> AppliedVertexQuadricScales; // needed to properly update scaled vertex quadrics in the memoryless update case
 	DYNAMICMESH_API virtual void InitializeTriQuadrics();
 
 	bool RegularizeUsesNormals() const
@@ -461,7 +461,7 @@ protected:
 
 private:
 
-	void ApplyRegularizationToVertexQuadric(FQuadricErrorType& Q, int32 VID, double AreaSum, FVector3d Normal);
+	void ApplyScalingAndRegularizationToVertexQuadric(FQuadricErrorType& Q, int32 VID, double AreaSum, FVector3d Normal);
 
 };
 
diff --git a/Engine/Plugins/Runtime/GeometryScripting/Source/GeometryScriptingCore/Private/MeshSimplifyFunctions.cpp b/Engine/Plugins/Runtime/GeometryScripting/Source/GeometryScriptingCore/Private/MeshSimplifyFunctions.cpp
--- a/Engine/Plugins/Runtime/GeometryScripting/Source/GeometryScriptingCore/Private/MeshSimplifyFunctions.cpp
+++ b/Engine/Plugins/Runtime/GeometryScripting/Source/GeometryScriptingCore/Private/MeshSimplifyFunctions.cpp
@@ -198,54 +198,60 @@ namespace UE::PrivateSimplifyHelper
 			Simplifier.CollapseMode = SimplificationType::ESimplificationCollapseModes::MinimalExistingVertexError;
 		}
 
-		TFunction<double(const FDynamicMesh3& Mesh, int VertexA, int VertexB)> EdgeMetricScale = nullptr;
-		TFunction<double(const FDynamicMesh3& Mesh, int Vertex)> VertexMetricScale = nullptr;
-		if (EditMesh.HasAttributes() && Options.WeightMapHandle.IsValid() && Options.WeightMapHandle.WeightMapAttributeLayerIndex < EditMesh.Attributes()->NumWeightLayers())
+		TFunction<double(const FDynamicMesh3& Mesh, int VertexA, int VertexB)> EdgeLengthScaleF = nullptr;
+		TFunction<double(const FDynamicMesh3& Mesh, int VertexA, int VertexB)> GeometricErrorScale = nullptr;
+		TFunction<double(const FDynamicMesh3& Mesh, int Vertex)> QEMErrorScale = nullptr;
+		if (EditMesh.HasAttributes())
 		{
-			const FDynamicMeshWeightAttribute* const AttributeLayer = EditMesh.Attributes()->GetWeightLayer(Options.WeightMapHandle.WeightMapAttributeLayerIndex);
-			EdgeMetricScale =
-				[AttributeLayer, RelativeDensity = Options.RelativeDensity](const FDynamicMesh3& Mesh, int VertexA, int VertexB) -> double
+			auto CreateWeightMapScaleF = [&EditMesh](int32 WeightMapAttributeLayerIndex, float RelativeDensity) -> TFunction<double(const FDynamicMesh3 & Mesh, int VertexA, int VertexB)>
+			{
+				TFunction<double(const FDynamicMesh3& Mesh, int VertexA, int VertexB)> ToRet = nullptr;
+				if (WeightMapAttributeLayerIndex >= 0 && WeightMapAttributeLayerIndex < EditMesh.Attributes()->NumWeightLayers())
 				{
-					float WeightValueA;
-					AttributeLayer->GetValue(VertexA, &WeightValueA);
-					WeightValueA = FMath::Clamp(WeightValueA, 0.0f, 1.0f);
-
-					float WeightValueB;
-					AttributeLayer->GetValue(VertexB, &WeightValueB);
-					WeightValueB = FMath::Clamp(WeightValueB, 0.0f, 1.0f);
-
-					// choose the weight that will lead to a larger measure, so we choose the more conservative of the two weights
-					const float UseWeight = RelativeDensity < 0 ? FMath::Max(WeightValueA, WeightValueB) : FMath::Min(WeightValueA, WeightValueB);
-
-					return FMath::Pow(0.5, -RelativeDensity * (double)UseWeight);
-				};
-		}
-		if (EditMesh.HasAttributes() && Options.RegularizeWeightMapHandle.IsValid() && Options.RegularizeWeightMapHandle.WeightMapAttributeLayerIndex < EditMesh.Attributes()->NumWeightLayers())
-		{
-			const FDynamicMeshWeightAttribute* const AttributeLayer = EditMesh.Attributes()->GetWeightLayer(Options.RegularizeWeightMapHandle.WeightMapAttributeLayerIndex);
-			VertexMetricScale =
-				[AttributeLayer, RelativeDensity = Options.RegularizeRelativeDensity](const FDynamicMesh3& Mesh, int Vertex) -> double
+					const FDynamicMeshWeightAttribute* const AttributeLayer = EditMesh.Attributes()->GetWeightLayer(WeightMapAttributeLayerIndex);
+					ToRet = [AttributeLayer, RelativeDensity](const FDynamicMesh3& Mesh, int VertexA, int VertexB) -> double
+					{
+						float WeightValueA;
+						AttributeLayer->GetValue(VertexA, &WeightValueA);
+						WeightValueA = FMath::Clamp(WeightValueA, 0.0f, 1.0f);
+
+						float WeightValueB;
+						AttributeLayer->GetValue(VertexB, &WeightValueB);
+						WeightValueB = FMath::Clamp(WeightValueB, 0.0f, 1.0f);
+
+						// choose the weight that will lead to a larger measure, so we choose the more conservative of the two weights
+						const float UseWeight = RelativeDensity < 0 ? FMath::Max(WeightValueA, WeightValueB) : FMath::Min(WeightValueA, WeightValueB);
+
+						return FMath::Pow(0.5, -RelativeDensity * (double)UseWeight);
+					};
+				}
+				return ToRet;
+			};
+			EdgeLengthScaleF = CreateWeightMapScaleF(Options.EdgeLengthWeightMap.Handle.WeightMapAttributeLayerIndex, Options.EdgeLengthWeightMap.RelativeDensity);
+			GeometricErrorScale = CreateWeightMapScaleF(Options.GeometricToleranceWeightMap.Handle.WeightMapAttributeLayerIndex, Options.GeometricToleranceWeightMap.RelativeDensity);
+			if (Options.QuadricErrorWeightMap.Handle.IsValid() && Options.QuadricErrorWeightMap.Handle.WeightMapAttributeLayerIndex < EditMesh.Attributes()->NumWeightLayers())
+			{
+				const FDynamicMeshWeightAttribute* const AttributeLayer = EditMesh.Attributes()->GetWeightLayer(Options.QuadricErrorWeightMap.Handle.WeightMapAttributeLayerIndex);
+				QEMErrorScale = [AttributeLayer, RelativeDensity = Options.QuadricErrorWeightMap.RelativeDensity](const FDynamicMesh3& Mesh, int Vertex) -> double
 				{
 					float WeightValue;
 					AttributeLayer->GetValue(Vertex, &WeightValue);
 					WeightValue = FMath::Clamp(WeightValue, 0.0f, 1.0f);
 
-					return FMath::Pow(0.5, -RelativeDensity * (double)WeightValue);
+					// We use a squared error metric for the quadric error scale to get a similar scale in effect as the tolerance/edge-len scales, since it is a squared error term
+					return FMath::Pow(0.5, -2. * RelativeDensity * (double)WeightValue);
 				};
+			}	
 		}
 
 		if (ProjectionTarget != nullptr && GeometricTolerance >= 0)
 		{
-			Simplifier.CustomGeometricErrorScaleF = EdgeMetricScale;
+			Simplifier.CustomGeometricErrorScaleF = GeometricErrorScale;
 			Simplifier.GeometricErrorConstraint = SimplificationType::EGeometricErrorCriteria::PredictedPointToProjectionTarget;
 			Simplifier.GeometricErrorTolerance = GeometricTolerance;
 		}
-
+		Simplifier.CustomQuadricErrorScaleF = QEMErrorScale;
 		Simplifier.RegularizeWeight = Options.RegularizeWeight;
-		if (Options.RegularizeWeight > 0)
-		{
-			Simplifier.CustomRegularizationScaleF = VertexMetricScale;
-		}
 
 		if (SimplifyTargetMode == ESimplifyTargetMode::TriangleCount)
 		{
@@ -257,10 +263,11 @@ namespace UE::PrivateSimplifyHelper
 		}
 		else // SimplifyTargetMode == ESimplifyTargetMode::EdgeLength
 		{
-			Simplifier.CustomEdgeLengthScaleF = EdgeMetricScale;
+			Simplifier.CustomEdgeLengthScaleF = EdgeLengthScaleF;
 			Simplifier.SimplifyToEdgeLength(TargetMinEdgeLength);
 		}
 
+
 		if (Options.bAutoCompact)
 		{
 			EditMesh.CompactInPlace();
diff --git a/Engine/Plugins/Runtime/GeometryScripting/Source/GeometryScriptingCore/Public/GeometryScript/MeshSimplifyFunctions.h b/Engine/Plugins/Runtime/GeometryScripting/Source/GeometryScriptingCore/Public/GeometryScript/MeshSimplifyFunctions.h
--- a/Engine/Plugins/Runtime/GeometryScripting/Source/GeometryScriptingCore/Public/GeometryScript/MeshSimplifyFunctions.h
+++ b/Engine/Plugins/Runtime/GeometryScripting/Source/GeometryScriptingCore/Public/GeometryScript/MeshSimplifyFunctions.h
@@ -48,6 +48,24 @@ enum class EGeometryScriptRemoveMeshSimplificationType : uint8
 	AttributeAware = 2
 };
 
+// Specify a weight map for controlling relative density of a mesh
+USTRUCT(BlueprintType)
+struct FGeometryScriptWeightMapDensity
+{
+	GENERATED_BODY()
+public:
+	
+	/** The weight map to reference */
+	UPROPERTY(EditAnywhere, BlueprintReadWrite, Category = Options)
+	FGeometryScriptWeightMapHandle Handle;
+
+	/** 
+	 * Controls the effect of the weight map:
+	 * Positive values increase local mesh density for larger weights, negative values decrease density, and zero gives no effect
+	 */
+	UPROPERTY(EditAnywhere, BlueprintReadWrite, Category = Options, meta = (UIMin = "-3.0", UIMax = "3.0"))
+	float RelativeDensity = 0.f;
+};
 
 USTRUCT(Blueprintable)
 struct FGeometryScriptSimplifyMeshOptions
@@ -73,25 +91,21 @@ public:
 	UPROPERTY(EditAnywhere, BlueprintReadWrite, Category = Options)
 	bool bRetainQuadricMemory = false;
 
-	// Optional weight map to control simplification error tolerances and/or edge lengths
+	/** Optional weight map to scale measured edge lengths */
 	UPROPERTY(EditAnywhere, BlueprintReadWrite, Category = Options)
-	FGeometryScriptWeightMapHandle WeightMapHandle;
+	FGeometryScriptWeightMapDensity EdgeLengthWeightMap;
 
-	/** Amount to adjust density in regions where the vertex weight map is greater than zero. A positive value results in more precision / less simplification in weighted regions, while a negative value results in more simplification */
-	UPROPERTY(EditAnywhere, BlueprintReadWrite, Category = Options, meta = (UIMin = "-2.0", UIMax = "2.0"))
-	float RelativeDensity = 0.f;
-
-	/** A small regularize weight can improve triangle quality in flat regions, and locally control the amount of simplification if Regularize Weight Map is provided */
-	UPROPERTY(EditAnywhere, BlueprintReadWrite, Category = Options, meta = (ClampMin = "0.0", ClampMax = "100.0", UIMax = ".01"))
-	float RegularizeWeight = 0.f;
+	/** Optional weight map to scale measured geometric errors */
+	UPROPERTY(EditAnywhere, BlueprintReadWrite, Category = Options)
+	FGeometryScriptWeightMapDensity GeometricToleranceWeightMap;
 
-	/** Optional weight map to locally control regularization weight */
+	/** Optional weight map to scale measured quadric error */
 	UPROPERTY(EditAnywhere, BlueprintReadWrite, Category = Options)
-	FGeometryScriptWeightMapHandle RegularizeWeightMapHandle;
+	FGeometryScriptWeightMapDensity QuadricErrorWeightMap;
 
-	/** Amount to adjust regularization weights in regions where the vertex weight map is greater than zero. A positive value results in more weight / less simplification in weighted regions, while a negative value results in more simplification */
-	UPROPERTY(EditAnywhere, BlueprintReadWrite, Category = Options, meta = (UIMin = "-2.0", UIMax = "2.0"))
-	float RegularizeRelativeDensity = 0.f;
+	/** Setting a very small RegularizeWeight value (e.g., .00001) can improve triangle quality and the effect of the QuadricErrorWeightMap in flat regions */
+	UPROPERTY(EditAnywhere, BlueprintReadWrite, Category = Options, meta = (ClampMin = "0.0", ClampMax = "100.0", UIMax = ".01"))
+	float RegularizeWeight = 0.f;
 
 	/** If enabled, the simplified mesh is automatically compacted to remove gaps in the index space. This is expensive and can be disabled by advanced users. */
 	UPROPERTY(EditAnywhere, BlueprintReadWrite, Category = Options)
diff --git a/Engine/Plugins/Runtime/MeshModelingToolset/Source/MeshModelingToolsEditorOnly/Private/SimplifyMeshTool.cpp b/Engine/Plugins/Runtime/MeshModelingToolset/Source/MeshModelingToolsEditorOnly/Private/SimplifyMeshTool.cpp
--- a/Engine/Plugins/Runtime/MeshModelingToolset/Source/MeshModelingToolsEditorOnly/Private/SimplifyMeshTool.cpp
+++ b/Engine/Plugins/Runtime/MeshModelingToolset/Source/MeshModelingToolsEditorOnly/Private/SimplifyMeshTool.cpp
@@ -286,7 +286,7 @@ TUniquePtr<FDynamicMeshOperator> USimplifyMeshTool::MakeNewOperator()
 				return 1.0;
 			};
 
-			TFunction<double(const FDynamicMesh3& Mesh, int Vertex)> VertexMetricScale =
+			TFunction<double(const FDynamicMesh3& Mesh, int Vertex)> SquaredVertexMetricScale =
 			[WeightMapAttributeLayerIndex, RelativeDensity = WeightMapProperties->RelativeDensity](const FDynamicMesh3& Mesh, int Vertex) -> double
 			{
 				if (Mesh.HasAttributes() && Mesh.Attributes()->GetWeightLayer(WeightMapAttributeLayerIndex))
@@ -297,26 +297,25 @@ TUniquePtr<FDynamicMeshOperator> USimplifyMeshTool::MakeNewOperator()
 					AttributeLayer->GetValue(Vertex, &WeightValue);
 					WeightValue = FMath::Clamp(WeightValue, 0.0f, 1.0f);
 
-					return FMath::Pow(0.5, -RelativeDensity * (double)WeightValue);
+					// We use a squared error metric for the quadric error scale to get a similar scale in effect as the tolerance/edge-len scales, since it is a squared error term
+					return FMath::Pow(0.5, -2. * RelativeDensity * (double)WeightValue);
 				}
 
 				return 1.0;
 			};
 
-			// geometric tolerance can be used in multiple modes, so we pass it through always
-			Op->CustomGeometricErrorScaleF = EdgeMetricScale;
-			// only scale quadric error or edge lengths when they are the active target, so weights scale the most relevant metric
-			if (SimplifyProperties->TargetMode == ESimplifyTargetType::QuadricError)
+			if (bool(SimplifyProperties->ApplyWeightMapTo & (uint32)EApplyWeightMapToSimplifyMeasure::GeometricTolerance))
 			{
-				Op->CustomQuadricErrorScaleF = EdgeMetricScale;
+				Op->CustomGeometricErrorScaleF = EdgeMetricScale;
 			}
-			if (SimplifyProperties->TargetMode == ESimplifyTargetType::EdgeLength)
+			if (bool(SimplifyProperties->ApplyWeightMapTo & (uint32)EApplyWeightMapToSimplifyMeasure::QuadricError))
 			{
-				Op->CustomEdgeLengthScaleF = EdgeMetricScale;
+				Op->CustomQuadricErrorScaleF = SquaredVertexMetricScale;
 			}
-			if (Op->RegularizeWeight > 0 && SimplifyProperties->bScaleRegularizationWithWeightMap)
+			// edge scale is only useful to the op if the simplify target is edge length
+			if (SimplifyProperties->TargetMode == ESimplifyTargetType::EdgeLength && bool(SimplifyProperties->ApplyWeightMapTo & (uint32)EApplyWeightMapToSimplifyMeasure::EdgeLengths))
 			{
-				Op->CustomRegularizationScaleF = VertexMetricScale;
+				Op->CustomEdgeLengthScaleF = EdgeMetricScale;
 			}
 
 			TFunction<double(int Vertex)> ClusterEdgeTargetLengthScale =
diff --git a/Engine/Plugins/Runtime/MeshModelingToolset/Source/MeshModelingToolsEditorOnly/Public/SimplifyMeshTool.h b/Engine/Plugins/Runtime/MeshModelingToolset/Source/MeshModelingToolsEditorOnly/Public/SimplifyMeshTool.h
--- a/Engine/Plugins/Runtime/MeshModelingToolset/Source/MeshModelingToolsEditorOnly/Public/SimplifyMeshTool.h
+++ b/Engine/Plugins/Runtime/MeshModelingToolset/Source/MeshModelingToolsEditorOnly/Public/SimplifyMeshTool.h
@@ -32,6 +32,17 @@ public:
 
 
 
+UENUM(meta = (Bitflags, UseEnumValuesAsMaskValuesInEditor = "true"))
+enum class EApplyWeightMapToSimplifyMeasure : uint32
+{
+	None = 0 UMETA(Hidden),
+	QuadricError = 1 << 0,
+	GeometricTolerance = 1 << 1,
+	EdgeLengths = 1 << 2,
+	All = 0b111 UMETA(Hidden)
+};
+ENUM_CLASS_FLAGS(EApplyWeightMapToSimplifyMeasure);
+
 /**
  * Standard properties of the Simplify operation
  */
@@ -104,12 +115,10 @@ public:
 	/** Optional regularization term to favor vertices remaining near their original locations. Using a small regularization weight can give cleaner topology in flat regions. */
 	UPROPERTY(EditAnywhere, Category = Options, meta = (UIMin = "0.0", UIMax = "0.01", ClampMin = "0.0", ClampMax = "100.0",
 		EditCondition = "SimplifierType != ESimplifyType::UEStandard && SimplifierType != ESimplifyType::MinimalPolygroup && SimplifierType != ESimplifyType::ClusterBased"))
-	float RegularizeWeight = 0.0001;
+	float RegularizeWeight = 0.000001;
 
-	/** Optionally scale the regularization with per-vertex weights, if a weight map is specified. Can be used to target simplification via the weights. */
-	UPROPERTY(EditAnywhere, Category = Options, meta = (
-		EditCondition = "RegularizeWeight > 0 && SimplifierType != ESimplifyType::UEStandard && SimplifierType != ESimplifyType::MinimalPolygroup && SimplifierType != ESimplifyType::ClusterBased"))
-	bool bScaleRegularizationWithWeightMap = true;
+	UPROPERTY(EditAnywhere, Category = Options, meta = (Bitmask, BitmaskEnum = "/Script/MeshModelingToolsEditorOnly.EApplyWeightMapToSimplifyMeasure"))
+	uint32 ApplyWeightMapTo = (int32)EApplyWeightMapToSimplifyMeasure::All;
 
 	/** Display colors corresponding to the mesh's polygon groups */
 	UPROPERTY(EditAnywhere, Category = Display)
diff --git a/Engine/Plugins/Runtime/MeshModelingToolset/Source/ModelingOperatorsEditorOnly/Private/CleaningOps/SimplifyMeshOp.cpp b/Engine/Plugins/Runtime/MeshModelingToolset/Source/ModelingOperatorsEditorOnly/Private/CleaningOps/SimplifyMeshOp.cpp
--- a/Engine/Plugins/Runtime/MeshModelingToolset/Source/ModelingOperatorsEditorOnly/Private/CleaningOps/SimplifyMeshOp.cpp
+++ b/Engine/Plugins/Runtime/MeshModelingToolset/Source/ModelingOperatorsEditorOnly/Private/CleaningOps/SimplifyMeshOp.cpp
@@ -43,9 +43,8 @@ void ComputeSimplify(FDynamicMesh3* TargetMesh, const bool bReproject,
 					 float GeometricTolerance,
 					 float RegularizeWeight,
 	TFunction<double(const FDynamicMesh3&, int VertexA, int VertexB)> CustomGeometricErrorScaleF,
-	TFunction<double(const FDynamicMesh3&, int VertexA, int VertexB)> CustomQuadricErrorScaleF,
 	TFunction<double(const FDynamicMesh3&, int VertexA, int VertexB)> CustomEdgeLengthScaleF,
-	TFunction<double(const FDynamicMesh3&, int Vertex)> CustomRegularizationScaleF)
+	TFunction<double(const FDynamicMesh3&, int Vertex)> CustomQuadricErrorScaleF)
 {
 	SimplificationType Reducer(TargetMesh);
 
@@ -61,7 +60,6 @@ void ComputeSimplify(FDynamicMesh3* TargetMesh, const bool bReproject,
 	Reducer.CustomGeometricErrorScaleF = CustomGeometricErrorScaleF;
 	Reducer.CustomQuadricErrorScaleF = CustomQuadricErrorScaleF;
 	Reducer.CustomEdgeLengthScaleF = CustomEdgeLengthScaleF;
-	Reducer.CustomRegularizationScaleF = CustomRegularizationScaleF;
 
 	Reducer.RegularizeWeight = RegularizeWeight;
 
@@ -258,9 +256,8 @@ void FSimplifyMeshOp::CalculateResult(FProgressCancel* Progress)
 											UseGeometricTolerance,
 											RegularizeWeight,
 											CustomGeometricErrorScaleF,
-											CustomQuadricErrorScaleF,
 											CustomEdgeLengthScaleF,
-											CustomRegularizationScaleF);
+											CustomQuadricErrorScaleF);
 	}
 	else if (SimplifierType == ESimplifyType::Attribute)
 	{
@@ -280,9 +277,8 @@ void FSimplifyMeshOp::CalculateResult(FProgressCancel* Progress)
 													UseGeometricTolerance,
 													RegularizeWeight,
 													CustomGeometricErrorScaleF,
-													CustomQuadricErrorScaleF,
 													CustomEdgeLengthScaleF,
-													CustomRegularizationScaleF);
+													CustomQuadricErrorScaleF);
 	}
 	else if (SimplifierType == ESimplifyType::MinimalPlanar)
 	{
@@ -302,9 +298,8 @@ void FSimplifyMeshOp::CalculateResult(FProgressCancel* Progress)
 			UseGeometricTolerance,
 			RegularizeWeight,
 			CustomGeometricErrorScaleF,
-			CustomQuadricErrorScaleF,
 			CustomEdgeLengthScaleF,
-			CustomRegularizationScaleF);
+			CustomQuadricErrorScaleF);
 	}
 	else if (SimplifierType == ESimplifyType::MinimalExistingVertex)
 	{
@@ -320,9 +315,8 @@ void FSimplifyMeshOp::CalculateResult(FProgressCancel* Progress)
 			UseGeometricTolerance,
 			RegularizeWeight,
 			CustomGeometricErrorScaleF,
-			CustomQuadricErrorScaleF,
 			CustomEdgeLengthScaleF,
-			CustomRegularizationScaleF);
+			CustomQuadricErrorScaleF);
 	}
 	else if (SimplifierType == ESimplifyType::MinimalPolygroup)
 	{
diff --git a/Engine/Plugins/Runtime/MeshModelingToolset/Source/ModelingOperatorsEditorOnly/Public/CleaningOps/SimplifyMeshOp.h b/Engine/Plugins/Runtime/MeshModelingToolset/Source/ModelingOperatorsEditorOnly/Public/CleaningOps/SimplifyMeshOp.h
--- a/Engine/Plugins/Runtime/MeshModelingToolset/Source/ModelingOperatorsEditorOnly/Public/CleaningOps/SimplifyMeshOp.h
+++ b/Engine/Plugins/Runtime/MeshModelingToolset/Source/ModelingOperatorsEditorOnly/Public/CleaningOps/SimplifyMeshOp.h
@@ -103,11 +103,11 @@ public:
 	// Weight of optional regularization term, for QEM-based simplifiers
 	float RegularizeWeight = 0.f;
 
-	// QEM simplifier custom scales to apply to *measured* errors and edge lengths. Note this is in effect the inverse of scaling the tolerances / target lengths.
+	// QEM simplifier custom scales to apply to errors and edge lengths. Note this is in effect the inverse of scaling the tolerances / target lengths.
 	TFunction<double(const FDynamicMesh3&, int VertexA, int VertexB)> CustomGeometricErrorScaleF;
-	TFunction<double(const FDynamicMesh3&, int VertexA, int VertexB)> CustomQuadricErrorScaleF;
 	TFunction<double(const FDynamicMesh3&, int VertexA, int VertexB)> CustomEdgeLengthScaleF;
-	TFunction<double(const FDynamicMesh3&, int Vertex)> CustomRegularizationScaleF;
+	// Custom per-vertex quadric error scaling. Values should be greater than zero.
+	TFunction<double(const FDynamicMesh3&, int Vertex)> CustomQuadricErrorScaleF;
 
 	// Cluster simplifier custom scales apply to the *target* edge lengths, and so is conceptually the inverse of the above scales.
 	TFunction<double(int32)> ClusterTargetEdgeLengthScaleF = nullptr;
