Change the default geometry script simplify to have a small regularize weight
+ reduce ClampMax values to 10 (which is already larger than should be used)

diff --git a/Engine/Plugins/Runtime/GeometryScripting/Source/GeometryScriptingCore/Public/GeometryScript/MeshSimplifyFunctions.h b/Engine/Plugins/Runtime/GeometryScripting/Source/GeometryScriptingCore/Public/GeometryScript/MeshSimplifyFunctions.h
--- a/Engine/Plugins/Runtime/GeometryScripting/Source/GeometryScriptingCore/Public/GeometryScript/MeshSimplifyFunctions.h
+++ b/Engine/Plugins/Runtime/GeometryScripting/Source/GeometryScriptingCore/Public/GeometryScript/MeshSimplifyFunctions.h
@@ -103,9 +103,9 @@ public:
 	UPROPERTY(EditAnywhere, BlueprintReadWrite, Category = Options)
 	FGeometryScriptWeightMapDensity QuadricErrorWeightMap;
 
-	/** Setting a very small RegularizeWeight value (e.g., .00001) can improve triangle quality and the effect of the QuadricErrorWeightMap in flat regions */
-	UPROPERTY(EditAnywhere, BlueprintReadWrite, Category = Options, meta = (ClampMin = "0.0", ClampMax = "100.0", UIMax = ".01"))
-	float RegularizeWeight = 0.f;
+	/** Setting a very small RegularizeWeight value (e.g., 0.000001) can improve triangle quality and the effect of the QuadricErrorWeightMap in flat regions */
+	UPROPERTY(EditAnywhere, BlueprintReadWrite, Category = Options, meta = (ClampMin = "0.0", ClampMax = "10.0", UIMax = ".01"))
+	float RegularizeWeight = 0.000001f;
 
 	/** If enabled, the simplified mesh is automatically compacted to remove gaps in the index space. This is expensive and can be disabled by advanced users. */
 	UPROPERTY(EditAnywhere, BlueprintReadWrite, Category = Options)
diff --git a/Engine/Plugins/Runtime/MeshModelingToolset/Source/MeshModelingToolsEditorOnly/Public/SimplifyMeshTool.h b/Engine/Plugins/Runtime/MeshModelingToolset/Source/MeshModelingToolsEditorOnly/Public/SimplifyMeshTool.h
--- a/Engine/Plugins/Runtime/MeshModelingToolset/Source/MeshModelingToolsEditorOnly/Public/SimplifyMeshTool.h
+++ b/Engine/Plugins/Runtime/MeshModelingToolset/Source/MeshModelingToolsEditorOnly/Public/SimplifyMeshTool.h
@@ -113,7 +113,7 @@ public:
 	float GeometricTolerance;
 
 	/** Optional regularization term to favor vertices remaining near their original locations. Using a small regularization weight can give cleaner topology in flat regions. */
-	UPROPERTY(EditAnywhere, Category = Options, meta = (UIMin = "0.0", UIMax = "0.01", ClampMin = "0.0", ClampMax = "100.0",
+	UPROPERTY(EditAnywhere, Category = Options, meta = (UIMin = "0.0", UIMax = "0.01", ClampMin = "0.0", ClampMax = "10.0",
 		EditCondition = "SimplifierType != ESimplifyType::UEStandard && SimplifierType != ESimplifyType::MinimalPolygroup && SimplifierType != ESimplifyType::ClusterBased"))
 	float RegularizeWeight = 0.000001;
 
