fix crash on destroying the global mesh pool twice

diff --git a/Engine/Plugins/Runtime/GeometryScripting/Source/GeometryScriptingCore/Private/MeshPoolFunctions.cpp b/Engine/Plugins/Runtime/GeometryScripting/Source/GeometryScriptingCore/Private/MeshPoolFunctions.cpp
--- a/Engine/Plugins/Runtime/GeometryScripting/Source/GeometryScriptingCore/Private/MeshPoolFunctions.cpp
+++ b/Engine/Plugins/Runtime/GeometryScripting/Source/GeometryScriptingCore/Private/MeshPoolFunctions.cpp
@@ -19,7 +19,7 @@ namespace UE::Private::MeshPoolHelper
 			MeshPool = NewObject<UDynamicMeshPool>();
 			MeshPool->AddToRoot();
 		}
-		else if (bDestroy)
+		else if (MeshPool && bDestroy)
 		{
 			MeshPool->FreeAllMeshes();
 			MeshPool->RemoveFromRoot();
