fix crash in Geometry Script's AppendBuffersToMesh when buffer contains triangles with invalid indices

#Synced-CL 49942658

diff --git a/Engine/Plugins/Runtime/GeometryScripting/Source/GeometryScriptingCore/Private/MeshBasicEditFunctions.cpp b/Engine/Plugins/Runtime/GeometryScripting/Source/GeometryScriptingCore/Private/MeshBasicEditFunctions.cpp
--- a/Engine/Plugins/Runtime/GeometryScripting/Source/GeometryScriptingCore/Private/MeshBasicEditFunctions.cpp
+++ b/Engine/Plugins/Runtime/GeometryScripting/Source/GeometryScriptingCore/Private/MeshBasicEditFunctions.cpp
@@ -916,11 +916,18 @@ UDynamicMesh* UGeometryScriptLibrary_MeshBasicEditFunctions::AppendBuffersToMesh
 
 		TArray<int32>& NewTriangleIndices = *NewTriangleIndicesList.List;
 		int32 NumTriangles = Buffers.Triangles.Num();
+		NewTriangleIndices.Init(INDEX_NONE, NumTriangles);
 		bool bHaveGroups = Buffers.TriGroupIDs.Num() == NumTriangles;
 		int32 ConstantGroupID = EditMesh.AllocateTriangleGroup();
 		for (int32 k = 0; k < NumTriangles; ++k)
 		{
 			int32 UseGroupID = bHaveGroups ? Buffers.TriGroupIDs[k] : ConstantGroupID;
+			const FIntVector& BufferTri = Buffers.Triangles[k];
+			if (!VertexIDMap.IsValidIndex(BufferTri.X) || !VertexIDMap.IsValidIndex(BufferTri.Y) || !VertexIDMap.IsValidIndex(BufferTri.Z))
+			{
+				UE::Geometry::AppendError(Debug, EGeometryScriptErrorType::InvalidInputs, LOCTEXT("AppendBuffersToMesh_Invalid", "AppendBuffersToMesh: Triangle cannot be added because it had invalid indices"));
+				continue;
+			}
 			int32 NewTriangleID = EditMesh.AppendTriangle(MapTriangleFunc(Buffers.Triangles[k]), UseGroupID);
 			if (NewTriangleID < 0)
 			{
@@ -936,13 +943,12 @@ UDynamicMesh* UGeometryScriptLibrary_MeshBasicEditFunctions::AppendBuffersToMesh
 				{
 					UE::Geometry::AppendError(Debug, EGeometryScriptErrorType::OperationFailed, LOCTEXT("AppendBuffersToMesh_Unknown", "AppendBuffersToMesh: adding Triangle Failed"));
 				}
-				NewTriangleID = INDEX_NONE;
 			}
 			else
 			{
 				MaterialIDs->SetValue(NewTriangleID, MaterialID);
+				NewTriangleIndices[k] = NewTriangleID;
 			}
-			NewTriangleIndices.Add(NewTriangleID);
 		}
 
 
