Add more material management to the copy mesh to static mesh / copy mesh to skeletal mesh options --
1. Optionally use New Materials to only set materials on the currently updated LOD
2. Add option to clean material IDs, so invalid material IDs will be set to reference a valid material instead

#Synced-CL 50262759

diff --git a/Engine/Plugins/Runtime/GeometryScripting/Source/GeometryScriptingCore/Private/MeshAssetFunctions.cpp b/Engine/Plugins/Runtime/GeometryScripting/Source/GeometryScriptingCore/Private/MeshAssetFunctions.cpp
--- a/Engine/Plugins/Runtime/GeometryScripting/Source/GeometryScriptingCore/Private/MeshAssetFunctions.cpp
+++ b/Engine/Plugins/Runtime/GeometryScripting/Source/GeometryScriptingCore/Private/MeshAssetFunctions.cpp
@@ -256,6 +256,10 @@ UDynamicMesh*  UGeometryScriptLibrary_StaticMeshFunctions::CopyMeshToStaticMesh(
 		}
 
 		FConversionToMeshDescriptionOptions ConversionOptions;
+		if (Options.bCleanAssignedMaterials)
+		{
+			ConversionOptions.MaxValidMaterialID = ToStaticMeshAsset->GetStaticMaterials().Num() - 1;
+		}
 		FDynamicMeshToMeshDescription Converter(ConversionOptions);
 		FromDynamicMesh->ProcessMesh([&](const FDynamicMesh3& ReadMesh)
 			{
@@ -295,6 +299,19 @@ UDynamicMesh*  UGeometryScriptLibrary_StaticMeshFunctions::CopyMeshToStaticMesh(
 
 		FConversionToMeshDescriptionOptions ConversionOptions;
 		ConversionOptions.bConvertBackToNonManifold = Options.bUseOriginalVertexOrder;
+		if (Options.bCleanAssignedMaterials)
+		{
+			if (Options.bReplaceMaterials)
+			{
+				ConversionOptions.MaxValidMaterialID = Options.NewMaterials.Num() - 1;
+			}
+			else
+			{
+				const FMeshSectionInfoMap& SectionMap = ToStaticMeshAsset->GetSectionInfoMap();
+				int32 LODSectionNum = SectionMap.GetSectionNumber(UseLODIndex);
+				ConversionOptions.MaxValidMaterialID = LODSectionNum - 1;
+			}
+		}
 		
 		FDynamicMeshToMeshDescription Converter(ConversionOptions);
 		if (!bUseSectionMaterials && !Options.bReplaceMaterials)
@@ -322,63 +339,117 @@ UDynamicMesh*  UGeometryScriptLibrary_StaticMeshFunctions::CopyMeshToStaticMesh(
 		{
 			bool bHaveSlotNames = (Options.NewMaterialSlotNames.Num() == Options.NewMaterials.Num());
 
-			TArray<FStaticMaterial> NewMaterials;
-			for (int32 k = 0; k < Options.NewMaterials.Num(); ++k)
+			auto SetMaterialSlotNames = [&MeshDescription](const TArray<FStaticMaterial>& NewMaterials)
+			{
+				// Set material slot names on the mesh description
+				FStaticMeshAttributes Attributes(*MeshDescription);
+				TPolygonGroupAttributesRef<FName> PolygonGroupImportedMaterialSlotNames = Attributes.GetPolygonGroupMaterialSlotNames();
+				for (int32 SlotIdx = 0; SlotIdx < NewMaterials.Num(); ++SlotIdx)
+				{
+					if (SlotIdx < PolygonGroupImportedMaterialSlotNames.GetNumElements())
+					{
+						PolygonGroupImportedMaterialSlotNames.Set(SlotIdx, NewMaterials[SlotIdx].ImportedMaterialSlotName);
+					}
+				}
+			};
+
+			auto MakeNewMaterial = [&Options, bHaveSlotNames](const TArray<FStaticMaterial>& StaticMeshMaterials, int32 NewMatIdx) -> FStaticMaterial
 			{
 				FStaticMaterial NewMaterial;
-				NewMaterial.MaterialInterface = Options.NewMaterials[k];
-				FName UseSlotName = (bHaveSlotNames && Options.NewMaterialSlotNames[k] != NAME_None) ? Options.NewMaterialSlotNames[k] :
-					UE::AssetUtils::GenerateNewMaterialSlotName(NewMaterials, NewMaterial.MaterialInterface, k);
+				NewMaterial.MaterialInterface = Options.NewMaterials[NewMatIdx];
+				FName UseSlotName = (bHaveSlotNames && Options.NewMaterialSlotNames[NewMatIdx] != NAME_None) ? Options.NewMaterialSlotNames[NewMatIdx] :
+					UE::AssetUtils::GenerateNewMaterialSlotName(StaticMeshMaterials, NewMaterial.MaterialInterface, NewMatIdx);
 
 				NewMaterial.MaterialSlotName = UseSlotName;
 				NewMaterial.ImportedMaterialSlotName = UseSlotName;
 				NewMaterial.UVChannelData = FMeshUVChannelInfo(1.f);		// this avoids an ensure in  UStaticMesh::GetUVChannelData
-				NewMaterials.Add(NewMaterial);
-			}
-
-			ToStaticMeshAsset->SetStaticMaterials(NewMaterials);
+				return NewMaterial;
+			};
 
-			// Set material slot names on the mesh description
-			FStaticMeshAttributes Attributes(*MeshDescription);
-			TPolygonGroupAttributesRef<FName> PolygonGroupImportedMaterialSlotNames = Attributes.GetPolygonGroupMaterialSlotNames();
-			for (int32 SlotIdx = 0; SlotIdx < NewMaterials.Num(); ++SlotIdx)
+			if (Options.NewMaterialUsage == ECopyMeshToAssetUseNewMaterialsOptions::ReplaceAssetMaterials)
 			{
-				if (SlotIdx < PolygonGroupImportedMaterialSlotNames.GetNumElements())
+				TArray<FStaticMaterial> NewMaterials;
+				for (int32 k = 0; k < Options.NewMaterials.Num(); ++k)
 				{
-					PolygonGroupImportedMaterialSlotNames.Set(SlotIdx, NewMaterials[SlotIdx].ImportedMaterialSlotName);
+					NewMaterials.Add(MakeNewMaterial(NewMaterials, k));
 				}
-			}
 
-			// Reset the section info map
-			ToStaticMeshAsset->GetSectionInfoMap().Clear();
-			ToStaticMeshAsset->GetOriginalSectionInfoMap().Clear();
+				ToStaticMeshAsset->SetStaticMaterials(NewMaterials);
 
-			// Repopulate section info map
-			FMeshSectionInfoMap SectionInfoMap;
-			for (int32 LODIndex = 0, NumLODs = ToStaticMeshAsset->GetNumSourceModels(); LODIndex < NumLODs; ++LODIndex)
-			{
-				if (const FMeshDescription* Mesh = (LODIndex == UseLODIndex) ? MeshDescription : ToStaticMeshAsset->GetMeshDescription(LODIndex))
+				SetMaterialSlotNames(NewMaterials);
+
+				// Reset the section info map
+				ToStaticMeshAsset->GetSectionInfoMap().Clear();
+				ToStaticMeshAsset->GetOriginalSectionInfoMap().Clear();
+
+				// Repopulate section info map
+				FMeshSectionInfoMap SectionInfoMap;
+				for (int32 LODIndex = 0, NumLODs = ToStaticMeshAsset->GetNumSourceModels(); LODIndex < NumLODs; ++LODIndex)
 				{
-					FStaticMeshConstAttributes MeshDescriptionAttributes(*Mesh);
-					TPolygonGroupAttributesConstRef<FName> MaterialSlotNames = MeshDescriptionAttributes.GetPolygonGroupMaterialSlotNames();
-					int32 SectionIndex = 0;
-					for (FPolygonGroupID PolygonGroupID : Mesh->PolygonGroups().GetElementIDs())
+					if (const FMeshDescription* Mesh = (LODIndex == UseLODIndex) ? MeshDescription : ToStaticMeshAsset->GetMeshDescription(LODIndex))
 					{
-						// Material index is either from the matching material slot name or the section index if that name is not found
-						int32 MaterialIndex = ToStaticMeshAsset->GetStaticMaterials().IndexOfByPredicate(
-							[&MaterialSlotName = MaterialSlotNames[PolygonGroupID]](const FStaticMaterial& StaticMaterial) { return StaticMaterial.MaterialSlotName == MaterialSlotName; }
-						);
-						if (MaterialIndex == INDEX_NONE)
+						FStaticMeshConstAttributes MeshDescriptionAttributes(*Mesh);
+						TPolygonGroupAttributesConstRef<FName> MaterialSlotNames = MeshDescriptionAttributes.GetPolygonGroupMaterialSlotNames();
+						int32 SectionIndex = 0;
+						for (FPolygonGroupID PolygonGroupID : Mesh->PolygonGroups().GetElementIDs())
 						{
-							MaterialIndex = SectionIndex;
+							// Material index is either from the matching material slot name or the section index if that name is not found
+							int32 MaterialIndex = ToStaticMeshAsset->GetStaticMaterials().IndexOfByPredicate(
+								[&MaterialSlotName = MaterialSlotNames[PolygonGroupID]](const FStaticMaterial& StaticMaterial) { return StaticMaterial.MaterialSlotName == MaterialSlotName; }
+							);
+							if (MaterialIndex == INDEX_NONE)
+							{
+								MaterialIndex = SectionIndex;
+							}
+							SectionInfoMap.Set(LODIndex, SectionIndex, FMeshSectionInfo(MaterialIndex));
+							SectionIndex++;
 						}
-						SectionInfoMap.Set(LODIndex, SectionIndex, FMeshSectionInfo(MaterialIndex));
-						SectionIndex++;
 					}
 				}
+				ToStaticMeshAsset->GetSectionInfoMap().CopyFrom(SectionInfoMap);
+				ToStaticMeshAsset->GetOriginalSectionInfoMap().CopyFrom(SectionInfoMap);
+			}
+			else
+			{
+				TArray<FStaticMaterial> ToUpdateMaterials = ToStaticMeshAsset->GetStaticMaterials();
+				TMap<UMaterialInterface*, int32> MatToIndex;
+				for (int32 OldMatIdx = 0; OldMatIdx < ToUpdateMaterials.Num(); ++OldMatIdx)
+				{
+					MatToIndex.Add(ToUpdateMaterials[OldMatIdx].MaterialInterface, OldMatIdx);
+				}
+				// 
+				for (int32 NewMatIdx = 0; NewMatIdx < Options.NewMaterials.Num(); ++NewMatIdx)
+				{
+					UMaterialInterface* NewMat = Options.NewMaterials[NewMatIdx];
+					if (!MatToIndex.Contains(NewMat))
+					{
+						int32 NewIdx = ToUpdateMaterials.Add(MakeNewMaterial(ToUpdateMaterials, NewMatIdx));
+						MatToIndex.Add(NewMat, NewIdx);
+					}
+				}
+
+				ToStaticMeshAsset->SetStaticMaterials(ToUpdateMaterials);
+
+				SetMaterialSlotNames(ToUpdateMaterials);
+
+				// update material sections map for the current LOD
+				auto UpdateMaterialSectionInfo = [&MatToIndex, &Options](FMeshSectionInfoMap& SectionInfoMap, int32 LODIndex)
+				{
+					// clear old sections for this LOD
+					int32 NumSection = SectionInfoMap.GetSectionNumber(LODIndex);
+					for (int32 Idx = 0; Idx < NumSection; ++Idx)
+					{
+						SectionInfoMap.Remove(LODIndex, Idx);
+					}
+					// add new section per material in the NewMaterials array
+					for (int32 NewMatIdx = 0; NewMatIdx < Options.NewMaterials.Num(); ++NewMatIdx)
+					{
+						SectionInfoMap.Set(LODIndex, NewMatIdx, FMeshSectionInfo(MatToIndex[Options.NewMaterials[NewMatIdx]]));
+					}
+				};
+				UpdateMaterialSectionInfo(ToStaticMeshAsset->GetSectionInfoMap(), UseLODIndex);
+				UpdateMaterialSectionInfo(ToStaticMeshAsset->GetOriginalSectionInfoMap(), UseLODIndex);
 			}
-			ToStaticMeshAsset->GetSectionInfoMap().CopyFrom(SectionInfoMap);
-			ToStaticMeshAsset->GetOriginalSectionInfoMap().CopyFrom(SectionInfoMap);
 		}
 
 		ToStaticMeshAsset->CommitMeshDescription(UseLODIndex);
@@ -998,7 +1069,98 @@ UDynamicMesh* UGeometryScriptLibrary_StaticMeshFunctions::CopyMeshToSkeletalMesh
 	FConversionToMeshDescriptionOptions ConversionOptions;
 	ConversionOptions.bConvertBackToNonManifold = Options.bUseOriginalVertexOrder;
 
+	TUniquePtr<TArray<int32>> MatIDInverseMap;
+
+	// Update asset materials / set a corresponding material remap on the converter if needed
+	if (Options.bReplaceMaterials)
+	{
+		bool bHaveSlotNames = (Options.NewMaterialSlotNames.Num() == Options.NewMaterials.Num());
+
+		auto MakeNewMaterial = [&Options, bHaveSlotNames](TArray<FSkeletalMaterial>& SkelMaterials, int32 NewMatIdx) -> FSkeletalMaterial
+		{
+			FSkeletalMaterial NewMaterial;
+			NewMaterial.MaterialInterface = Options.NewMaterials[NewMatIdx];
+			FName UseSlotName = (bHaveSlotNames && Options.NewMaterialSlotNames[NewMatIdx] != NAME_None) ? Options.NewMaterialSlotNames[NewMatIdx] :
+				UELocal::GenerateNewMaterialSlotName(SkelMaterials, NewMaterial.MaterialInterface, NewMatIdx);
+
+			NewMaterial.MaterialSlotName = UseSlotName;
+			NewMaterial.ImportedMaterialSlotName = UseSlotName;
+			NewMaterial.UVChannelData = FMeshUVChannelInfo(1.f);		// this avoids an ensure in  UStaticMesh::GetUVChannelData
+			return NewMaterial;
+		};
+
+		if (Options.NewMaterialUsage == ECopyMeshToAssetUseNewMaterialsOptions::ReplaceAssetMaterials)
+		{
+			TArray<FSkeletalMaterial> NewMaterials;
+			for (int32 k = 0; k < Options.NewMaterials.Num(); ++k)
+			{
+				NewMaterials.Add(MakeNewMaterial(NewMaterials, k));
+			}
+
+			ToSkeletalMeshAsset->SetMaterials(NewMaterials);
+		}
+		else
+		{
+			TArray<FSkeletalMaterial> ToUpdateMaterials = ToSkeletalMeshAsset->GetMaterials();
+			TMap<UMaterialInterface*, int32> MatToIndex;
+			for (int32 OldMatIdx = 0; OldMatIdx < ToUpdateMaterials.Num(); ++OldMatIdx)
+			{
+				MatToIndex.Add(ToUpdateMaterials[OldMatIdx].MaterialInterface, OldMatIdx);
+			}
+			bool bUpdated = false;
+			for (int32 NewMatIdx = 0; NewMatIdx < Options.NewMaterials.Num(); ++NewMatIdx)
+			{
+				UMaterialInterface* NewMat = Options.NewMaterials[NewMatIdx];
+				if (!MatToIndex.Contains(NewMat))
+				{
+					bUpdated = true;
+					int32 NewIdx = ToUpdateMaterials.Add(MakeNewMaterial(ToUpdateMaterials, NewMatIdx));
+					MatToIndex.Add(NewMat, NewIdx);
+				}
+			}
+
+			// build material remap s.t. dynamic mesh material indices map to where materials were set in the skeletal mesh asset
+			MatIDInverseMap = MakeUnique<TArray<int32>>();
+			MatIDInverseMap->Init(-1, ToUpdateMaterials.Num());
+			for (int32 NewMatIdx = 0; NewMatIdx < Options.NewMaterials.Num(); ++NewMatIdx)
+			{
+				(*MatIDInverseMap)[MatToIndex[Options.NewMaterials[NewMatIdx]]] = NewMatIdx;
+			}
+
+			if (bUpdated)
+			{
+				ToSkeletalMeshAsset->SetMaterials(ToUpdateMaterials);
+			}
+		}
+	}
+
+	
 	FDynamicMeshToMeshDescription Converter(ConversionOptions);
+	if (MatIDInverseMap)
+	{
+		Converter.SetMaterialIDMapFromInverseMap(*MatIDInverseMap);
+	}
+	if (Options.bCleanAssignedMaterials)
+	{
+		if (Options.bReplaceMaterials)
+		{
+			if (Options.NewMaterialUsage == ECopyMeshToAssetUseNewMaterialsOptions::ReplaceAssetMaterials)
+			{
+				ConversionOptions.MaxValidMaterialID = Options.NewMaterials.Num() - 1;
+			}
+			else if (ensure(MatIDInverseMap))
+			{
+				ConversionOptions.MaxValidMaterialID = MatIDInverseMap->Num() - 1;
+			}
+		}
+		else
+		{
+			FSkeletalMeshLODInfo* SkeletalLODInfo = ToSkeletalMeshAsset->GetLODInfo(TargetLOD.LODIndex);
+			int32 LODNumMats = SkeletalLODInfo->LODMaterialMap.IsEmpty() ? ToSkeletalMeshAsset->GetNumMaterials() : SkeletalLODInfo->LODMaterialMap.Num();
+			ConversionOptions.MaxValidMaterialID = LODNumMats - 1;
+		}
+	}
+
 	FromDynamicMesh->ProcessMesh([&](const FDynamicMesh3& ReadMesh)
 	{
 		Converter.Convert(&ReadMesh, *MeshDescription, !Options.bEnableRecomputeTangents);
@@ -1094,35 +1256,20 @@ UDynamicMesh* UGeometryScriptLibrary_StaticMeshFunctions::CopyMeshToSkeletalMesh
 		ToSkeletalMeshAsset->NaniteSettings = Options.NewNaniteSettings;
 	}
 
-	// update materials on the Asset
+	// Final pass of material updating
 	if (Options.bReplaceMaterials)
 	{
-		bool bHaveSlotNames = (Options.NewMaterialSlotNames.Num() == Options.NewMaterials.Num());
-
-		TArray<FSkeletalMaterial> NewMaterials;
-		for (int32 k = 0; k < Options.NewMaterials.Num(); ++k)
-		{
-			FSkeletalMaterial NewMaterial;
-			NewMaterial.MaterialInterface = Options.NewMaterials[k];
-			FName UseSlotName = (bHaveSlotNames && Options.NewMaterialSlotNames[k] != NAME_None) ? Options.NewMaterialSlotNames[k] :
-				UELocal::GenerateNewMaterialSlotName(NewMaterials, NewMaterial.MaterialInterface, k);
-
-			NewMaterial.MaterialSlotName = UseSlotName;
-			NewMaterial.ImportedMaterialSlotName = UseSlotName;
-			NewMaterial.UVChannelData = FMeshUVChannelInfo(1.f);		// this avoids an ensure in  UStaticMesh::GetUVChannelData
-			NewMaterials.Add(NewMaterial);
-		}
-		SkeletalLODInfo->LODMaterialMap.Empty();
+		// reset LODMaterialMap; it will be set up later with just the materials actually used
+		SkeletalLODInfo->LODMaterialMap.Reset();
 		
-		ToSkeletalMeshAsset->SetMaterials(NewMaterials);
-
 		// Set material slot names on the mesh description
 		TPolygonGroupAttributesRef<FName> PolygonGroupImportedMaterialSlotNames = MeshAttributes.GetPolygonGroupMaterialSlotNames();
-		for (int32 SlotIdx = 0; SlotIdx < NewMaterials.Num(); ++SlotIdx)
+		const TArray<FSkeletalMaterial>& UpdatedMaterials = ToSkeletalMeshAsset->GetMaterials();
+		for (int32 SlotIdx = 0; SlotIdx < UpdatedMaterials.Num(); ++SlotIdx)
 		{
 			if (SlotIdx < PolygonGroupImportedMaterialSlotNames.GetNumElements())
 			{
-				PolygonGroupImportedMaterialSlotNames.Set(SlotIdx, NewMaterials[SlotIdx].ImportedMaterialSlotName);
+				PolygonGroupImportedMaterialSlotNames.Set(SlotIdx, UpdatedMaterials[SlotIdx].ImportedMaterialSlotName);
 			}
 		}
 	}
diff --git a/Engine/Plugins/Runtime/GeometryScripting/Source/GeometryScriptingCore/Public/GeometryScript/MeshAssetFunctions.h b/Engine/Plugins/Runtime/GeometryScripting/Source/GeometryScriptingCore/Public/GeometryScript/MeshAssetFunctions.h
--- a/Engine/Plugins/Runtime/GeometryScripting/Source/GeometryScriptingCore/Public/GeometryScript/MeshAssetFunctions.h
+++ b/Engine/Plugins/Runtime/GeometryScripting/Source/GeometryScriptingCore/Public/GeometryScript/MeshAssetFunctions.h
@@ -95,6 +95,17 @@ struct FGeometryScriptNaniteOptions
 	float FallbackRelativeError = 0.0f;
 };
 
+// Options for how to use the NewMaterials array when copying materials to a mesh asset
+UENUM(BlueprintType)
+enum class ECopyMeshToAssetUseNewMaterialsOptions : uint8
+{
+	// Fully replace asset materials list with the provided set
+	ReplaceAssetMaterials = 0,
+	// Use the provided new materials as the material for the newly-set LOD, only adding asset materials as-needed
+	SetLODMaterials = 1 UMETA(DisplayName = "Set LOD Materials")
+};
+
+
 
 USTRUCT(BlueprintType)
 struct FGeometryScriptCopyMeshToAssetOptions
@@ -131,13 +142,13 @@ struct FGeometryScriptCopyMeshToAssetOptions
 	UPROPERTY(EditAnywhere, BlueprintReadWrite, Category = Options)
 	bool bUseBuildScale = true;
 	
-	// Whether to replace the materials on the asset with those in the New Materials array
-	UPROPERTY(EditAnywhere, BlueprintReadWrite, Category = Options)
+	// Whether to update the materials on the asset with those in the New Materials array
+	UPROPERTY(EditAnywhere, BlueprintReadWrite, Category = Options, meta=(DisplayName = "Use New Materials"))
 	bool bReplaceMaterials = false;
 
-	// Whether to generate lightmap UVs
+	// If using New Materials, whether to use them as a full replacement of all asset materials, or just to update the target LOD
 	UPROPERTY(EditAnywhere, BlueprintReadWrite, Category = Options)
-	EGeometryScriptGenerateLightmapUVOptions GenerateLightmapUVs = EGeometryScriptGenerateLightmapUVOptions::MatchTargetLODSetting;
+	ECopyMeshToAssetUseNewMaterialsOptions NewMaterialUsage = ECopyMeshToAssetUseNewMaterialsOptions::ReplaceAssetMaterials;
 
 	// New materials to set if Replace Materials is enabled. Ignored otherwise.
 	UPROPERTY(EditAnywhere, BlueprintReadWrite, Category = Options)
@@ -147,6 +158,14 @@ struct FGeometryScriptCopyMeshToAssetOptions
 	UPROPERTY(EditAnywhere, BlueprintReadWrite, Category = Options)
 	TArray<FName> NewMaterialSlotNames;
 
+	// Whether to apply a cleanup pass to assigned materials -- fixing invalid material assignments
+	UPROPERTY(EditAnywhere, BlueprintReadWrite, Category = Options)
+	bool bCleanAssignedMaterials = false;
+
+	// Whether to generate lightmap UVs
+	UPROPERTY(EditAnywhere, BlueprintReadWrite, Category = Options)
+	EGeometryScriptGenerateLightmapUVOptions GenerateLightmapUVs = EGeometryScriptGenerateLightmapUVOptions::MatchTargetLODSetting;
+
 	/** If enabled, NaniteSettings will be applied to the target Asset if possible */
 	UPROPERTY(EditAnywhere, BlueprintReadWrite, Category = Options)
 	bool bApplyNaniteSettings = false;
diff --git a/Engine/Source/Runtime/MeshConversion/Private/DynamicMeshToMeshDescription.cpp b/Engine/Source/Runtime/MeshConversion/Private/DynamicMeshToMeshDescription.cpp
--- a/Engine/Source/Runtime/MeshConversion/Private/DynamicMeshToMeshDescription.cpp
+++ b/Engine/Source/Runtime/MeshConversion/Private/DynamicMeshToMeshDescription.cpp
@@ -141,6 +141,14 @@ namespace DynamicMeshToMeshDescriptionConversionHelper
 			}
 		}
 	}
+
+	void ValidateMaterialID(int32& MaterialID, int32 MaxValidMaterialID)
+	{
+		if (MaxValidMaterialID > INDEX_NONE && (MaterialID < 0 || MaterialID > MaxValidMaterialID))
+		{
+			MaterialID = 0;
+		}
+	}
 }
 
 
@@ -663,6 +671,7 @@ void FDynamicMeshToMeshDescription::Convert_SharedInstances(const FDynamicMesh3*
 			{
 				MaterialID = MaterialIDToPolygroupIDMap[MaterialID];
 			}
+			DynamicMeshToMeshDescriptionConversionHelper::ValidateMaterialID(MaterialID, ConversionOptions.MaxValidMaterialID);
 			MaxPolygonGroupID = FMath::Max(MaterialID, MaxPolygonGroupID);
 		}
 		if (MaxPolygonGroupID == 0)
@@ -774,6 +783,7 @@ void FDynamicMeshToMeshDescription::Convert_SharedInstances(const FDynamicMesh3*
 				MaterialID = MaterialIDToPolygroupIDMap[MaterialID];
 			}
 			MaterialID = FMath::Max(0, MaterialID);
+			DynamicMeshToMeshDescriptionConversionHelper::ValidateMaterialID(MaterialID, ConversionOptions.MaxValidMaterialID);
 			UsePolygonGroupID = FPolygonGroupID(MaterialID);
 		}
 
@@ -1033,6 +1043,7 @@ void FDynamicMeshToMeshDescription::Convert_NoSharedInstances(const FDynamicMesh
 			{
 				MaterialID = MaterialIDToPolygroupIDMap[MaterialID];
 			}
+			DynamicMeshToMeshDescriptionConversionHelper::ValidateMaterialID(MaterialID, ConversionOptions.MaxValidMaterialID);
 			MaxPolygonGroupID = FMath::Max(MaterialID, MaxPolygonGroupID);
 		}
 		if (MaxPolygonGroupID == 0)
@@ -1080,6 +1091,7 @@ void FDynamicMeshToMeshDescription::Convert_NoSharedInstances(const FDynamicMesh
 				MaterialID = MaterialIDToPolygroupIDMap[MaterialID];
 			}
 			MaterialID = FMath::Max(0, MaterialID);
+			DynamicMeshToMeshDescriptionConversionHelper::ValidateMaterialID(MaterialID, ConversionOptions.MaxValidMaterialID);
 			UsePolygonGroupID = FPolygonGroupID(MaterialID);
 		}
 
diff --git a/Engine/Source/Runtime/MeshConversion/Public/MeshConversionOptions.h b/Engine/Source/Runtime/MeshConversion/Public/MeshConversionOptions.h
--- a/Engine/Source/Runtime/MeshConversion/Public/MeshConversionOptions.h
+++ b/Engine/Source/Runtime/MeshConversion/Public/MeshConversionOptions.h
@@ -36,6 +36,9 @@ public:
 	/** Should normal seams be tagged as hard edges in the mesh description */
 	bool bSetHardEdgesFromNormalSeams = true;
 
+	/** If non-negative, sets the max valid material ID to assign -- IDs higher than this will be reset to 0. Ignored if negative. */
+	int32 MaxValidMaterialID = INDEX_NONE;
+
 	//
 	// utility functions for common configuration cases
 	//
