move some dataflow/geometryscript/modeling-tool mesh splitting calls of FDynamicMeshEditor::SplitMesh to the TUniquePtr version, to prep for deprecation of the unsafe version
+ also add a TUniquePtr version and deprecate the unsafe version for the similar method in FaceGroupUtil

#Synced-CL 50490755

diff --git a/Engine/Plugins/Experimental/GeometryCollectionPlugin/Source/GeometryCollectionNodes/Private/Dataflow/GeometryCollectionMeshNodes.cpp b/Engine/Plugins/Experimental/GeometryCollectionPlugin/Source/GeometryCollectionNodes/Private/Dataflow/GeometryCollectionMeshNodes.cpp
--- a/Engine/Plugins/Experimental/GeometryCollectionPlugin/Source/GeometryCollectionNodes/Private/Dataflow/GeometryCollectionMeshNodes.cpp
+++ b/Engine/Plugins/Experimental/GeometryCollectionPlugin/Source/GeometryCollectionNodes/Private/Dataflow/GeometryCollectionMeshNodes.cpp
@@ -182,7 +182,7 @@ void FMeshToCollectionDataflowNode::Evaluate(UE::Dataflow::FContext& Context, co
 				FGeometryCollectionToDynamicMeshes Convert;
 				FGeometryCollectionToDynamicMeshes::FToCollectionOptions Options;
 
-				TArray<FDynamicMesh3> SplitMeshes;
+				TArray<TUniquePtr<FDynamicMesh3>> SplitMeshes;
 				if (bInSplitIslands)
 				{
 					FVertexConnectedComponents Components(DynMesh.MaxVertexID());
@@ -191,7 +191,7 @@ void FMeshToCollectionDataflowNode::Evaluate(UE::Dataflow::FContext& Context, co
 					{
 						Components.ConnectCloseVertices(DynMesh, ConnectVerticesThreshold, 2);
 					}
-					FDynamicMeshEditor::SplitMesh(&DynMesh, SplitMeshes, [&Components, &DynMesh](int32 TID)
+					FDynamicMeshEditor::SplitMesh(&DynMesh, SplitMeshes, false, [&Components, &DynMesh](int32 TID)
 						{
 							return Components.GetComponent(DynMesh.GetTriangle(TID).A);
 						});
@@ -208,9 +208,9 @@ void FMeshToCollectionDataflowNode::Evaluate(UE::Dataflow::FContext& Context, co
 				if (SplitMeshes.Num() > 1)
 				{
 					Options.NewMeshParentIndex = AddRoot(NewGeometryCollection);
-					for (const FDynamicMesh3& SplitMesh : SplitMeshes)
+					for (const TUniquePtr<FDynamicMesh3>& SplitMesh : SplitMeshes)
 					{
-						Convert.AppendMeshToCollection(NewGeometryCollection, SplitMesh, FTransform::Identity, Options);
+						Convert.AppendMeshToCollection(NewGeometryCollection, *SplitMesh, FTransform::Identity, Options);
 					}
 				}
 				else
@@ -556,14 +556,14 @@ void FSplitMeshIslandsDataflowNode::Evaluate(UE::Dataflow::FContext& Context, co
 			{
 				InMesh->ProcessMesh([&OutMeshes, this](const FDynamicMesh3& ToSplit)
 				{
-					TArray<FDynamicMesh3> SplitMeshes;
+					TArray<TUniquePtr<FDynamicMesh3>> SplitMeshes;
 					FVertexConnectedComponents Components(ToSplit.MaxVertexID());
 					Components.ConnectTriangles(ToSplit);
 					if (SplitMethod == EDataflowMeshSplitIslandsMethod::ByVertexOverlap)
 					{
 						Components.ConnectCloseVertices(ToSplit, ConnectVerticesThreshold, 2);
 					}
-					FDynamicMeshEditor::SplitMesh(&ToSplit, SplitMeshes, [&Components, &ToSplit](int32 TID)
+					FDynamicMeshEditor::SplitMesh(&ToSplit, SplitMeshes, true, [&Components, &ToSplit](int32 TID)
 					{
 						return Components.GetComponent(ToSplit.GetTriangle(TID).A);
 					});
@@ -571,7 +571,7 @@ void FSplitMeshIslandsDataflowNode::Evaluate(UE::Dataflow::FContext& Context, co
 					for (int32 Idx = 0; Idx < SplitMeshes.Num(); ++Idx)
 					{
 						OutMeshes[Idx] = NewObject<UDynamicMesh>();
-						OutMeshes[Idx]->SetMesh(MoveTemp(SplitMeshes[Idx]));
+						OutMeshes[Idx]->SetMesh(MoveTemp(*SplitMeshes[Idx]));
 					}
 				});
 			}
diff --git a/Engine/Plugins/Experimental/MeshModelingToolsetExp/Source/MeshModelingToolsExp/Private/PlaneCutTool.cpp b/Engine/Plugins/Experimental/MeshModelingToolsetExp/Source/MeshModelingToolsExp/Private/PlaneCutTool.cpp
--- a/Engine/Plugins/Experimental/MeshModelingToolsetExp/Source/MeshModelingToolsExp/Private/PlaneCutTool.cpp
+++ b/Engine/Plugins/Experimental/MeshModelingToolsetExp/Source/MeshModelingToolsExp/Private/PlaneCutTool.cpp
@@ -408,7 +408,7 @@ void UPlaneCutTool::GenerateAsset(const TArray<FDynamicMeshOpResult>& Results)
 
 	ensure(Results.Num() > 0);
 	int32 NumSourceMeshes = MeshesToCut.Num();
-	TArray<TArray<FDynamicMesh3>> AllSplitMeshes; AllSplitMeshes.SetNum(NumSourceMeshes);
+	TArray<TArray<TUniquePtr<FDynamicMesh3>>> AllSplitMeshes; AllSplitMeshes.SetNum(NumSourceMeshes);
 
 	// build a selection change starting w/ the original selection (used if objects are added below)
 	FSelectedObjectsChangeList NewSelection;
@@ -456,8 +456,8 @@ void UPlaneCutTool::GenerateAsset(const TArray<FDynamicMeshOpResult>& Results)
 			TDynamicMeshScalarTriangleAttribute<int>* SubMeshIDs =
 				static_cast<TDynamicMeshScalarTriangleAttribute<int>*>(UseMesh->Attributes()->GetAttachedAttribute(
 					FPlaneCutOp::ObjectIndexAttribute));
-			TArray<FDynamicMesh3>& SplitMeshes = AllSplitMeshes[OrigMeshIdx];
-			bool bWasSplit = FDynamicMeshEditor::SplitMesh(UseMesh, SplitMeshes, [SubMeshIDs](int TID)
+			TArray<TUniquePtr<FDynamicMesh3>>& SplitMeshes = AllSplitMeshes[OrigMeshIdx];
+			bool bWasSplit = FDynamicMeshEditor::SplitMesh(UseMesh, SplitMeshes, false, [SubMeshIDs](int TID)
 			{
 				return SubMeshIDs->GetValue(TID);
 			}
@@ -470,7 +470,7 @@ void UPlaneCutTool::GenerateAsset(const TArray<FDynamicMeshOpResult>& Results)
 					continue;
 				}
 				bNeedToAdd = bNeedToAdd || (SplitMeshes.Num() > 1);
-				UseMesh = &SplitMeshes[0];
+				UseMesh = SplitMeshes[0].Get();
 			}
 		}
 
@@ -481,7 +481,7 @@ void UPlaneCutTool::GenerateAsset(const TArray<FDynamicMeshOpResult>& Results)
 	{
 		for (int OrigMeshIdx = 0; OrigMeshIdx < NumSourceMeshes; OrigMeshIdx++)
 		{
-			TArray<FDynamicMesh3>& SplitMeshes = AllSplitMeshes[OrigMeshIdx];
+			TArray<TUniquePtr<FDynamicMesh3>>& SplitMeshes = AllSplitMeshes[OrigMeshIdx];
 			if (SplitMeshes.Num() < 2)
 			{
 				continue;
@@ -500,7 +500,7 @@ void UPlaneCutTool::GenerateAsset(const TArray<FDynamicMeshOpResult>& Results)
 				NewMeshObjectParams.BaseName = TEXT("PlaneCutOtherPart");
 				NewMeshObjectParams.Materials = ComponentMaterialSet.Materials;
 				NewMeshObjectParams.AssetMaterials = AssetMaterialSet.Materials;
-				NewMeshObjectParams.SetMesh(&SplitMeshes[AddMeshIdx]);
+				NewMeshObjectParams.SetMesh(MoveTemp(*SplitMeshes[AddMeshIdx]));
 				UE::ToolTarget::ConfigureCreateMeshObjectParams(Targets[OrigMeshIdx], NewMeshObjectParams);
 				FCreateMeshObjectResult Result = UE::Modeling::CreateMeshObject(GetToolManager(), MoveTemp(NewMeshObjectParams));
 				if (Result.IsOK() && Result.NewActor != nullptr)
diff --git a/Engine/Plugins/Runtime/GeometryProcessing/Source/DynamicMesh/Private/FaceGroupUtil.cpp b/Engine/Plugins/Runtime/GeometryProcessing/Source/DynamicMesh/Private/FaceGroupUtil.cpp
--- a/Engine/Plugins/Runtime/GeometryProcessing/Source/DynamicMesh/Private/FaceGroupUtil.cpp
+++ b/Engine/Plugins/Runtime/GeometryProcessing/Source/DynamicMesh/Private/FaceGroupUtil.cpp
@@ -184,15 +184,25 @@ bool FaceGroupUtil::FindTrianglesByGroup(FDynamicMesh3& Mesh, int32 FindGroupID,
 
 void FaceGroupUtil::SeparateMeshByGroups(FDynamicMesh3& Mesh, TArray<FDynamicMesh3>& SplitMeshes)
 {
+PRAGMA_DISABLE_DEPRECATION_WARNINGS
 	FDynamicMeshEditor::SplitMesh(&Mesh, SplitMeshes, [&Mesh](int32 TID)
 	{
 		return Mesh.GetTriangleGroup(TID);
 	});
+PRAGMA_ENABLE_DEPRECATION_WARNINGS
 }
 
+void FaceGroupUtil::SeparateMeshByGroups(FDynamicMesh3& Mesh, TArray<TUniquePtr<FDynamicMesh3>>& SplitMeshes)
+{
+	FDynamicMeshEditor::SplitMesh(&Mesh, SplitMeshes, true, [&Mesh](int32 TID)
+	{
+		return Mesh.GetTriangleGroup(TID);
+	});
+}
 
 void FaceGroupUtil::SeparateMeshByGroups(FDynamicMesh3& Mesh, TArray<FDynamicMesh3>& SplitMeshes, TArray<int32>& GroupIDs)
 {
+PRAGMA_DISABLE_DEPRECATION_WARNINGS
 	// build split meshes
 	SeparateMeshByGroups(Mesh, SplitMeshes);
 
@@ -203,6 +213,15 @@ void FaceGroupUtil::SeparateMeshByGroups(FDynamicMesh3& Mesh, TArray<FDynamicMes
 		check(M.TriangleCount() > 0); // SplitMesh should never add an empty mesh
 		GroupIDs.Add(M.GetTriangleGroup(0));
 	}
+PRAGMA_ENABLE_DEPRECATION_WARNINGS
+}
+
+void FaceGroupUtil::SeparateMeshByGroups(FDynamicMesh3& Mesh, TArray<TUniquePtr<FDynamicMesh3>>& SplitMeshes, TArray<int32>& GroupIDs)
+{
+	FDynamicMeshEditor::SplitMesh(&Mesh, SplitMeshes, true, [&Mesh](int32 TID)
+	{
+		return Mesh.GetTriangleGroup(TID);
+	}, INDEX_NONE, &GroupIDs);
 }
 
 
diff --git a/Engine/Plugins/Runtime/GeometryProcessing/Source/DynamicMesh/Public/FaceGroupUtil.h b/Engine/Plugins/Runtime/GeometryProcessing/Source/DynamicMesh/Public/FaceGroupUtil.h
--- a/Engine/Plugins/Runtime/GeometryProcessing/Source/DynamicMesh/Public/FaceGroupUtil.h
+++ b/Engine/Plugins/Runtime/GeometryProcessing/Source/DynamicMesh/Public/FaceGroupUtil.h
@@ -63,13 +63,25 @@ namespace FaceGroupUtil
 	* split input Mesh into submeshes based on group ID
 	* **does not** separate disconnected components w/ same group ID
 	*/
+	UE_DEPRECATED(5.8, "Use TArray<TUniquePtr<FDynamicMesh3>> version")
 	DYNAMICMESH_API void SeparateMeshByGroups(FDynamicMesh3& Mesh, TArray<FDynamicMesh3>& SplitMeshes);
+	/**
+	 * split input Mesh into submeshes based on group ID
+	 * **does not** separate disconnected components w/ same group ID
+	 */
+	DYNAMICMESH_API void SeparateMeshByGroups(FDynamicMesh3& Mesh, TArray<TUniquePtr<FDynamicMesh3>>& SplitMeshes);
 
 	/**
 	* split input Mesh into submeshes based on group ID
 	* **does not** separate disconnected components w/ same group ID
 	*/
+	UE_DEPRECATED(5.8, "Use TArray<TUniquePtr<FDynamicMesh3>> version")
 	DYNAMICMESH_API void SeparateMeshByGroups(FDynamicMesh3& Mesh, TArray<FDynamicMesh3>& SplitMeshes, TArray<int32>& GroupIDs);
+	/**
+	 * split input Mesh into submeshes based on group ID
+	 * **does not** separate disconnected components w/ same group ID
+	 */
+	DYNAMICMESH_API void SeparateMeshByGroups(FDynamicMesh3& Mesh, TArray<TUniquePtr<FDynamicMesh3>>& SplitMeshes, TArray<int32>& GroupIDs);
 
 }
 
diff --git a/Engine/Plugins/Runtime/GeometryScripting/Source/GeometryScriptingCore/Private/MeshDecompositionFunctions.cpp b/Engine/Plugins/Runtime/GeometryScripting/Source/GeometryScriptingCore/Private/MeshDecompositionFunctions.cpp
--- a/Engine/Plugins/Runtime/GeometryScripting/Source/GeometryScriptingCore/Private/MeshDecompositionFunctions.cpp
+++ b/Engine/Plugins/Runtime/GeometryScripting/Source/GeometryScriptingCore/Private/MeshDecompositionFunctions.cpp
@@ -19,7 +19,7 @@ using namespace UE::Geometry;
 #define LOCTEXT_NAMESPACE "UGeometryScriptLibrary_MeshDecompositionFunctions"
 
 
-void BuildNewDynamicMeshes(UDynamicMesh* TargetMesh, UDynamicMeshPool* MeshPool, TArray<FDynamicMesh3>& SplitMeshes, TArray<UDynamicMesh*>& ComponentMeshes)
+static void BuildNewDynamicMeshes(UDynamicMesh* TargetMesh, UDynamicMeshPool* MeshPool, TArray<TUniquePtr<FDynamicMesh3>>& SplitMeshes, TArray<UDynamicMesh*>& ComponentMeshes)
 {
 	ComponentMeshes.SetNum(0);
 
@@ -40,7 +40,7 @@ void BuildNewDynamicMeshes(UDynamicMesh* TargetMesh, UDynamicMeshPool* MeshPool,
 		{
 			UDynamicMesh* ComponentMesh = (MeshPool != nullptr) ?
 				MeshPool->RequestMesh() : NewObject<UDynamicMesh>();
-			ComponentMesh->SetMesh(MoveTemp(SplitMeshes[mi]));
+			ComponentMesh->SetMesh(MoveTemp(*SplitMeshes[mi]));
 			ComponentMeshes.Add(ComponentMesh);
 		}
 	}
@@ -61,7 +61,7 @@ UDynamicMesh* UGeometryScriptLibrary_MeshDecompositionFunctions::SplitMeshByComp
 		return TargetMesh;
 	}
 
-	TArray<FDynamicMesh3> SplitMeshes;
+	TArray<TUniquePtr<FDynamicMesh3>> SplitMeshes;
 
 	TargetMesh->ProcessMesh([&](const FDynamicMesh3& EditMesh)
 	{
@@ -83,7 +83,7 @@ UDynamicMesh* UGeometryScriptLibrary_MeshDecompositionFunctions::SplitMeshByComp
 			}
 		}
 
-		FDynamicMeshEditor::SplitMesh(&EditMesh, SplitMeshes, [&](int32 tid) { return TriSubmeshIndices[tid]; });
+		FDynamicMeshEditor::SplitMesh(&EditMesh, SplitMeshes, false, [&](int32 tid) { return TriSubmeshIndices[tid]; });
 	});
 
 	BuildNewDynamicMeshes(TargetMesh, MeshPool, SplitMeshes, ComponentMeshes);
@@ -106,14 +106,14 @@ UDynamicMesh* UGeometryScriptLibrary_MeshDecompositionFunctions::SplitMeshByVert
 	
 	double UseThreshold = FMath::Max(ConnectVerticesThreshold, UE_DOUBLE_KINDA_SMALL_NUMBER);
 
-	TArray<FDynamicMesh3> SplitMeshes;
+	TArray<TUniquePtr<FDynamicMesh3>> SplitMeshes;
 
 	TargetMesh->ProcessMesh([&SplitMeshes, UseThreshold](const FDynamicMesh3& EditMesh) -> void
 	{
 		FVertexConnectedComponents Components(EditMesh.MaxVertexID());
 		Components.ConnectTriangles(EditMesh);
 		Components.ConnectCloseVertices(EditMesh, UseThreshold, 2);
-		FDynamicMeshEditor::SplitMesh(&EditMesh, SplitMeshes, [&Components, &EditMesh](int32 TID) -> int32
+		FDynamicMeshEditor::SplitMesh(&EditMesh, SplitMeshes, false, [&Components, &EditMesh](int32 TID) -> int32
 		{
 			return Components.GetComponent(EditMesh.GetTriangle(TID).A);
 		});
@@ -141,7 +141,7 @@ UDynamicMesh* UGeometryScriptLibrary_MeshDecompositionFunctions::SplitMeshByMate
 		return TargetMesh;
 	}
 
-	TArray<FDynamicMesh3> SplitMeshes;
+	TArray<TUniquePtr<FDynamicMesh3>> SplitMeshes;
 	TargetMesh->ProcessMesh([&](const FDynamicMesh3& EditMesh)
 	{
 		const FDynamicMeshMaterialAttribute* MaterialIDs =
@@ -152,7 +152,7 @@ UDynamicMesh* UGeometryScriptLibrary_MeshDecompositionFunctions::SplitMeshByMate
 			return;
 		}
 
-		FDynamicMeshEditor::SplitMesh(&EditMesh, SplitMeshes, [&MaterialIDs](int32 TID)
+		FDynamicMeshEditor::SplitMesh(&EditMesh, SplitMeshes, false, [&MaterialIDs](int32 TID)
 		{
 			return MaterialIDs->GetValue(TID);
 		}, INDEX_NONE, &ComponentMaterialIDs, true);
@@ -184,7 +184,7 @@ UDynamicMesh* UGeometryScriptLibrary_MeshDecompositionFunctions::SplitMeshByPoly
 		return TargetMesh;
 	}
 
-	TArray<FDynamicMesh3> SplitMeshes;
+	TArray<TUniquePtr<FDynamicMesh3>> SplitMeshes;
 	TargetMesh->ProcessMesh([&](const FDynamicMesh3& EditMesh)
 	{
 		TUniquePtr<FPolygroupSet> SplitGroups;
@@ -196,7 +196,7 @@ UDynamicMesh* UGeometryScriptLibrary_MeshDecompositionFunctions::SplitMeshByPoly
 		}
 		SplitGroups = MakeUnique<FPolygroupSet>(&EditMesh, InputGroupLayer);
 
-		FDynamicMeshEditor::SplitMesh(&EditMesh, SplitMeshes, [&SplitGroups](int32 TID)
+		FDynamicMeshEditor::SplitMesh(&EditMesh, SplitMeshes, false, [&SplitGroups](int32 TID)
 		{
 			return SplitGroups->GetGroup(TID);
 		}, INDEX_NONE, &ComponentPolygroups, true);
