Add editor-only dynamic mesh processor class, so dataflow geometry script users can access the editor only nodes by picking the appropriate parent class

#Synced-CL 52752549

diff --git a/Engine/Plugins/Runtime/GeometryScripting/Source/GeometryScriptingEditor/Private/EditorDynamicMeshProcessor.cpp b/Engine/Plugins/Runtime/GeometryScripting/Source/GeometryScriptingEditor/Private/EditorDynamicMeshProcessor.cpp
new file mode 100644
--- /dev/null
+++ b/Engine/Plugins/Runtime/GeometryScripting/Source/GeometryScriptingEditor/Private/EditorDynamicMeshProcessor.cpp
@@ -0,0 +1,5 @@
+// Copyright Epic Games, Inc. All Rights Reserved.
+
+#include "GeometryScript/EditorDynamicMeshProcessor.h"
+
+#include UE_INLINE_GENERATED_CPP_BY_NAME(EditorDynamicMeshProcessor)
diff --git a/Engine/Plugins/Runtime/GeometryScripting/Source/GeometryScriptingEditor/Public/GeometryScript/EditorDynamicMeshProcessor.h b/Engine/Plugins/Runtime/GeometryScripting/Source/GeometryScriptingEditor/Public/GeometryScript/EditorDynamicMeshProcessor.h
new file mode 100644
--- /dev/null
+++ b/Engine/Plugins/Runtime/GeometryScripting/Source/GeometryScriptingEditor/Public/GeometryScript/EditorDynamicMeshProcessor.h
@@ -0,0 +1,17 @@
+// Copyright Epic Games, Inc. All Rights Reserved.
+
+#pragma once
+
+#include "CoreMinimal.h"
+#include "DynamicMeshProcessor.h"
+
+#include "EditorDynamicMeshProcessor.generated.h"
+
+// Editor-only variant of UDynamicMeshProcessorBlueprint.
+// Blueprints with this parent class can use editor-only blueprint nodes
+// (e.g. from GeometryScriptingEditor) that are not available to runtime blueprints.
+UCLASS(MinimalAPI, Abstract, BlueprintType, Blueprintable)
+class UEditorDynamicMeshProcessorBlueprint : public UDynamicMeshProcessorBlueprint
+{
+	GENERATED_BODY()
+};
